rename UE::Geometry::TTransform3 to TTransformSRT3, update references
#rnx

diff --git a/Engine/Plugins/Experimental/Enterprise/DataprepGeometryOperations/Source/Private/DataprepGeometryOperations.cpp b/Engine/Plugins/Experimental/Enterprise/DataprepGeometryOperations/Source/Private/DataprepGeometryOperations.cpp
--- a/Engine/Plugins/Experimental/Enterprise/DataprepGeometryOperations/Source/Private/DataprepGeometryOperations.cpp
+++ b/Engine/Plugins/Experimental/Enterprise/DataprepGeometryOperations/Source/Private/DataprepGeometryOperations.cpp
@@ -273,22 +273,22 @@ void UDataprepBakeTransformOperation::OnExecution_Implementation(const FDataprep
 		}
 	}
 
-	TArray<UE::Geometry::FTransform3d> BakedTransforms;
+	TArray<UE::Geometry::FTransformSRT3d> BakedTransforms;
 	for (int32 ComponentIdx = 0; ComponentIdx < StaticMeshComponents.Num(); ComponentIdx++)
 	{
-		UE::Geometry::FTransform3d ComponentToWorld(StaticMeshComponents[ComponentIdx]->GetComponentTransform());
-		UE::Geometry::FTransform3d ToBakePart = UE::Geometry::FTransform3d::Identity();
-		UE::Geometry::FTransform3d NewWorldPart = ComponentToWorld;
+		UE::Geometry::FTransformSRT3d ComponentToWorld(StaticMeshComponents[ComponentIdx]->GetComponentTransform());
+		UE::Geometry::FTransformSRT3d ToBakePart = UE::Geometry::FTransformSRT3d::Identity();
+		UE::Geometry::FTransformSRT3d NewWorldPart = ComponentToWorld;
 
 		if (MapToFirstOccurrences[ComponentIdx] < ComponentIdx)
 		{
 			ToBakePart = BakedTransforms[MapToFirstOccurrences[ComponentIdx]];
 			BakedTransforms.Add(ToBakePart);
 			// Try to invert baked transform
-			NewWorldPart = UE::Geometry::FTransform3d(
+			NewWorldPart = UE::Geometry::FTransformSRT3d(
 				NewWorldPart.GetRotation() * ToBakePart.GetRotation().Inverse(),
 				NewWorldPart.GetTranslation(),
-				NewWorldPart.GetScale() * UE::Geometry::FTransform3d::GetSafeScaleReciprocal(ToBakePart.GetScale())
+				NewWorldPart.GetScale() * UE::Geometry::FTransformSRT3d::GetSafeScaleReciprocal(ToBakePart.GetScale())
 			);
 			NewWorldPart.SetTranslation(NewWorldPart.GetTranslation() - NewWorldPart.TransformVector(ToBakePart.GetTranslation()));
 		}
@@ -864,7 +864,7 @@ TUniquePtr<UE::Geometry::FDynamicMeshOperator> UDataprepPlaneCutOperation::MakeN
 	LocalToWorld.SetScale3D(LocalToWorldScale);
 	FTransform WorldToLocal = LocalToWorld.Inverse();
 	FVector LocalOrigin = WorldToLocal.TransformPosition(CutPlaneOrigin);
-	UE::Geometry::FTransform3d W2LForNormal(WorldToLocal);
+	UE::Geometry::FTransformSRT3d W2LForNormal(WorldToLocal);
 	FVector LocalNormal = (FVector)W2LForNormal.TransformNormal((FVector3d)WorldNormal);
 	FVector BackTransformed = LocalToWorld.TransformVector(LocalNormal);
 	float NormalScaleFactor = FVector::DotProduct(BackTransformed, WorldNormal);
diff --git a/Engine/Plugins/Experimental/GeometryScripting/Source/GeometryScriptingCore/Private/MeshBasicEditFunctions.cpp b/Engine/Plugins/Experimental/GeometryScripting/Source/GeometryScriptingCore/Private/MeshBasicEditFunctions.cpp
--- a/Engine/Plugins/Experimental/GeometryScripting/Source/GeometryScriptingCore/Private/MeshBasicEditFunctions.cpp
+++ b/Engine/Plugins/Experimental/GeometryScripting/Source/GeometryScriptingCore/Private/MeshBasicEditFunctions.cpp
@@ -322,7 +322,7 @@ UDynamicMesh* UGeometryScriptLibrary_MeshBasicEditFunctions::AppendMesh(
 	{
 		AppendMesh->ProcessMesh([&](const FDynamicMesh3& OtherMesh)
 		{
-			UE::Geometry::FTransform3d XForm(AppendTransform);
+			FTransformSRT3d XForm(AppendTransform);
 			FMeshIndexMappings TmpMappings;
 			FDynamicMeshEditor Editor(&AppendToMesh);
 			const FDynamicMesh3* UseOtherMesh = &OtherMesh;
@@ -364,7 +364,7 @@ UDynamicMesh* UGeometryScriptLibrary_MeshBasicEditFunctions::AppendMeshRepeated(
 	}
 	if (RepeatCount > 0)
 	{
-		UE::Geometry::FTransform3d XForm(AppendTransform);
+		FTransformSRT3d XForm(AppendTransform);
 		FDynamicMesh3 TmpMesh;
 		AppendMesh->ProcessMesh([&](const FDynamicMesh3& OtherMesh) { TmpMesh.Copy(OtherMesh); });
 		if (bApplyTransformToFirstInstance)
diff --git a/Engine/Plugins/Experimental/GeometryScripting/Source/GeometryScriptingCore/Private/MeshBooleanFunctions.cpp b/Engine/Plugins/Experimental/GeometryScripting/Source/GeometryScriptingCore/Private/MeshBooleanFunctions.cpp
--- a/Engine/Plugins/Experimental/GeometryScripting/Source/GeometryScriptingCore/Private/MeshBooleanFunctions.cpp
+++ b/Engine/Plugins/Experimental/GeometryScripting/Source/GeometryScriptingCore/Private/MeshBooleanFunctions.cpp
@@ -65,8 +65,8 @@ UDynamicMesh* UGeometryScriptLibrary_MeshBooleanFunctions::ApplyMeshBoolean(
 		ToolMesh->ProcessMesh([&](const FDynamicMesh3& Mesh2)
 		{
 			FMeshBoolean MeshBoolean(
-				&Mesh1, (UE::Geometry::FTransform3d)TargetTransform,
-				&Mesh2, (UE::Geometry::FTransform3d)ToolTransform,
+				&Mesh1, (FTransformSRT3d)TargetTransform,
+				&Mesh2, (FTransformSRT3d)ToolTransform,
 				&NewResultMesh, ApplyOperation);
 			MeshBoolean.bPutResultInInputSpace = true;
 			MeshBoolean.bSimplifyAlongNewEdges = Options.bSimplifyOutput;
@@ -85,7 +85,7 @@ UDynamicMesh* UGeometryScriptLibrary_MeshBooleanFunctions::ApplyMeshBoolean(
 	}
 
 	// Boolean result is in the space of TargetTransform, so invert that
-	MeshTransforms::ApplyTransform(NewResultMesh, (UE::Geometry::FTransform3d)TargetTransform.Inverse());
+	MeshTransforms::ApplyTransform(NewResultMesh, (FTransformSRT3d)TargetTransform.Inverse());
 
 	if (NewBoundaryEdges.Num() > 0 && Options.bFillHoles)
 	{
diff --git a/Engine/Plugins/Experimental/GeometryScripting/Source/GeometryScriptingCore/Private/MeshComparisonFunctions.cpp b/Engine/Plugins/Experimental/GeometryScripting/Source/GeometryScriptingCore/Private/MeshComparisonFunctions.cpp
--- a/Engine/Plugins/Experimental/GeometryScripting/Source/GeometryScriptingCore/Private/MeshComparisonFunctions.cpp
+++ b/Engine/Plugins/Experimental/GeometryScripting/Source/GeometryScriptingCore/Private/MeshComparisonFunctions.cpp
@@ -85,9 +85,9 @@ UDynamicMesh* UGeometryScriptLibrary_MeshComparisonFunctions::IsIntersectingMesh
 			ParallelFor(2, [&](int32 k) { Spatials[k].SetMesh((k == 0) ? &Mesh1 : &Mesh2, true); }, EParallelForFlags::Unbalanced);
 
 			bool bIsIdentity1 = TargetTransformIn.Equals(FTransform::Identity, 0);
-			UE::Geometry::FTransform3d Transform1(TargetTransformIn);
+			FTransformSRT3d Transform1(TargetTransformIn);
 			bool bIsIdentity2 = OtherTransformIn.Equals(FTransform::Identity, 0);
-			UE::Geometry::FTransform3d Transform2(OtherTransformIn);
+			FTransformSRT3d Transform2(OtherTransformIn);
 			if (bIsIdentity1 && bIsIdentity2)
 			{
 				bIsIntersecting = Spatials[0].TestIntersection(Spatials[1]);
@@ -95,7 +95,7 @@ UDynamicMesh* UGeometryScriptLibrary_MeshComparisonFunctions::IsIntersectingMesh
 			else if (bIsIdentity1 || bIsIdentity2)
 			{
 				FIndex2i Indices = (bIsIdentity1) ? FIndex2i(0,1) : FIndex2i(1,0);
-				UE::Geometry::FTransform3d UseTransform = (bIsIdentity2) ? Transform1 : Transform2;
+				FTransformSRT3d UseTransform = (bIsIdentity2) ? Transform1 : Transform2;
 				bIsIntersecting = Spatials[Indices.A].TestIntersection(Spatials[Indices.B],
 					[&](const FVector3d& Pos) { return UseTransform.TransformPosition(Pos); });
 			}
diff --git a/Engine/Plugins/Experimental/GeometryScripting/Source/GeometryScriptingCore/Private/MeshDeformFunctions.cpp b/Engine/Plugins/Experimental/GeometryScripting/Source/GeometryScriptingCore/Private/MeshDeformFunctions.cpp
--- a/Engine/Plugins/Experimental/GeometryScripting/Source/GeometryScriptingCore/Private/MeshDeformFunctions.cpp
+++ b/Engine/Plugins/Experimental/GeometryScripting/Source/GeometryScriptingCore/Private/MeshDeformFunctions.cpp
@@ -180,7 +180,7 @@ UDynamicMesh* UGeometryScriptLibrary_MeshDeformFunctions::ApplyMathWarpToMesh(
 		return TargetMesh;
 	}
 
-	UE::Geometry::FTransform3d WarpTransform(WarpOrientation);
+	FTransformSRT3d WarpTransform(WarpOrientation);
 	double UseShift = (double)Options.FrequencyShift;
 	double UseFrequency = (double)Options.Frequency;
 	double UseMagnitude = (double)Options.Magnitude;
diff --git a/Engine/Plugins/Experimental/GeometryScripting/Source/GeometryScriptingCore/Private/MeshPrimitiveFunctions.cpp b/Engine/Plugins/Experimental/GeometryScripting/Source/GeometryScriptingCore/Private/MeshPrimitiveFunctions.cpp
--- a/Engine/Plugins/Experimental/GeometryScripting/Source/GeometryScriptingCore/Private/MeshPrimitiveFunctions.cpp
+++ b/Engine/Plugins/Experimental/GeometryScripting/Source/GeometryScriptingCore/Private/MeshPrimitiveFunctions.cpp
@@ -40,7 +40,7 @@ static void AppendPrimitive(
 			MeshTransforms::Translate(Mesh, PreTranslate);
 		}
 
-		MeshTransforms::ApplyTransform(Mesh, (UE::Geometry::FTransform3d)Transform);
+		MeshTransforms::ApplyTransform(Mesh, (FTransformSRT3d)Transform);
 		if (PrimitiveOptions.PolygroupMode == EGeometryScriptPrimitivePolygroupMode::SingleGroup)
 		{
 			for (int32 tid : Mesh.TriangleIndicesItr())
@@ -100,7 +100,7 @@ static void AppendPrimitiveMesh(
 			MeshTransforms::Translate(Mesh, PreTranslate);
 		}
 
-		MeshTransforms::ApplyTransform(Mesh, (UE::Geometry::FTransform3d)Transform);
+		MeshTransforms::ApplyTransform(Mesh, (FTransformSRT3d)Transform);
 		if (PrimitiveOptions.PolygroupMode == EGeometryScriptPrimitivePolygroupMode::SingleGroup)
 		{
 			for (int32 tid : Mesh.TriangleIndicesItr())
diff --git a/Engine/Plugins/Experimental/GeometryScripting/Source/GeometryScriptingCore/Private/MeshRepairFunctions.cpp b/Engine/Plugins/Experimental/GeometryScripting/Source/GeometryScriptingCore/Private/MeshRepairFunctions.cpp
--- a/Engine/Plugins/Experimental/GeometryScripting/Source/GeometryScriptingCore/Private/MeshRepairFunctions.cpp
+++ b/Engine/Plugins/Experimental/GeometryScripting/Source/GeometryScriptingCore/Private/MeshRepairFunctions.cpp
@@ -174,8 +174,8 @@ UDynamicMesh* UGeometryScriptLibrary_MeshRepairFunctions::RemoveHiddenTriangles(
 		Jacket.AddRandomRays = Options.RaysPerSample;
 		Jacket.AddTriangleSamples = Options.SamplesPerTriangle;
 
-		TArray<UE::Geometry::FTransform3d> NoTransforms;
-		NoTransforms.Add(UE::Geometry::FTransform3d::Identity());
+		TArray<FTransformSRT3d> NoTransforms;
+		NoTransforms.Add(FTransformSRT3d::Identity());
 
 		//  set up AABBTree and FWNTree lists
 		FDynamicMeshAABBTree3 Spatial(&EditMesh);
diff --git a/Engine/Plugins/Experimental/GeometryScripting/Source/GeometryScriptingCore/Private/MeshTransformFunctions.cpp b/Engine/Plugins/Experimental/GeometryScripting/Source/GeometryScriptingCore/Private/MeshTransformFunctions.cpp
--- a/Engine/Plugins/Experimental/GeometryScripting/Source/GeometryScriptingCore/Private/MeshTransformFunctions.cpp
+++ b/Engine/Plugins/Experimental/GeometryScripting/Source/GeometryScriptingCore/Private/MeshTransformFunctions.cpp
@@ -25,7 +25,7 @@ UDynamicMesh* UGeometryScriptLibrary_MeshTransformFunctions::TransformMesh(
 	// todo: publish correct change types
 	TargetMesh->EditMesh([&](FDynamicMesh3& EditMesh) 
 	{
-		MeshTransforms::ApplyTransform(EditMesh, (UE::Geometry::FTransform3d)Transform);
+		MeshTransforms::ApplyTransform(EditMesh, (FTransformSRT3d)Transform);
 
 	}, EDynamicMeshChangeType::GeneralEdit, EDynamicMeshAttributeChangeFlags::Unknown, false);
 
diff --git a/Engine/Plugins/Experimental/GeometryScripting/Source/GeometryScriptingCore/Private/SceneUtilityFunctions.cpp b/Engine/Plugins/Experimental/GeometryScripting/Source/GeometryScriptingCore/Private/SceneUtilityFunctions.cpp
--- a/Engine/Plugins/Experimental/GeometryScripting/Source/GeometryScriptingCore/Private/SceneUtilityFunctions.cpp
+++ b/Engine/Plugins/Experimental/GeometryScripting/Source/GeometryScriptingCore/Private/SceneUtilityFunctions.cpp
@@ -106,7 +106,7 @@ UDynamicMesh* UGeometryScriptLibrary_SceneUtilityFunctions::CopyMeshFromComponen
 	{
 		ToDynamicMesh->EditMesh([&](FDynamicMesh3& EditMesh) 
 		{
-			MeshTransforms::ApplyTransform(EditMesh, (UE::Geometry::FTransform3d)LocalToWorld);
+			MeshTransforms::ApplyTransform(EditMesh, (FTransformSRT3d)LocalToWorld);
 
 		}, EDynamicMeshChangeType::GeneralEdit, EDynamicMeshAttributeChangeFlags::Unknown, false);	
 	}
diff --git a/Engine/Plugins/Experimental/HairModelingToolset/Source/HairModelingToolset/Private/GroomCardsEditorTool.cpp b/Engine/Plugins/Experimental/HairModelingToolset/Source/HairModelingToolset/Private/GroomCardsEditorTool.cpp
--- a/Engine/Plugins/Experimental/HairModelingToolset/Source/HairModelingToolset/Private/GroomCardsEditorTool.cpp
+++ b/Engine/Plugins/Experimental/HairModelingToolset/Source/HairModelingToolset/Private/GroomCardsEditorTool.cpp
@@ -543,7 +543,7 @@ void UGroomCardsEditorTool::BeginCardEdit(int32 CardGroupID)
 {
 	const UE::GroomQueries::FMeshCardStrip& CardStrip = EditableCardSet->CardStrips->FindStripForGroup(CardGroupID);
 
-	UE::Geometry::FTransform3d Transform(PreviewMesh->GetTransform());
+	FTransformSRT3d Transform(PreviewMesh->GetTransform());
 	check(FMathd::Abs(Transform.TransformVector(UE::Geometry::Normalized(FVector3d(1,1,1))).SquaredLength() - 1.0) < FMathd::ZeroTolerance);
 	FFrame3d TransformFrame(Transform.GetTranslation(), Transform.GetRotation());
 
@@ -578,7 +578,7 @@ void UGroomCardsEditorTool::BeginCardEdit(int32 CardGroupID)
 void UGroomCardsEditorTool::RestoreCardEdit(const FGroomCardEdit* RestoreEdit)
 {
 	// this should probably be cached once...
-	UE::Geometry::FTransform3d Transform(PreviewMesh->GetTransform());
+	FTransformSRT3d Transform(PreviewMesh->GetTransform());
 	check(FMathd::Abs(Transform.TransformVector(UE::Geometry::Normalized(FVector3d(1, 1, 1))).SquaredLength() - 1.0) < FMathd::ZeroTolerance);
 	FFrame3d TransformFrame(Transform.GetTranslation(), Transform.GetRotation());
 
@@ -626,7 +626,7 @@ void UGroomCardsEditorTool::UpdateOnCurveEdit()
 		//BeginMoveChange();
 	}
 
-	UE::Geometry::FTransform3d Transform(PreviewMesh->GetTransform());
+	FTransformSRT3d Transform(PreviewMesh->GetTransform());
 	check( FMathd::Abs(Transform.TransformVector(UE::Geometry::Normalized(FVector3d(1,1,1))).SquaredLength() - 1.0) < FMathd::ZeroTolerance );
 	FFrame3d TransformFrame(Transform.GetTranslation(), Transform.GetRotation());
 
diff --git a/Engine/Plugins/Experimental/HairModelingToolset/Source/HairModelingToolset/Private/GroomToMeshTool.cpp b/Engine/Plugins/Experimental/HairModelingToolset/Source/HairModelingToolset/Private/GroomToMeshTool.cpp
--- a/Engine/Plugins/Experimental/HairModelingToolset/Source/HairModelingToolset/Private/GroomToMeshTool.cpp
+++ b/Engine/Plugins/Experimental/HairModelingToolset/Source/HairModelingToolset/Private/GroomToMeshTool.cpp
@@ -151,7 +151,7 @@ void UGroomToMeshTool::Shutdown(EToolShutdownType ShutdownType)
 {
 	Settings->SaveProperties(this);
 
-	UE::Geometry::FTransform3d Transform(PreviewMesh->GetTransform());
+	FTransformSRT3d Transform(PreviewMesh->GetTransform());
 	PreviewMesh->Disconnect();
 	PreviewMesh = nullptr;
 
@@ -586,7 +586,7 @@ struct FTargetSpace
 
 	AActor* TargetActor = nullptr;
 	USceneComponent* TargetComponent = nullptr;
-	UE::Geometry::FTransform3d TargetTransform = UE::Geometry::FTransform3d::Identity();
+	FTransformSRT3d TargetTransform = FTransformSRT3d::Identity();
 	FFrame3d TargetFrame = FFrame3d();
 
 	FTargetSpace() {}
@@ -609,7 +609,7 @@ struct FTargetSpace
 		TargetFrame = FrameIn;
 	}
 
-	FTargetSpace(const UE::Geometry::FTransform3d& TransformIn)
+	FTargetSpace(const FTransformSRT3d& TransformIn)
 	{
 		TargetSpace = ETargetSpaceType::Transform;
 		TargetTransform = TransformIn;
@@ -632,22 +632,22 @@ static void TransformToSpace(FDynamicMesh3* Mesh, UPrimitiveComponent* SourceCom
 	// Could also add a variant there that applies multiple transforms in sequence to avoid write overhead
 
 	// transform up to world
-	UE::Geometry::FTransform3d ToWorld(SourceComponent->GetComponentTransform());
+	FTransformSRT3d ToWorld(SourceComponent->GetComponentTransform());
 	MeshTransforms::ApplyTransform(*Mesh, ToWorld);
 
 	if (TargetSpace.TargetSpace == ETargetSpaceType::Actor)
 	{
 		check(TargetSpace.TargetActor);
-		UE::Geometry::FTransform3d ActorTransform(TargetSpace.TargetActor->GetActorTransform());
-		UE::Geometry::FTransform3d ToActorTransform = ActorTransform.Inverse();
+		FTransformSRT3d ActorTransform(TargetSpace.TargetActor->GetActorTransform());
+		FTransformSRT3d ToActorTransform = ActorTransform.Inverse();
 		MeshTransforms::ApplyTransform(*Mesh, ToActorTransform);
 
 	}
 	else if (TargetSpace.TargetSpace == ETargetSpaceType::Component)
 	{
 		check(TargetSpace.TargetComponent);
-		UE::Geometry::FTransform3d ComponentTransform(TargetSpace.TargetComponent->GetComponentTransform());
-		UE::Geometry::FTransform3d ToComponentTransform = ComponentTransform.Inverse();
+		FTransformSRT3d ComponentTransform(TargetSpace.TargetComponent->GetComponentTransform());
+		FTransformSRT3d ToComponentTransform = ComponentTransform.Inverse();
 		MeshTransforms::ApplyTransform(*Mesh, ToComponentTransform);
 	}
 	else if (TargetSpace.TargetSpace == ETargetSpaceType::Frame)
diff --git a/Engine/Plugins/Experimental/MeshLODToolset/Source/MeshLODToolset/Private/Tools/GenerateStaticMeshLODAssetTool.cpp b/Engine/Plugins/Experimental/MeshLODToolset/Source/MeshLODToolset/Private/Tools/GenerateStaticMeshLODAssetTool.cpp
--- a/Engine/Plugins/Experimental/MeshLODToolset/Source/MeshLODToolset/Private/Tools/GenerateStaticMeshLODAssetTool.cpp
+++ b/Engine/Plugins/Experimental/MeshLODToolset/Source/MeshLODToolset/Private/Tools/GenerateStaticMeshLODAssetTool.cpp
@@ -156,7 +156,7 @@ namespace GenerateStaticMeshLODAssetLocals
 
 	public:
 
-		FGenerateStaticMeshLODAssetOperatorFactory(UGenerateStaticMeshLODAssetTool* AutoLODTool, UE::Geometry::FTransform3d ResultTransform) :
+		FGenerateStaticMeshLODAssetOperatorFactory(UGenerateStaticMeshLODAssetTool* AutoLODTool, FTransformSRT3d ResultTransform) :
 			AutoLODTool(AutoLODTool), 
 			ResultTransform(ResultTransform) 
 		{}
@@ -194,7 +194,7 @@ namespace GenerateStaticMeshLODAssetLocals
 		}
 
 		UGenerateStaticMeshLODAssetTool* AutoLODTool = nullptr;
-		UE::Geometry::FTransform3d ResultTransform;
+		FTransformSRT3d ResultTransform;
 	};
 
 	static void DisplayCriticalWarningMessage(const FString& Message)
@@ -368,7 +368,7 @@ void UGenerateStaticMeshLODAssetTool::Setup()
 
 
 	FBoxSphereBounds Bounds = StaticMeshComponent->Bounds;
-	UE::Geometry::FTransform3d PreviewTransform = UE::ToolTarget::GetLocalToWorldTransform(Targets[0]);
+	FTransformSRT3d PreviewTransform = UE::ToolTarget::GetLocalToWorldTransform(Targets[0]);
 	PreviewTransform.SetTranslation(PreviewTransform.GetTranslation() + 2.5 * (double)Bounds.BoxExtent.Y * FVector3d::UnitY());
 
 	this->OpFactory = MakeUnique<FGenerateStaticMeshLODAssetOperatorFactory>(this, PreviewTransform);
diff --git a/Engine/Plugins/Experimental/MeshModelingToolsetExp/Source/GeometryProcessingAdapters/Private/ApproximateActorsImpl.cpp b/Engine/Plugins/Experimental/MeshModelingToolsetExp/Source/GeometryProcessingAdapters/Private/ApproximateActorsImpl.cpp
--- a/Engine/Plugins/Experimental/MeshModelingToolsetExp/Source/GeometryProcessingAdapters/Private/ApproximateActorsImpl.cpp
+++ b/Engine/Plugins/Experimental/MeshModelingToolsetExp/Source/GeometryProcessingAdapters/Private/ApproximateActorsImpl.cpp
@@ -552,9 +552,8 @@ static TSharedPtr<FApproximationMeshData> GenerateApproximationMesh(
 			CurResultMeshSpatial.Build();
 		}
 
-		using FTransform3d = UE::Geometry::FTransform3d;
-		TArray<FTransform3d> NoTransforms;
-		NoTransforms.Add(FTransform3d::Identity());
+		TArray<FTransformSRT3d> NoTransforms;
+		NoTransforms.Add(FTransformSRT3d::Identity());
 		TArray<FDynamicMeshAABBTree3*> Spatials;
 		Spatials.Add(&CurResultMeshSpatial);
 
@@ -571,7 +570,7 @@ static TSharedPtr<FApproximationMeshData> GenerateApproximationMesh(
 			RectGen.Width = RectGen.Height = 10.0 * Bounds.MaxDim();
 			BasePlaneOccluderMesh.Copy(&RectGen.Generate());
 			BasePlaneOccluderSpatial.SetMesh(&BasePlaneOccluderMesh, true);
-			NoTransforms.Add(FTransform3d::Identity());
+			NoTransforms.Add(FTransformSRT3d::Identity());
 			Spatials.Add(&BasePlaneOccluderSpatial);
 		}
 
diff --git a/Engine/Plugins/Experimental/MeshModelingToolsetExp/Source/MeshModelingToolsEditorOnly/Private/BspConversionTool.cpp b/Engine/Plugins/Experimental/MeshModelingToolsetExp/Source/MeshModelingToolsEditorOnly/Private/BspConversionTool.cpp
--- a/Engine/Plugins/Experimental/MeshModelingToolsetExp/Source/MeshModelingToolsEditorOnly/Private/BspConversionTool.cpp
+++ b/Engine/Plugins/Experimental/MeshModelingToolsetExp/Source/MeshModelingToolsEditorOnly/Private/BspConversionTool.cpp
@@ -33,9 +33,9 @@ using namespace UE::Geometry;
 void ConvertBrushesToDynamicMesh(TArray<ABrush*>& BrushesToConvert, FDynamicMesh3& OutputMesh, 
 	TArray<UMaterialInterface*>& OutputMaterials);
 bool ApplyDynamicMeshBooleanOperation(
-	FDynamicMesh3& MeshA, const UE::Geometry::FTransform3d& TransformA, const TArray<UMaterialInterface*>& MaterialsA,
-	FDynamicMesh3& MeshB, const UE::Geometry::FTransform3d& TransformB, const TArray<UMaterialInterface*>& MaterialsB,
-	FDynamicMesh3& OutputMesh, UE::Geometry::FTransform3d& OutputTransform, TArray<UMaterialInterface*>& OutputMaterials,
+	FDynamicMesh3& MeshA, const FTransformSRT3d& TransformA, const TArray<UMaterialInterface*>& MaterialsA,
+	FDynamicMesh3& MeshB, const FTransformSRT3d& TransformB, const TArray<UMaterialInterface*>& MaterialsB,
+	FDynamicMesh3& OutputMesh, FTransformSRT3d& OutputTransform, TArray<UMaterialInterface*>& OutputMaterials,
 	FMeshBoolean::EBooleanOp Operation);
 FText GetBrushGeometryErrorMessage(ABrush* Brush);
 
@@ -357,7 +357,7 @@ bool UBspConversionTool::CombineThenConvert(FText*)
 	}
 
 	MeshTransforms::Translate(OutputMesh, -NewPivot);
-	UE::Geometry::FTransform3d Transform = UE::Geometry::FTransform3d::Identity();
+	FTransformSRT3d Transform = FTransformSRT3d::Identity();
 	Transform.SetTranslation(NewPivot);
 
 	PreviewMesh->UpdatePreview(&OutputMesh);
@@ -387,14 +387,14 @@ bool UBspConversionTool::ConvertThenCombine(FText *ErrorMessage)
 	// We'll need temporary space to store meshes as we build them up. The input and output meshes,
 	// transforms, and material arrays will swap after we merge in each new mesh.
 	FDynamicMesh3 MeshStorage[2];
-	UE::Geometry::FTransform3d TransformStorage[2];
+	FTransformSRT3d TransformStorage[2];
 	TArray<UMaterialInterface*> MaterialArraysStorage[2];
 
 	FDynamicMesh3* InputMesh = &MeshStorage[0];
 	FDynamicMesh3* OutputMesh = &MeshStorage[1];
 	
-	UE::Geometry::FTransform3d* InputTransform = &TransformStorage[0];
-	UE::Geometry::FTransform3d* OutputTransform = &TransformStorage[1];
+	FTransformSRT3d* InputTransform = &TransformStorage[0];
+	FTransformSRT3d* OutputTransform = &TransformStorage[1];
 
 	TArray<UMaterialInterface*>* InputMaterials = &MaterialArraysStorage[0];
 	TArray<UMaterialInterface*>* OutputMaterials = &MaterialArraysStorage[1];
@@ -406,7 +406,7 @@ bool UBspConversionTool::ConvertThenCombine(FText *ErrorMessage)
 	FDynamicMesh3 NextMesh;
 
 	// The transform when a brush is first converted is always the same
-	UE::Geometry::FTransform3d NextTransform = UE::Geometry::FTransform3d::Identity();
+	FTransformSRT3d NextTransform = FTransformSRT3d::Identity();
 
 	// We could actually point to cached materials since we don't change these, but for consistency,
 	// we'll have a local copy.
@@ -439,7 +439,7 @@ bool UBspConversionTool::ConvertThenCombine(FText *ErrorMessage)
 					MakeShared<const TArray<UMaterialInterface*>>(*OutputMaterials)));
 			}
 		}
-		*OutputTransform = UE::Geometry::FTransform3d::Identity(); // Always identity when first converted
+		*OutputTransform = FTransformSRT3d::Identity(); // Always identity when first converted
 	}
 
 	// Now convert the other meshes one at a time and apply them using boolean mesh operations
@@ -582,9 +582,9 @@ void ConvertBrushesToDynamicMesh(TArray<ABrush*>& BrushesToConvert, FDynamicMesh
 }
 
 bool ApplyDynamicMeshBooleanOperation(
-	FDynamicMesh3& MeshA, const UE::Geometry::FTransform3d& TransformA, const TArray<UMaterialInterface*>& MaterialsA,
-	FDynamicMesh3& MeshB, const UE::Geometry::FTransform3d& TransformB, const TArray<UMaterialInterface*>& MaterialsB,
-	FDynamicMesh3& OutputMesh, UE::Geometry::FTransform3d& OutputTransform, TArray<UMaterialInterface*>& OutputMaterials,
+	FDynamicMesh3& MeshA, const FTransformSRT3d& TransformA, const TArray<UMaterialInterface*>& MaterialsA,
+	FDynamicMesh3& MeshB, const FTransformSRT3d& TransformB, const TArray<UMaterialInterface*>& MaterialsB,
+	FDynamicMesh3& OutputMesh, FTransformSRT3d& OutputTransform, TArray<UMaterialInterface*>& OutputMaterials,
 	FMeshBoolean::EBooleanOp Operation)
 {
 	// These need to be enabled on both meshes to deal with materials properly. This is relevant, for
diff --git a/Engine/Plugins/Experimental/MeshModelingToolsetExp/Source/MeshModelingToolsEditorOnly/Private/PolygonOnMeshTool.cpp b/Engine/Plugins/Experimental/MeshModelingToolsetExp/Source/MeshModelingToolsEditorOnly/Private/PolygonOnMeshTool.cpp
--- a/Engine/Plugins/Experimental/MeshModelingToolsetExp/Source/MeshModelingToolsEditorOnly/Private/PolygonOnMeshTool.cpp
+++ b/Engine/Plugins/Experimental/MeshModelingToolsetExp/Source/MeshModelingToolsEditorOnly/Private/PolygonOnMeshTool.cpp
@@ -75,7 +75,7 @@ void UPolygonOnMeshTool::Setup()
 	AddInputBehavior(HoverBehavior);
 
 	IPrimitiveComponentBackedTarget* TargetComponent = Cast<IPrimitiveComponentBackedTarget>(Target);
-	WorldTransform = UE::Geometry::FTransform3d(TargetComponent->GetWorldTransform());
+	WorldTransform = TargetComponent->GetWorldTransform();
 
 	// hide input StaticMeshComponent
 	TargetComponent->SetOwnerVisibility(false);
@@ -261,7 +261,7 @@ TUniquePtr<FDynamicMeshOperator> UPolygonOnMeshTool::MakeNewOperator()
 	EmbedOp->bAttemptFixHolesOnBoolean = !bOpLeavesOpenBoundaries && BasicProperties->bTryToFixCracks;
 
 	FFrame3d LocalFrame = DrawPlaneWorld;
-	UE::Geometry::FTransform3d ToLocal = WorldTransform.Inverse();
+	FTransform3d ToLocal = WorldTransform.Inverse();
 	LocalFrame.Transform(ToLocal);
 	EmbedOp->PolygonFrame = LocalFrame;
 	
diff --git a/Engine/Plugins/Experimental/MeshModelingToolsetExp/Source/MeshModelingToolsEditorOnly/Private/ShapeSprayTool.cpp b/Engine/Plugins/Experimental/MeshModelingToolsetExp/Source/MeshModelingToolsEditorOnly/Private/ShapeSprayTool.cpp
--- a/Engine/Plugins/Experimental/MeshModelingToolsetExp/Source/MeshModelingToolsEditorOnly/Private/ShapeSprayTool.cpp
+++ b/Engine/Plugins/Experimental/MeshModelingToolsetExp/Source/MeshModelingToolsEditorOnly/Private/ShapeSprayTool.cpp
@@ -128,7 +128,7 @@ void UShapeSprayTool::OnUpdateDrag(const FRay& Ray)
 	UDynamicMeshBrushTool::OnUpdateDrag(Ray);
 
 	FFrame3f WorldFrame((FVector3f)LastBrushStamp.WorldPosition, (FVector3f)LastBrushStamp.WorldNormal);
-	UE::Geometry::FTransform3d Transform(Cast<IPrimitiveComponentBackedTarget>(Target)->GetWorldTransform());
+	FTransform3d Transform(Cast<IPrimitiveComponentBackedTarget>(Target)->GetWorldTransform());
 
 
 	FDynamicMesh3* Mesh = AccumMeshComponent->GetMesh();
@@ -228,7 +228,7 @@ bool UShapeSprayTool::CanAccept() const
 void UShapeSprayTool::EmitResult()
 {
 	const FDynamicMesh3* Mesh = AccumMeshComponent->GetMesh();
-	UE::Geometry::FTransform3d UseTransform(Cast<IPrimitiveComponentBackedTarget>(Target)->GetOwnerActor()->GetTransform());
+	FTransform3d UseTransform(Cast<IPrimitiveComponentBackedTarget>(Target)->GetOwnerActor()->GetTransform());
 
 	GetToolManager()->BeginUndoTransaction(LOCTEXT("EmitShapeSpray", "Create ShapeSpray"));
 
diff --git a/Engine/Plugins/Experimental/MeshModelingToolsetExp/Source/MeshModelingToolsEditorOnly/Public/PolygonOnMeshTool.h b/Engine/Plugins/Experimental/MeshModelingToolsetExp/Source/MeshModelingToolsEditorOnly/Public/PolygonOnMeshTool.h
--- a/Engine/Plugins/Experimental/MeshModelingToolsetExp/Source/MeshModelingToolsEditorOnly/Public/PolygonOnMeshTool.h
+++ b/Engine/Plugins/Experimental/MeshModelingToolsetExp/Source/MeshModelingToolsEditorOnly/Public/PolygonOnMeshTool.h
@@ -211,7 +211,7 @@ protected:
 protected:
 	UWorld* TargetWorld;
 
-	UE::Geometry::FTransform3d WorldTransform;
+	FTransform3d WorldTransform;
 	FViewCameraState CameraState;
 
 	TSharedPtr<UE::Geometry::FDynamicMesh3, ESPMode::ThreadSafe> OriginalDynamicMesh;
diff --git a/Engine/Plugins/Experimental/MeshModelingToolsetExp/Source/MeshModelingToolsExp/Private/AddPatchTool.cpp b/Engine/Plugins/Experimental/MeshModelingToolsetExp/Source/MeshModelingToolsExp/Private/AddPatchTool.cpp
--- a/Engine/Plugins/Experimental/MeshModelingToolsetExp/Source/MeshModelingToolsExp/Private/AddPatchTool.cpp
+++ b/Engine/Plugins/Experimental/MeshModelingToolsetExp/Source/MeshModelingToolsExp/Private/AddPatchTool.cpp
@@ -293,7 +293,7 @@ void UAddPatchTool::GeneratePreviewBaseMesh()
 void UAddPatchTool::OnClicked(const FInputDeviceRay& DeviceClickPos)
 {
 	const FDynamicMesh3* CurMesh = PreviewMesh->GetPreviewDynamicMesh();
-	UE::Geometry::FTransform3d CurTransform(PreviewMesh->GetTransform());
+	FTransform3d CurTransform(PreviewMesh->GetTransform());
 	UMaterialInterface* Material = PreviewMesh->GetMaterial();
 	GetToolManager()->BeginUndoTransaction(LOCTEXT("AddPatchToolTransactionName", "Add Patch Mesh"));
 
diff --git a/Engine/Plugins/Experimental/MeshModelingToolsetExp/Source/MeshModelingToolsExp/Private/BakeMeshAttributeMapsTool.cpp b/Engine/Plugins/Experimental/MeshModelingToolsetExp/Source/MeshModelingToolsExp/Private/BakeMeshAttributeMapsTool.cpp
--- a/Engine/Plugins/Experimental/MeshModelingToolsetExp/Source/MeshModelingToolsExp/Private/BakeMeshAttributeMapsTool.cpp
+++ b/Engine/Plugins/Experimental/MeshModelingToolsetExp/Source/MeshModelingToolsExp/Private/BakeMeshAttributeMapsTool.cpp
@@ -519,10 +519,9 @@ void UBakeMeshAttributeMapsTool::UpdateDetailMesh()
 
 	if (InputMeshSettings->bProjectionInWorldSpace && bIsBakeToSelf == false)
 	{
-		using FTransform3d = UE::Geometry::FTransform3d;
-		const FTransform3d DetailToWorld = UE::ToolTarget::GetLocalToWorldTransform(DetailTarget);
+		const FTransformSRT3d DetailToWorld = UE::ToolTarget::GetLocalToWorldTransform(DetailTarget);
 		MeshTransforms::ApplyTransform(*DetailMesh, DetailToWorld);
-		const FTransform3d WorldToBase = UE::ToolTarget::GetLocalToWorldTransform(Targets[0]);
+		const FTransformSRT3d WorldToBase = UE::ToolTarget::GetLocalToWorldTransform(Targets[0]);
 		MeshTransforms::ApplyTransform(*DetailMesh, WorldToBase.Inverse());
 	}
 
diff --git a/Engine/Plugins/Experimental/MeshModelingToolsetExp/Source/MeshModelingToolsExp/Private/BakeMeshAttributeMapsToolBase.cpp b/Engine/Plugins/Experimental/MeshModelingToolsetExp/Source/MeshModelingToolsExp/Private/BakeMeshAttributeMapsToolBase.cpp
--- a/Engine/Plugins/Experimental/MeshModelingToolsetExp/Source/MeshModelingToolsExp/Private/BakeMeshAttributeMapsToolBase.cpp
+++ b/Engine/Plugins/Experimental/MeshModelingToolsetExp/Source/MeshModelingToolsExp/Private/BakeMeshAttributeMapsToolBase.cpp
@@ -56,7 +56,7 @@ void UBakeMeshAttributeMapsToolBase::Setup()
 	UE::ToolTarget::HideSourceObject(Targets[0]);
 
 	const FDynamicMesh3 InputMesh = UE::ToolTarget::GetDynamicMeshCopy(Targets[0], true);
-	const UE::Geometry::FTransform3d BaseToWorld = UE::ToolTarget::GetLocalToWorldTransform(Targets[0]);
+	const FTransformSRT3d BaseToWorld = UE::ToolTarget::GetLocalToWorldTransform(Targets[0]);
 	PreviewMesh = NewObject<UPreviewMesh>(this);
 	PreviewMesh->CreateInWorld(GetTargetWorld(), FTransform::Identity);
 	ToolSetupUtil::ApplyRenderingConfigurationToPreview(PreviewMesh, nullptr);
diff --git a/Engine/Plugins/Experimental/MeshModelingToolsetExp/Source/MeshModelingToolsExp/Private/BakeMeshAttributeVertexTool.cpp b/Engine/Plugins/Experimental/MeshModelingToolsetExp/Source/MeshModelingToolsExp/Private/BakeMeshAttributeVertexTool.cpp
--- a/Engine/Plugins/Experimental/MeshModelingToolsetExp/Source/MeshModelingToolsExp/Private/BakeMeshAttributeVertexTool.cpp
+++ b/Engine/Plugins/Experimental/MeshModelingToolsetExp/Source/MeshModelingToolsExp/Private/BakeMeshAttributeVertexTool.cpp
@@ -461,9 +461,9 @@ void UBakeMeshAttributeVertexTool::UpdateDetailMesh()
 	Converter.Convert( UE::ToolTarget::GetMeshDescription(DetailTargetMesh), *DetailMesh);
 	if (InputMeshSettings->bProjectionInWorldSpace && bIsBakeToSelf == false)
 	{
-		const UE::Geometry::FTransform3d DetailToWorld(DetailComponent->GetWorldTransform());
+		const FTransformSRT3d DetailToWorld(DetailComponent->GetWorldTransform());
 		MeshTransforms::ApplyTransform(*DetailMesh, DetailToWorld);
-		const UE::Geometry::FTransform3d WorldToBase(TargetComponent->GetWorldTransform());
+		const FTransformSRT3d WorldToBase(TargetComponent->GetWorldTransform());
 		MeshTransforms::ApplyTransform(*DetailMesh, WorldToBase.Inverse());
 	}
 
diff --git a/Engine/Plugins/Experimental/MeshModelingToolsetExp/Source/MeshModelingToolsExp/Private/BakeMultiMeshAttributeMapsTool.cpp b/Engine/Plugins/Experimental/MeshModelingToolsetExp/Source/MeshModelingToolsExp/Private/BakeMultiMeshAttributeMapsTool.cpp
--- a/Engine/Plugins/Experimental/MeshModelingToolsetExp/Source/MeshModelingToolsExp/Private/BakeMultiMeshAttributeMapsTool.cpp
+++ b/Engine/Plugins/Experimental/MeshModelingToolsetExp/Source/MeshModelingToolsExp/Private/BakeMultiMeshAttributeMapsTool.cpp
@@ -382,7 +382,7 @@ void UBakeMultiMeshAttributeMapsTool::Setup()
 	Super::Setup();
 
 	// Initialize base mesh
-	const UE::Geometry::FTransform3d BaseToWorld = UE::ToolTarget::GetLocalToWorldTransform(Targets[0]);
+	const FTransformSRT3d BaseToWorld = UE::ToolTarget::GetLocalToWorldTransform(Targets[0]);
 	PreviewMesh->ProcessMesh([this, BaseToWorld](const FDynamicMesh3& Mesh)
 	{
 		TargetMesh.Copy(Mesh);
diff --git a/Engine/Plugins/Experimental/MeshModelingToolsetExp/Source/MeshModelingToolsExp/Private/BakeTransformTool.cpp b/Engine/Plugins/Experimental/MeshModelingToolsetExp/Source/MeshModelingToolsExp/Private/BakeTransformTool.cpp
--- a/Engine/Plugins/Experimental/MeshModelingToolsetExp/Source/MeshModelingToolsExp/Private/BakeTransformTool.cpp
+++ b/Engine/Plugins/Experimental/MeshModelingToolsetExp/Source/MeshModelingToolsExp/Private/BakeTransformTool.cpp
@@ -118,26 +118,26 @@ void UBakeTransformTool::UpdateAssets()
 
 	GetToolManager()->BeginUndoTransaction(LOCTEXT("BakeTransformToolTransactionName", "Bake Transforms"));
 
-	TArray<UE::Geometry::FTransform3d> BakedTransforms;
+	TArray<FTransformSRT3d> BakedTransforms;
 	for (int32 ComponentIdx = 0; ComponentIdx < Targets.Num(); ComponentIdx++)
 	{
 		UToolTarget* Target = Targets[ComponentIdx];
 		UPrimitiveComponent* Component = UE::ToolTarget::GetTargetComponent(Target);
 		Component->Modify();
 
-		UE::Geometry::FTransform3d ComponentToWorld = UE::ToolTarget::GetLocalToWorldTransform(Target);
-		UE::Geometry::FTransform3d ToBakePart = UE::Geometry::FTransform3d::Identity();
-		UE::Geometry::FTransform3d NewWorldPart = ComponentToWorld;
+		FTransformSRT3d ComponentToWorld = UE::ToolTarget::GetLocalToWorldTransform(Target);
+		FTransformSRT3d ToBakePart = FTransformSRT3d::Identity();
+		FTransformSRT3d NewWorldPart = ComponentToWorld;
 
 		if (MapToFirstOccurrences[ComponentIdx] < ComponentIdx)
 		{
 			ToBakePart = BakedTransforms[MapToFirstOccurrences[ComponentIdx]];
 			BakedTransforms.Add(ToBakePart);
 			// try to invert baked transform
-			NewWorldPart = UE::Geometry::FTransform3d(
+			NewWorldPart = FTransformSRT3d(
 				NewWorldPart.GetRotation() * ToBakePart.GetRotation().Inverse(),
 				NewWorldPart.GetTranslation(),
-				NewWorldPart.GetScale() * UE::Geometry::FTransform3d::GetSafeScaleReciprocal(ToBakePart.GetScale())
+				NewWorldPart.GetScale() * FTransformSRT3d::GetSafeScaleReciprocal(ToBakePart.GetScale())
 			);
 			NewWorldPart.SetTranslation(NewWorldPart.GetTranslation() - NewWorldPart.TransformVector(ToBakePart.GetTranslation()));
 		}
diff --git a/Engine/Plugins/Experimental/MeshModelingToolsetExp/Source/MeshModelingToolsExp/Private/ConvertToPolygonsTool.cpp b/Engine/Plugins/Experimental/MeshModelingToolsetExp/Source/MeshModelingToolsExp/Private/ConvertToPolygonsTool.cpp
--- a/Engine/Plugins/Experimental/MeshModelingToolsetExp/Source/MeshModelingToolsExp/Private/ConvertToPolygonsTool.cpp
+++ b/Engine/Plugins/Experimental/MeshModelingToolsetExp/Source/MeshModelingToolsExp/Private/ConvertToPolygonsTool.cpp
@@ -136,9 +136,9 @@ public:
 
 	}
 
-	void SetTransform(const FTransform& Transform)
+	void SetTransform(const FTransformSRT3d& Transform)
 	{
-		ResultTransform = (UE::Geometry::FTransform3d)Transform;
+		ResultTransform = Transform;
 	}
 
 
diff --git a/Engine/Plugins/Experimental/MeshModelingToolsetExp/Source/MeshModelingToolsExp/Private/CubeGridTool.cpp b/Engine/Plugins/Experimental/MeshModelingToolsetExp/Source/MeshModelingToolsExp/Private/CubeGridTool.cpp
--- a/Engine/Plugins/Experimental/MeshModelingToolsetExp/Source/MeshModelingToolsExp/Private/CubeGridTool.cpp
+++ b/Engine/Plugins/Experimental/MeshModelingToolsetExp/Source/MeshModelingToolsExp/Private/CubeGridTool.cpp
@@ -855,8 +855,6 @@ void UCubeGridTool::UpdateComputeInputs()
 
 void UCubeGridTool::Shutdown(EToolShutdownType ShutdownType)
 {
-	using UE::Geometry::FTransform3d;
-
 	Settings->SaveProperties(this);
 	MaterialProperties->SaveProperties(this);
 
diff --git a/Engine/Plugins/Experimental/MeshModelingToolsetExp/Source/MeshModelingToolsExp/Private/DrawPolyPathTool.cpp b/Engine/Plugins/Experimental/MeshModelingToolsetExp/Source/MeshModelingToolsExp/Private/DrawPolyPathTool.cpp
--- a/Engine/Plugins/Experimental/MeshModelingToolsetExp/Source/MeshModelingToolsExp/Private/DrawPolyPathTool.cpp
+++ b/Engine/Plugins/Experimental/MeshModelingToolsetExp/Source/MeshModelingToolsExp/Private/DrawPolyPathTool.cpp
@@ -688,7 +688,7 @@ void UDrawPolyPathTool::BeginSettingWidth()
 
 	CurveDistMechanic = NewObject<USpatialCurveDistanceMechanic>(this);
 	CurveDistMechanic->Setup(this);
-	CurveDistMechanic->InitializePolyCurve(CurPolyLine, UE::Geometry::FTransform3d::Identity());
+	CurveDistMechanic->InitializePolyCurve(CurPolyLine, FTransform3d::Identity);
 
 	ExtrudeHeightMechanic = nullptr;
 
@@ -742,7 +742,7 @@ void UDrawPolyPathTool::BeginSettingRadius()
 
 	CurveDistMechanic = NewObject<USpatialCurveDistanceMechanic>(this);
 	CurveDistMechanic->Setup(this);
-	CurveDistMechanic->InitializePolyCurve(CurPolyLine, UE::Geometry::FTransform3d::Identity());
+	CurveDistMechanic->InitializePolyCurve(CurPolyLine, FTransform3d::Identity);
 
 	ExtrudeHeightMechanic = nullptr;
 
diff --git a/Engine/Plugins/Experimental/MeshModelingToolsetExp/Source/MeshModelingToolsExp/Private/DynamicMeshBrushTool.cpp b/Engine/Plugins/Experimental/MeshModelingToolsetExp/Source/MeshModelingToolsExp/Private/DynamicMeshBrushTool.cpp
--- a/Engine/Plugins/Experimental/MeshModelingToolsetExp/Source/MeshModelingToolsExp/Private/DynamicMeshBrushTool.cpp
+++ b/Engine/Plugins/Experimental/MeshModelingToolsetExp/Source/MeshModelingToolsExp/Private/DynamicMeshBrushTool.cpp
@@ -25,7 +25,7 @@ void UDynamicMeshBrushTool::Setup()
 	PreviewMesh = NewObject<UPreviewMesh>(this);
 	PreviewMesh->bBuildSpatialDataStructure = true;
 	PreviewMesh->CreateInWorld(UE::ToolTarget::GetTargetActor(Target)->GetWorld(), FTransform::Identity);
-	UE::Geometry::FTransform3d LocalToWorldTransform = UE::ToolTarget::GetLocalToWorldTransform(Target);
+	FTransformSRT3d LocalToWorldTransform = UE::ToolTarget::GetLocalToWorldTransform(Target);
 	PreviewMesh->SetTransform((FTransform)LocalToWorldTransform);
 	ToolSetupUtil::ApplyRenderingConfigurationToPreview(PreviewMesh, Target);
 
diff --git a/Engine/Plugins/Experimental/MeshModelingToolsetExp/Source/MeshModelingToolsExp/Private/DynamicMeshSculptTool.cpp b/Engine/Plugins/Experimental/MeshModelingToolsetExp/Source/MeshModelingToolsExp/Private/DynamicMeshSculptTool.cpp
--- a/Engine/Plugins/Experimental/MeshModelingToolsetExp/Source/MeshModelingToolsExp/Private/DynamicMeshSculptTool.cpp
+++ b/Engine/Plugins/Experimental/MeshModelingToolsetExp/Source/MeshModelingToolsExp/Private/DynamicMeshSculptTool.cpp
@@ -117,7 +117,7 @@ void UDynamicMeshSculptTool::Setup()
 	InitialTargetTransform.SetTranslation(FVector3d::Zero());
 	DynamicMeshComponent->ApplyTransform(InitialTargetTransform, false);
 	// since we moved to World coords there is not a current transform anymore.
-	CurTargetTransform = UE::Geometry::FTransform3d(Translation);
+	CurTargetTransform = FTransformSRT3d(Translation);
 	DynamicMeshComponent->SetWorldTransform((FTransform)CurTargetTransform);
 
 	// copy material if there is one
diff --git a/Engine/Plugins/Experimental/MeshModelingToolsetExp/Source/MeshModelingToolsExp/Private/EditPivotTool.cpp b/Engine/Plugins/Experimental/MeshModelingToolsetExp/Source/MeshModelingToolsExp/Private/EditPivotTool.cpp
--- a/Engine/Plugins/Experimental/MeshModelingToolsetExp/Source/MeshModelingToolsExp/Private/EditPivotTool.cpp
+++ b/Engine/Plugins/Experimental/MeshModelingToolsetExp/Source/MeshModelingToolsExp/Private/EditPivotTool.cpp
@@ -165,10 +165,10 @@ void UEditPivotTool::Precompute()
 	}
 	else
 	{
-		Transform = UE::Geometry::FTransform3d::Identity();
+		Transform = FTransform3d::Identity;
 		for (int k = 0; k < NumComponents; ++k)
 		{
-			UE::Geometry::FTransform3d CurTransform = UE::ToolTarget::GetLocalToWorldTransform(Targets[k]);
+			FTransform3d CurTransform = UE::ToolTarget::GetLocalToWorldTransform(Targets[k]);
 			const FMeshDescription* Mesh = UE::ToolTarget::GetMeshDescription(Targets[k]);
 			VertexIteration(Mesh, [&](int32 VertexID, const FVector& Position) {
 				ObjectBounds.Contain(CurTransform.TransformPosition((FVector3d)Position));
@@ -439,11 +439,11 @@ void UEditPivotTool::UpdateAssets(const FFrame3d& NewPivotWorldFrame)
 		}
 		else if (MapToFirstOccurrences[ComponentIdx] == ComponentIdx)
 		{
-			UE::Geometry::FTransform3d ToBake(OriginalTransforms[ComponentIdx] * NewWorldInverse);
+			FTransformSRT3d ToBake(OriginalTransforms[ComponentIdx] * NewWorldInverse);
 			// to preserve scale, after baking the first transform, we will also need to bake an inverse scale transform
 			// this bake will need to be applied separately only in cases where FTransform3d cannot correctly represent the combination of the two transforms
 			// TODO: we could skip the extra bake step if the bake functions could take a general matrix
-			UE::Geometry::FTransform3d SeparateBakeScale = UE::Geometry::FTransform3d::Identity();
+			FTransformSRT3d SeparateBakeScale = FTransformSRT3d::Identity();
 			bool bNeedSeparateScale = false;
 			// ScaledNewWorldTransform is the pivot widget's transform with the scale of the original component transform
 			FTransform ScaledNewWorldTransform = NewWorldTransform;
@@ -472,7 +472,7 @@ void UEditPivotTool::UpdateAssets(const FFrame3d& NewPivotWorldFrame)
 						(!AxisZero[2] && bEqXY)
 					));
 				bNeedSeparateScale = !bCanCombinedScales;
-				FVector3d InvScale = UE::Geometry::FTransform3d::GetSafeScaleReciprocal(OriginalTransforms[ComponentIdx].GetScale3D());
+				FVector3d InvScale = FTransformSRT3d::GetSafeScaleReciprocal(OriginalTransforms[ComponentIdx].GetScale3D());
 				if (!bNeedSeparateScale)
 				{
 					ToBake.SetScale(FVector3d::One()); // clear scale; it will be fully captured by the new transform
diff --git a/Engine/Plugins/Experimental/MeshModelingToolsetExp/Source/MeshModelingToolsExp/Private/EditUVIslandsTool.cpp b/Engine/Plugins/Experimental/MeshModelingToolsetExp/Source/MeshModelingToolsExp/Private/EditUVIslandsTool.cpp
--- a/Engine/Plugins/Experimental/MeshModelingToolsetExp/Source/MeshModelingToolsExp/Private/EditUVIslandsTool.cpp
+++ b/Engine/Plugins/Experimental/MeshModelingToolsetExp/Source/MeshModelingToolsExp/Private/EditUVIslandsTool.cpp
@@ -52,7 +52,7 @@ void UEditUVIslandsTool::Setup()
 	DynamicMeshComponent->RegisterComponent();
 	DynamicMeshComponent->SetWorldTransform((FTransform)UE::ToolTarget::GetLocalToWorldTransform(Target));
 	ToolSetupUtil::ApplyRenderingConfigurationToPreview(DynamicMeshComponent, Target);
-	WorldTransform = UE::Geometry::FTransform3d(DynamicMeshComponent->GetComponentTransform());
+	WorldTransform = FTransform3d(DynamicMeshComponent->GetComponentTransform());
 
 	// set materials
 	FComponentMaterialSet MaterialSet = UE::ToolTarget::GetMaterialSet(Target);
diff --git a/Engine/Plugins/Experimental/MeshModelingToolsetExp/Source/MeshModelingToolsExp/Private/HoleFillTool.cpp b/Engine/Plugins/Experimental/MeshModelingToolsetExp/Source/MeshModelingToolsExp/Private/HoleFillTool.cpp
--- a/Engine/Plugins/Experimental/MeshModelingToolsetExp/Source/MeshModelingToolsExp/Private/HoleFillTool.cpp
+++ b/Engine/Plugins/Experimental/MeshModelingToolsetExp/Source/MeshModelingToolsExp/Private/HoleFillTool.cpp
@@ -92,7 +92,7 @@ TUniquePtr<FDynamicMeshOperator> UHoleFillOperatorFactory::MakeNewOperator()
 	TUniquePtr<FHoleFillOp> FillOp = MakeUnique<FHoleFillOp>();
 
 	FTransform LocalToWorld = Cast<IPrimitiveComponentBackedTarget>(FillTool->Target)->GetWorldTransform();
-	FillOp->SetResultTransform((UE::Geometry::FTransform3d)LocalToWorld);
+	FillOp->SetResultTransform((FTransformSRT3d)LocalToWorld);
 	FillOp->OriginalMesh = FillTool->OriginalMesh;
 	FillOp->MeshUVScaleFactor = FillTool->MeshUVScaleFactor;
 	FillTool->GetLoopsToFill(FillOp->Loops);
diff --git a/Engine/Plugins/Experimental/MeshModelingToolsetExp/Source/MeshModelingToolsExp/Private/LatticeDeformerTool.cpp b/Engine/Plugins/Experimental/MeshModelingToolsetExp/Source/MeshModelingToolsExp/Private/LatticeDeformerTool.cpp
--- a/Engine/Plugins/Experimental/MeshModelingToolsetExp/Source/MeshModelingToolsExp/Private/LatticeDeformerTool.cpp
+++ b/Engine/Plugins/Experimental/MeshModelingToolsetExp/Source/MeshModelingToolsExp/Private/LatticeDeformerTool.cpp
@@ -153,7 +153,7 @@ void ULatticeDeformerTool::InitializeLattice(TArray<FVector3d>& OutLatticePoints
 	Lattice->GenerateInitialLatticePositions(OutLatticePoints);
 
 	// Put the lattice in world space
-	UE::Geometry::FTransform3d LocalToWorld(Cast<IPrimitiveComponentBackedTarget>(Target)->GetWorldTransform());
+	FTransform3d LocalToWorld(Cast<IPrimitiveComponentBackedTarget>(Target)->GetWorldTransform());
 	Algo::ForEach(OutLatticePoints, [&LocalToWorld](FVector3d& Point) {
 		Point = LocalToWorld.TransformPosition(Point);
 	});
@@ -216,7 +216,7 @@ void ULatticeDeformerTool::Setup()
 	ControlPointsMechanic = NewObject<ULatticeControlPointsMechanic>(this);
 	ControlPointsMechanic->Setup(this);
 	ControlPointsMechanic->SetWorld(GetTargetWorld());
-	UE::Geometry::FTransform3d LocalToWorld(Cast<IPrimitiveComponentBackedTarget>(Target)->GetWorldTransform());
+	FTransform3d LocalToWorld(Cast<IPrimitiveComponentBackedTarget>(Target)->GetWorldTransform());
 	ControlPointsMechanic->Initialize(LatticePoints, LatticeEdges, LocalToWorld);
 
 	auto OnPointsChangedLambda = [this]()
@@ -357,7 +357,7 @@ void ULatticeDeformerTool::OnShutdown(EToolShutdownType ShutdownType)
 
 			// The lattice and its output mesh are in world space, so get them in local space.
 			// TODO: Would it make more sense to do all the lattice computation in local space?
-			UE::Geometry::FTransform3d LocalToWorld(TargetComponent->GetWorldTransform());
+			FTransform3d LocalToWorld(TargetComponent->GetWorldTransform());
 			MeshTransforms::ApplyTransformInverse(*DynamicMeshResult, LocalToWorld);
 
 			UE::ToolTarget::CommitMeshDescriptionUpdateViaDynamicMesh(Target, *DynamicMeshResult, true);
@@ -432,7 +432,7 @@ void ULatticeDeformerTool::OnTick(float DeltaTime)
 			TArray<FVector3d> LatticePoints;
 			TArray<FVector2i> LatticeEdges;
 			InitializeLattice(LatticePoints, LatticeEdges);
-			UE::Geometry::FTransform3d LocalToWorld(Cast<IPrimitiveComponentBackedTarget>(Target)->GetWorldTransform());
+			FTransform3d LocalToWorld(Cast<IPrimitiveComponentBackedTarget>(Target)->GetWorldTransform());
 			ControlPointsMechanic->Initialize(LatticePoints, LatticeEdges, LocalToWorld);
 			Preview->InvalidateResult();
 			bShouldRebuild = false;
diff --git a/Engine/Plugins/Experimental/MeshModelingToolsetExp/Source/MeshModelingToolsExp/Private/MeshAttributePaintTool.cpp b/Engine/Plugins/Experimental/MeshModelingToolsetExp/Source/MeshModelingToolsExp/Private/MeshAttributePaintTool.cpp
--- a/Engine/Plugins/Experimental/MeshModelingToolsetExp/Source/MeshModelingToolsExp/Private/MeshAttributePaintTool.cpp
+++ b/Engine/Plugins/Experimental/MeshModelingToolsetExp/Source/MeshModelingToolsExp/Private/MeshAttributePaintTool.cpp
@@ -376,7 +376,7 @@ bool UMeshAttributePaintTool::OnUpdateHover(const FInputDeviceRay& DevicePos)
 
 void UMeshAttributePaintTool::CalculateVertexROI(const FBrushStampData& Stamp, TArray<int>& VertexROI)
 {
-	UE::Geometry::FTransform3d Transform(Cast<IPrimitiveComponentBackedTarget>(Target)->GetWorldTransform());
+	FTransform3d Transform(Cast<IPrimitiveComponentBackedTarget>(Target)->GetWorldTransform());
 	FVector3d StampPosLocal = Transform.InverseTransformPosition((FVector3d)Stamp.WorldPosition);
 
 	float Radius = GetCurrentBrushRadiusLocal();
@@ -543,7 +543,7 @@ void UMeshAttributePaintTool::ApplyStamp(const FBrushStampData& Stamp)
 
 void UMeshAttributePaintTool::ApplyStamp_Paint(const FBrushStampData& Stamp, FStampActionData& ActionData)
 {
-	UE::Geometry::FTransform3d Transform(Cast<IPrimitiveComponentBackedTarget>(Target)->GetWorldTransform());
+	FTransform3d Transform(Cast<IPrimitiveComponentBackedTarget>(Target)->GetWorldTransform());
 	FVector3d StampPosLocal = Transform.InverseTransformPosition((FVector3d)Stamp.WorldPosition);
 
 	int32 NumVertices = ActionData.ROIVertices.Num();
diff --git a/Engine/Plugins/Experimental/MeshModelingToolsetExp/Source/MeshModelingToolsExp/Private/MeshSelectionTool.cpp b/Engine/Plugins/Experimental/MeshModelingToolsetExp/Source/MeshModelingToolsExp/Private/MeshSelectionTool.cpp
--- a/Engine/Plugins/Experimental/MeshModelingToolsetExp/Source/MeshModelingToolsExp/Private/MeshSelectionTool.cpp
+++ b/Engine/Plugins/Experimental/MeshModelingToolsetExp/Source/MeshModelingToolsExp/Private/MeshSelectionTool.cpp
@@ -359,7 +359,7 @@ bool UMeshSelectionTool::HitTest(const FRay& Ray, FHitResult& OutHit)
 		SourceMesh->GetTriInfo(OutHit.FaceIndex, Normal, Area, Centroid);
 		FViewCameraState StateOut;
 		GetToolManager()->GetContextQueriesAPI()->GetCurrentViewState(StateOut);
-		UE::Geometry::FTransform3d LocalToWorld = UE::ToolTarget::GetLocalToWorldTransform(Target);
+		FTransform3d LocalToWorld = UE::ToolTarget::GetLocalToWorldTransform(Target);
 		FVector3d LocalEyePosition(LocalToWorld.InverseTransformPosition(StateOut.Position));
 
 		if (Normal.Dot((Centroid - LocalEyePosition)) > 0)
@@ -415,7 +415,7 @@ TUniquePtr<FDynamicMeshOctree3>& UMeshSelectionTool::GetOctree()
 
 void UMeshSelectionTool::CalculateVertexROI(const FBrushStampData& Stamp, TArray<int>& VertexROI)
 {
-	UE::Geometry::FTransform3d Transform = UE::ToolTarget::GetLocalToWorldTransform(Target);
+	FTransform3d Transform = UE::ToolTarget::GetLocalToWorldTransform(Target);
 	FVector3d StampPosLocal = Transform.InverseTransformPosition((FVector3d)Stamp.WorldPosition);
 
 	// TODO: need dynamic vertex hash table!
@@ -437,7 +437,7 @@ void UMeshSelectionTool::CalculateVertexROI(const FBrushStampData& Stamp, TArray
 
 void UMeshSelectionTool::CalculateTriangleROI(const FBrushStampData& Stamp, TArray<int>& TriangleROI)
 {
-	UE::Geometry::FTransform3d Transform = UE::ToolTarget::GetLocalToWorldTransform(Target);
+	FTransform3d Transform = UE::ToolTarget::GetLocalToWorldTransform(Target);
 	FVector3d StampPosLocal = Transform.InverseTransformPosition((FVector3d)Stamp.WorldPosition);
 
 	// always select first triangle
@@ -590,7 +590,7 @@ void UMeshSelectionTool::UpdateFaceSelection(const FBrushStampData& Stamp, const
 	{
 		FViewCameraState StateOut;
 		GetToolManager()->GetContextQueriesAPI()->GetCurrentViewState(StateOut);
-		UE::Geometry::FTransform3d LocalToWorld = UE::ToolTarget::GetLocalToWorldTransform(Target);
+		FTransform3d LocalToWorld = UE::ToolTarget::GetLocalToWorldTransform(Target);
 		FVector3d LocalEyePosition(LocalToWorld.InverseTransformPosition(StateOut.Position));
 
 		for (int tid : TriangleROI)
@@ -701,7 +701,7 @@ FBox UMeshSelectionTool::GetWorldSpaceFocusBox()
 		FAxisAlignedBox3d Bounds = FAxisAlignedBox3d::Empty();
 
 		const FDynamicMesh3* Mesh = PreviewMesh->GetMesh();
-		UE::Geometry::FTransform3d Transform(PreviewMesh->GetTransform());
+		FTransform3d Transform(PreviewMesh->GetTransform());
 		for (int32 tid : SelectedFaces)
 		{
 			FIndex3i Tri = Mesh->GetTriangle(tid);
@@ -1479,7 +1479,7 @@ void UMeshSelectionTool::SeparateSelectedTriangles(bool bDeleteSelected)
 	Editor.AppendTriangles(SourceMesh, SelectedFaces, Mappings, EditResult);
 
 	// emit new asset
-	UE::Geometry::FTransform3d Transform(PreviewMesh->GetTransform());
+	FTransform3d Transform(PreviewMesh->GetTransform());
 	GetToolManager()->BeginUndoTransaction(
 		(bDeleteSelected) ? LOCTEXT("MeshSelectionToolSeparate", "Separate") : LOCTEXT("MeshSelectionToolDuplicate", "Duplicate") );
 
diff --git a/Engine/Plugins/Experimental/MeshModelingToolsetExp/Source/MeshModelingToolsExp/Private/MirrorTool.cpp b/Engine/Plugins/Experimental/MeshModelingToolsetExp/Source/MeshModelingToolsExp/Private/MirrorTool.cpp
--- a/Engine/Plugins/Experimental/MeshModelingToolsetExp/Source/MeshModelingToolsExp/Private/MirrorTool.cpp
+++ b/Engine/Plugins/Experimental/MeshModelingToolsetExp/Source/MeshModelingToolsExp/Private/MirrorTool.cpp
@@ -61,7 +61,7 @@ TUniquePtr<FDynamicMeshOperator> UMirrorOperatorFactory::MakeNewOperator()
 		}
 	}
 	LocalToWorld.SetScale3D(LocalToWorldScale);
-	UE::Geometry::FTransform3d WorldToLocal = UE::Geometry::FTransform3d(LocalToWorld).Inverse();
+	FTransformSRT3d WorldToLocal = FTransformSRT3d(LocalToWorld).Inverse();
 
 	// Now we can get the plane parameters in local space.
 	MirrorOp->LocalPlaneOrigin = WorldToLocal.TransformPosition(MirrorTool->MirrorPlaneOrigin);;
diff --git a/Engine/Plugins/Experimental/MeshModelingToolsetExp/Source/MeshModelingToolsExp/Private/Physics/CollisionGeometryVisualization.cpp b/Engine/Plugins/Experimental/MeshModelingToolsetExp/Source/MeshModelingToolsExp/Private/Physics/CollisionGeometryVisualization.cpp
--- a/Engine/Plugins/Experimental/MeshModelingToolsetExp/Source/MeshModelingToolsExp/Private/Physics/CollisionGeometryVisualization.cpp
+++ b/Engine/Plugins/Experimental/MeshModelingToolsetExp/Source/MeshModelingToolsExp/Private/Physics/CollisionGeometryVisualization.cpp
@@ -22,7 +22,7 @@ void UE::PhysicsTools::InitializePreviewGeometryLines(const FPhysicsDataCollecti
 		const FKSphereElem& Sphere = AggGeom.SphereElems[Index];
 		FTransform ElemTransform = Sphere.GetTransform();
 		ElemTransform.ScaleTranslation(PhysicsData.ExternalScale3D);
-		UE::Geometry::FTransform3f ElemTransformf(ElemTransform);
+		FTransformSRT3f ElemTransformf(ElemTransform);
 		float Radius = PhysicsData.ExternalScale3D.GetAbsMin() * Sphere.Radius;
 		UE::Geometry::GenerateCircleSegments<float>(CircleSteps, Radius, FVector3f::Zero(), FVector3f::UnitX(), FVector3f::UnitY(), ElemTransformf,
 			[&](const FVector3f& A, const FVector3f& B) { LinesOut.Add(FRenderableLine((FVector)A, (FVector)B, SphereColor, LineThickness, DepthBias)); });
@@ -39,7 +39,7 @@ void UE::PhysicsTools::InitializePreviewGeometryLines(const FPhysicsDataCollecti
 		const FKBoxElem& Box = AggGeom.BoxElems[Index];
 		FTransform ElemTransform = Box.GetTransform();
 		ElemTransform.ScaleTranslation(PhysicsData.ExternalScale3D);
-		UE::Geometry::FTransform3f ElemTransformf(ElemTransform);
+		FTransformSRT3f ElemTransformf(ElemTransform);
 		FVector3f HalfDimensions(
 			PhysicsData.ExternalScale3D.X * Box.X * 0.5f,
 			PhysicsData.ExternalScale3D.Y * Box.Y * 0.5f,
@@ -55,7 +55,7 @@ void UE::PhysicsTools::InitializePreviewGeometryLines(const FPhysicsDataCollecti
 		const FKSphylElem& Capsule = AggGeom.SphylElems[Index];
 		FTransform ElemTransform = Capsule.GetTransform();
 		ElemTransform.ScaleTranslation(PhysicsData.ExternalScale3D);
-		UE::Geometry::FTransform3f ElemTransformf(ElemTransform);
+		FTransformSRT3f ElemTransformf(ElemTransform);
 		const float HalfLength = Capsule.GetScaledCylinderLength(PhysicsData.ExternalScale3D) * .5f;
 		const float Radius = Capsule.GetScaledRadius(PhysicsData.ExternalScale3D);
 		FVector3f Top(0, 0, HalfLength), Bottom(0, 0, -HalfLength);
diff --git a/Engine/Plugins/Experimental/MeshModelingToolsetExp/Source/MeshModelingToolsExp/Private/Physics/ExtractCollisionGeometryTool.cpp b/Engine/Plugins/Experimental/MeshModelingToolsetExp/Source/MeshModelingToolsExp/Private/Physics/ExtractCollisionGeometryTool.cpp
--- a/Engine/Plugins/Experimental/MeshModelingToolsetExp/Source/MeshModelingToolsExp/Private/Physics/ExtractCollisionGeometryTool.cpp
+++ b/Engine/Plugins/Experimental/MeshModelingToolsetExp/Source/MeshModelingToolsExp/Private/Physics/ExtractCollisionGeometryTool.cpp
@@ -119,12 +119,11 @@ void UExtractCollisionGeometryTool::Setup()
 
 void UExtractCollisionGeometryTool::OnShutdown(EToolShutdownType ShutdownType)
 {
-	using FTransform3d = UE::Geometry::FTransform3d;
 	OutputTypeProperties->SaveProperties(this, TEXT("ExtractCollisionTool"));
 	Settings->SaveProperties(this);
 	VizSettings->SaveProperties(this);
 
-	UE::Geometry::FTransform3d ActorTransform(PreviewMesh->GetTransform());
+	FTransform3d ActorTransform(PreviewMesh->GetTransform());
 
 	PreviewElements->Disconnect();
 	PreviewMesh->SetVisible(false);
@@ -295,7 +294,7 @@ void UExtractCollisionGeometryTool::RecalculateMesh_Simple()
 
 		MeshTransforms::Translate(CapsuleMesh, FVector3d(0,0,-0.5*Capsule.Length) );
 
-		UE::Geometry::FTransform3d Transform(Capsule.GetTransform());
+		FTransformSRT3d Transform(Capsule.GetTransform());
 		MeshTransforms::ApplyTransform(CapsuleMesh, Transform);
 
 		FMeshIndexMappings Mappings;
@@ -306,7 +305,7 @@ void UExtractCollisionGeometryTool::RecalculateMesh_Simple()
 
 	for (const FKConvexElem& Convex : AggGeom.ConvexElems)
 	{
-		UE::Geometry::FTransform3d ElemTransform(Convex.GetTransform());
+		FTransformSRT3d ElemTransform(Convex.GetTransform());
 		FDynamicMesh3 ConvexMesh(EMeshComponents::None);
 		int32 NumVertices = Convex.VertexData.Num();
 		for (int32 k = 0; k < NumVertices; ++k)
diff --git a/Engine/Plugins/Experimental/MeshModelingToolsetExp/Source/MeshModelingToolsExp/Private/Physics/SetCollisionGeometryTool.cpp b/Engine/Plugins/Experimental/MeshModelingToolsetExp/Source/MeshModelingToolsExp/Private/Physics/SetCollisionGeometryTool.cpp
--- a/Engine/Plugins/Experimental/MeshModelingToolsetExp/Source/MeshModelingToolsExp/Private/Physics/SetCollisionGeometryTool.cpp
+++ b/Engine/Plugins/Experimental/MeshModelingToolsetExp/Source/MeshModelingToolsExp/Private/Physics/SetCollisionGeometryTool.cpp
@@ -482,8 +482,8 @@ void USetCollisionGeometryTool::PrecomputeInputMeshes()
 	}
 
 	UToolTarget* CollisionTarget = Targets[Targets.Num() - 1];
-	UE::Geometry::FTransform3d TargetTransform(UE::ToolTarget::GetLocalToWorldTransform(CollisionTarget));
-	UE::Geometry::FTransform3d TargetTransformInv = TargetTransform.Inverse();
+	FTransformSRT3d TargetTransform(UE::ToolTarget::GetLocalToWorldTransform(CollisionTarget));
+	FTransformSRT3d TargetTransformInv = TargetTransform.Inverse();
 
 	InputMeshes.Reset();
 	InputMeshes.SetNum(SourceObjectIndices.Num());
@@ -492,7 +492,7 @@ void USetCollisionGeometryTool::PrecomputeInputMeshes()
 		FDynamicMesh3 SourceMesh = InitialSourceMeshes[k];
 		if (Settings->bUseWorldSpace)
 		{
-			UE::Geometry::FTransform3d ToWorld(UE::ToolTarget::GetLocalToWorldTransform(Targets[k]));
+			FTransformSRT3d ToWorld(UE::ToolTarget::GetLocalToWorldTransform(Targets[k]));
 			MeshTransforms::ApplyTransform(SourceMesh, ToWorld);
 			MeshTransforms::ApplyTransform(SourceMesh, TargetTransformInv);
 		}
diff --git a/Engine/Plugins/Experimental/MeshModelingToolsetExp/Source/MeshModelingToolsExp/Private/PlaneCutTool.cpp b/Engine/Plugins/Experimental/MeshModelingToolsetExp/Source/MeshModelingToolsExp/Private/PlaneCutTool.cpp
--- a/Engine/Plugins/Experimental/MeshModelingToolsetExp/Source/MeshModelingToolsExp/Private/PlaneCutTool.cpp
+++ b/Engine/Plugins/Experimental/MeshModelingToolsetExp/Source/MeshModelingToolsExp/Private/PlaneCutTool.cpp
@@ -289,7 +289,7 @@ TUniquePtr<FDynamicMeshOperator> UPlaneCutOperatorFactory::MakeNewOperator()
 
 	FVector LocalOrigin = WorldToLocal.TransformPosition((FVector)CutTool->CutPlaneWorld.Origin);
 	FVector3d WorldNormal = CutTool->CutPlaneWorld.GetAxis(2);
-	UE::Geometry::FTransform3d W2LForNormal(WorldToLocal);
+	FTransformSRT3d W2LForNormal(WorldToLocal);
 	FVector LocalNormal = (FVector)W2LForNormal.TransformNormal(WorldNormal);
 	FVector BackTransformed = LocalToWorld.TransformVector(LocalNormal);
 	float NormalScaleFactor = FVector::DotProduct(BackTransformed, (FVector)WorldNormal);
diff --git a/Engine/Plugins/Experimental/MeshModelingToolsetExp/Source/MeshModelingToolsExp/Private/ProjectToTargetTool.cpp b/Engine/Plugins/Experimental/MeshModelingToolsetExp/Source/MeshModelingToolsExp/Private/ProjectToTargetTool.cpp
--- a/Engine/Plugins/Experimental/MeshModelingToolsetExp/Source/MeshModelingToolsExp/Private/ProjectToTargetTool.cpp
+++ b/Engine/Plugins/Experimental/MeshModelingToolsetExp/Source/MeshModelingToolsExp/Private/ProjectToTargetTool.cpp
@@ -61,7 +61,6 @@ TUniquePtr<FDynamicMeshOperator> UProjectToTargetTool::MakeNewOperator()
 	RemeshOp->ProjectionTarget = ProjectionTarget.Get();
 	RemeshOp->ProjectionTargetSpatial = ProjectionTargetSpatial.Get();
 
-	using FTransform3d = UE::Geometry::FTransform3d;
 	RemeshOp->ToolMeshLocalToWorld = UE::ToolTarget::GetLocalToWorldTransform(Targets[0]);
 	RemeshOp->TargetMeshLocalToWorld = UE::ToolTarget::GetLocalToWorldTransform(Targets[1]);
 	RemeshOp->bUseWorldSpace = ProjectProperties->bWorldSpace;
diff --git a/Engine/Plugins/Experimental/MeshModelingToolsetExp/Source/MeshModelingToolsExp/Private/RemoveOccludedTrianglesTool.cpp b/Engine/Plugins/Experimental/MeshModelingToolsetExp/Source/MeshModelingToolsExp/Private/RemoveOccludedTrianglesTool.cpp
--- a/Engine/Plugins/Experimental/MeshModelingToolsetExp/Source/MeshModelingToolsExp/Private/RemoveOccludedTrianglesTool.cpp
+++ b/Engine/Plugins/Experimental/MeshModelingToolsetExp/Source/MeshModelingToolsExp/Private/RemoveOccludedTrianglesTool.cpp
@@ -369,7 +369,7 @@ TUniquePtr<FDynamicMeshOperator> URemoveOccludedTrianglesOperatorFactory::MakeNe
 		int32 TargetIdx = Tool->PreviewToTargetIdx[PreviewIdx];
 		Op->OccluderTrees.Add(Tool->OccluderTrees[TargetIdx]);
 		Op->OccluderWindings.Add(Tool->OccluderWindings[TargetIdx]);
-		Op->OccluderTransforms.Add(UE::Geometry::FTransform3d::Identity());
+		Op->OccluderTransforms.Add(FTransformSRT3d::Identity());
 	}
 	else
 	{
@@ -390,10 +390,10 @@ TUniquePtr<FDynamicMeshOperator> URemoveOccludedTrianglesOperatorFactory::MakeNe
 	
 	Op->SetTransform(LocalToWorld);
 
-	Op->MeshTransforms.Add((UE::Geometry::FTransform3d)LocalToWorld);
+	Op->MeshTransforms.Add((FTransformSRT3d)LocalToWorld);
 	for (int32 CopyIdx : Tool->PreviewToCopyIdx[PreviewIdx])
 	{
-		Op->MeshTransforms.Add((UE::Geometry::FTransform3d)Tool->PreviewCopies[CopyIdx]->GetTransform());
+		Op->MeshTransforms.Add((FTransformSRT3d)Tool->PreviewCopies[CopyIdx]->GetTransform());
 	}
 
 	return Op;
diff --git a/Engine/Plugins/Experimental/MeshModelingToolsetExp/Source/MeshModelingToolsExp/Private/SeamSculptTool.cpp b/Engine/Plugins/Experimental/MeshModelingToolsetExp/Source/MeshModelingToolsExp/Private/SeamSculptTool.cpp
--- a/Engine/Plugins/Experimental/MeshModelingToolsetExp/Source/MeshModelingToolsExp/Private/SeamSculptTool.cpp
+++ b/Engine/Plugins/Experimental/MeshModelingToolsetExp/Source/MeshModelingToolsExp/Private/SeamSculptTool.cpp
@@ -64,7 +64,7 @@ void USeamSculptTool::Setup()
 	BrushProperties->RestoreProperties(this);
 
 	IPrimitiveComponentBackedTarget* TargetComponent = Cast<IPrimitiveComponentBackedTarget>(Target);
-	MeshTransform = UE::Geometry::FTransform3d(TargetComponent->GetWorldTransform());
+	MeshTransform = FTransform3d(TargetComponent->GetWorldTransform());
 	InputMesh = MakeShared<FDynamicMesh3, ESPMode::ThreadSafe>(*PreviewMesh->GetMesh());
 	FMeshNormals::QuickComputeVertexNormals(*InputMesh);
 	NormalOffset = InputMesh->GetBounds(true).MinDim() * 0.001;
diff --git a/Engine/Plugins/Experimental/MeshModelingToolsetExp/Source/MeshModelingToolsExp/Private/SelfUnionMeshesTool.cpp b/Engine/Plugins/Experimental/MeshModelingToolsetExp/Source/MeshModelingToolsExp/Private/SelfUnionMeshesTool.cpp
--- a/Engine/Plugins/Experimental/MeshModelingToolsetExp/Source/MeshModelingToolsExp/Private/SelfUnionMeshesTool.cpp
+++ b/Engine/Plugins/Experimental/MeshModelingToolsetExp/Source/MeshModelingToolsExp/Private/SelfUnionMeshesTool.cpp
@@ -113,7 +113,7 @@ void USelfUnionMeshesTool::ConvertInputsAndSetPreviewMaterials(bool bSetPreviewM
 			MaterialIDs->SetValue(TID, MaterialRemap[ComponentIdx][MaterialIDs->GetValue(TID)]);
 		}
 		// TODO: center the meshes
-		FTransform3d WorldTransform = (FTransform3d)TransformProxies[ComponentIdx]->GetTransform();
+		FTransformSRT3d WorldTransform = TransformProxies[ComponentIdx]->GetTransform();
 		if (WorldTransform.GetDeterminant() < 0)
 		{
 			ComponentMesh.ReverseOrientation(false);
@@ -198,7 +198,7 @@ TUniquePtr<FDynamicMeshOperator> USelfUnionMeshesTool::MakeNewOperator()
 	Op->WindingNumberThreshold = Properties->WindingThreshold;
 	Op->bTrimFlaps = Properties->bTrimFlaps;
 
-	Op->SetResultTransform(FTransform3d::Identity()); // TODO Center the combined meshes (when building them) and change this transform accordingly
+	Op->SetResultTransform(FTransformSRT3d::Identity()); // TODO Center the combined meshes (when building them) and change this transform accordingly
 	Op->CombinedMesh = CombinedSourceMeshes;
 
 	return Op;
diff --git a/Engine/Plugins/Experimental/MeshModelingToolsetExp/Source/MeshModelingToolsExp/Private/SkinWeightsBindingTool.cpp b/Engine/Plugins/Experimental/MeshModelingToolsetExp/Source/MeshModelingToolsExp/Private/SkinWeightsBindingTool.cpp
--- a/Engine/Plugins/Experimental/MeshModelingToolsetExp/Source/MeshModelingToolsExp/Private/SkinWeightsBindingTool.cpp
+++ b/Engine/Plugins/Experimental/MeshModelingToolsetExp/Source/MeshModelingToolsExp/Private/SkinWeightsBindingTool.cpp
@@ -993,8 +993,8 @@ TUniquePtr<UE::Geometry::FDynamicMeshOperator> USkinWeightsBindingTool::MakeNewO
 	Op->OriginalMesh = OriginalMesh;
 	Op->TransformHierarchy = TransformHierarchy;
 
-	const FTransform LocalToWorld = (FTransform) UE::ToolTarget::GetLocalToWorldTransform(Targets[0]);
-	Op->SetResultTransform(static_cast<UE::Geometry::FTransform3d>(LocalToWorld));
+	const FTransform3d LocalToWorld = UE::ToolTarget::GetLocalToWorldTransform(Targets[0]);
+	Op->SetResultTransform(LocalToWorld);
 	
 	return Op;
 }
diff --git a/Engine/Plugins/Experimental/MeshModelingToolsetExp/Source/MeshModelingToolsExp/Private/SkinWeightsPaintTool.cpp b/Engine/Plugins/Experimental/MeshModelingToolsetExp/Source/MeshModelingToolsExp/Private/SkinWeightsPaintTool.cpp
--- a/Engine/Plugins/Experimental/MeshModelingToolsetExp/Source/MeshModelingToolsExp/Private/SkinWeightsPaintTool.cpp
+++ b/Engine/Plugins/Experimental/MeshModelingToolsetExp/Source/MeshModelingToolsExp/Private/SkinWeightsPaintTool.cpp
@@ -281,7 +281,7 @@ void USkinWeightsPaintTool::CalculateVertexROI(const FBrushStampData& Stamp, TAr
 {
 	using namespace UE::Geometry;
 	IPrimitiveComponentBackedTarget* TargetComponent = Cast<IPrimitiveComponentBackedTarget>(Target);
-	UE::Geometry::FTransform3d Transform(TargetComponent->GetWorldTransform());
+	FTransform3d Transform(TargetComponent->GetWorldTransform());
 	FVector3d StampPosLocal = Transform.InverseTransformPosition((FVector3d)Stamp.WorldPosition);
 
 	float RadiusSqr = CurrentBrushRadius * CurrentBrushRadius;
@@ -352,7 +352,7 @@ void USkinWeightsPaintTool::ApplyStamp(const FBrushStampData& Stamp)
 		return;
 
 	IPrimitiveComponentBackedTarget* TargetComponent = Cast<IPrimitiveComponentBackedTarget>(Target);
-	UE::Geometry::FTransform3d Transform(TargetComponent->GetWorldTransform());
+	FTransform3d Transform(TargetComponent->GetWorldTransform());
 	FVector3d StampPosLocal = Transform.InverseTransformPosition((FVector3d)Stamp.WorldPosition);
 
 	TArray<int32> ROIVertices;
diff --git a/Engine/Plugins/Experimental/MeshModelingToolsetExp/Source/MeshModelingToolsExp/Private/SplitMeshesTool.cpp b/Engine/Plugins/Experimental/MeshModelingToolsetExp/Source/MeshModelingToolsExp/Private/SplitMeshesTool.cpp
--- a/Engine/Plugins/Experimental/MeshModelingToolsetExp/Source/MeshModelingToolsExp/Private/SplitMeshesTool.cpp
+++ b/Engine/Plugins/Experimental/MeshModelingToolsetExp/Source/MeshModelingToolsExp/Private/SplitMeshesTool.cpp
@@ -103,7 +103,7 @@ void USplitMeshesTool::OnShutdown(EToolShutdownType ShutdownType)
 			check(TargetActor != nullptr);
 			DeleteActors.Add(TargetActor);
 
-			UE::Geometry::FTransform3d SourceTransform = UE::ToolTarget::GetLocalToWorldTransform(Targets[ti]);
+			FTransform3d SourceTransform = UE::ToolTarget::GetLocalToWorldTransform(Targets[ti]);
 			FDynamicMesh3 SourceMesh = UE::ToolTarget::GetDynamicMeshCopy(Targets[ti], true);
 			FString AssetName = TargetActor->GetActorNameOrLabel();
 
@@ -124,7 +124,7 @@ void USplitMeshesTool::OnShutdown(EToolShutdownType ShutdownType)
 			{
 				FCreateMeshObjectParams NewMeshObjectParams = BaseMeshObjectParams;
 				NewMeshObjectParams.BaseName = FString::Printf(TEXT("%s_%d"), *AssetName, k);
-				UE::Geometry::FTransform3d PartTransform = SourceTransform;
+				FTransform3d PartTransform = SourceTransform;
 				PartTransform.SetTranslation(SourceTransform.GetTranslation() + SplitInfo.Origins[k]);
 				NewMeshObjectParams.Transform = (FTransform)PartTransform;
 				if (BasicProperties->bTransferMaterials)
diff --git a/Engine/Plugins/Experimental/MeshModelingToolsetExp/Source/MeshModelingToolsExp/Private/VolumeToMeshTool.cpp b/Engine/Plugins/Experimental/MeshModelingToolsetExp/Source/MeshModelingToolsExp/Private/VolumeToMeshTool.cpp
--- a/Engine/Plugins/Experimental/MeshModelingToolsetExp/Source/MeshModelingToolsExp/Private/VolumeToMeshTool.cpp
+++ b/Engine/Plugins/Experimental/MeshModelingToolsetExp/Source/MeshModelingToolsExp/Private/VolumeToMeshTool.cpp
@@ -111,7 +111,7 @@ void UVolumeToMeshTool::Shutdown(EToolShutdownType ShutdownType)
 	OutputTypeProperties->SaveProperties(this, TEXT("VolumeToMeshTool"));
 	Settings->SaveProperties(this);
 
-	UE::Geometry::FTransform3d Transform(PreviewMesh->GetTransform());
+	FTransform3d Transform(PreviewMesh->GetTransform());
 	PreviewMesh->SetVisible(false);
 	PreviewMesh->Disconnect();
 	PreviewMesh = nullptr;
diff --git a/Engine/Plugins/Experimental/MeshModelingToolsetExp/Source/MeshModelingToolsExp/Private/WeldMeshEdgesTool.cpp b/Engine/Plugins/Experimental/MeshModelingToolsetExp/Source/MeshModelingToolsExp/Private/WeldMeshEdgesTool.cpp
--- a/Engine/Plugins/Experimental/MeshModelingToolsetExp/Source/MeshModelingToolsExp/Private/WeldMeshEdgesTool.cpp
+++ b/Engine/Plugins/Experimental/MeshModelingToolsetExp/Source/MeshModelingToolsExp/Private/WeldMeshEdgesTool.cpp
@@ -68,9 +68,9 @@ public:
 		}
 	}
 
-	void SetTransform(const FTransform& Transform)
+	void SetTransform(const FTransformSRT3d& Transform)
 	{
-		ResultTransform = (UE::Geometry::FTransform3d)Transform;
+		ResultTransform = Transform;
 	}
 };
 
diff --git a/Engine/Plugins/Experimental/MeshModelingToolsetExp/Source/MeshModelingToolsExp/Public/AlignObjectsTool.h b/Engine/Plugins/Experimental/MeshModelingToolsetExp/Source/MeshModelingToolsExp/Public/AlignObjectsTool.h
--- a/Engine/Plugins/Experimental/MeshModelingToolsetExp/Source/MeshModelingToolsExp/Public/AlignObjectsTool.h
+++ b/Engine/Plugins/Experimental/MeshModelingToolsetExp/Source/MeshModelingToolsExp/Public/AlignObjectsTool.h
@@ -102,7 +102,6 @@ UCLASS()
 class MESHMODELINGTOOLSEXP_API UAlignObjectsTool : public UMultiSelectionMeshEditingTool, public IClickDragBehaviorTarget
 {
 	GENERATED_BODY()
-	using FTransform3d = UE::Geometry::FTransform3d;
 	using FAxisAlignedBox3d = UE::Geometry::FAxisAlignedBox3d;
 public:
 	UAlignObjectsTool();
diff --git a/Engine/Plugins/Experimental/MeshModelingToolsetExp/Source/MeshModelingToolsExp/Public/CubeGridTool.h b/Engine/Plugins/Experimental/MeshModelingToolsetExp/Source/MeshModelingToolsExp/Public/CubeGridTool.h
--- a/Engine/Plugins/Experimental/MeshModelingToolsetExp/Source/MeshModelingToolsExp/Public/CubeGridTool.h
+++ b/Engine/Plugins/Experimental/MeshModelingToolsetExp/Source/MeshModelingToolsExp/Public/CubeGridTool.h
@@ -461,7 +461,7 @@ protected:
 
 	TSharedPtr<UE::Geometry::FDynamicMesh3, ESPMode::ThreadSafe> CurrentMesh;
 	TSharedPtr<UE::Geometry::FDynamicMeshAABBTree3, ESPMode::ThreadSafe> MeshSpatial;
-	UE::Geometry::FTransform3d CurrentMeshTransform = UE::Geometry::FTransform3d::Identity();
+	UE::Geometry::FTransformSRT3d CurrentMeshTransform = UE::Geometry::FTransformSRT3d::Identity();
 	TSharedPtr<TArray<int32>, ESPMode::ThreadSafe> LastOpChangedTids;
 
 	// Safe inputs for the background compute to use, untouched by undo/redo/other CurrentMesh updates.
diff --git a/Engine/Plugins/Experimental/MeshModelingToolsetExp/Source/MeshModelingToolsExp/Public/DeformMeshPolygonsTool.h b/Engine/Plugins/Experimental/MeshModelingToolsetExp/Source/MeshModelingToolsExp/Public/DeformMeshPolygonsTool.h
--- a/Engine/Plugins/Experimental/MeshModelingToolsetExp/Source/MeshModelingToolsExp/Public/DeformMeshPolygonsTool.h
+++ b/Engine/Plugins/Experimental/MeshModelingToolsetExp/Source/MeshModelingToolsExp/Public/DeformMeshPolygonsTool.h
@@ -200,7 +200,7 @@ protected:
 	
 	// camera state at last render
 	FViewCameraState CameraState;
-	UE::Geometry::FTransform3d WorldTransform;
+	FTransform3d WorldTransform;
 
 	FToolDataVisualizer PolyEdgesRenderer;
 
diff --git a/Engine/Plugins/Experimental/MeshModelingToolsetExp/Source/MeshModelingToolsExp/Public/DynamicMeshSculptTool.h b/Engine/Plugins/Experimental/MeshModelingToolsetExp/Source/MeshModelingToolsExp/Public/DynamicMeshSculptTool.h
--- a/Engine/Plugins/Experimental/MeshModelingToolsetExp/Source/MeshModelingToolsExp/Public/DynamicMeshSculptTool.h
+++ b/Engine/Plugins/Experimental/MeshModelingToolsetExp/Source/MeshModelingToolsExp/Public/DynamicMeshSculptTool.h
@@ -339,8 +339,8 @@ private:
 	UPROPERTY()
 	TObjectPtr<UMaterialInstanceDynamic> ActiveOverrideMaterial;
 
-	UE::Geometry::FTransform3d InitialTargetTransform;
-	UE::Geometry::FTransform3d CurTargetTransform;
+	UE::Geometry::FTransformSRT3d InitialTargetTransform;
+	UE::Geometry::FTransformSRT3d CurTargetTransform;
 
 	// realtime visualization
 	void OnDynamicMeshComponentChanged();
diff --git a/Engine/Plugins/Experimental/MeshModelingToolsetExp/Source/MeshModelingToolsExp/Public/EditPivotTool.h b/Engine/Plugins/Experimental/MeshModelingToolsetExp/Source/MeshModelingToolsExp/Public/EditPivotTool.h
--- a/Engine/Plugins/Experimental/MeshModelingToolsetExp/Source/MeshModelingToolsExp/Public/EditPivotTool.h
+++ b/Engine/Plugins/Experimental/MeshModelingToolsetExp/Source/MeshModelingToolsExp/Public/EditPivotTool.h
@@ -187,7 +187,7 @@ public:
 protected:
 	TArray<int> MapToFirstOccurrences;
 
-	UE::Geometry::FTransform3d Transform;
+	FTransform3d Transform;
 	UE::Geometry::FAxisAlignedBox3d ObjectBounds;
 	UE::Geometry::FAxisAlignedBox3d WorldBounds;
 	void Precompute();
diff --git a/Engine/Plugins/Experimental/MeshModelingToolsetExp/Source/MeshModelingToolsExp/Public/RemoveOccludedTrianglesTool.h b/Engine/Plugins/Experimental/MeshModelingToolsetExp/Source/MeshModelingToolsExp/Public/RemoveOccludedTrianglesTool.h
--- a/Engine/Plugins/Experimental/MeshModelingToolsetExp/Source/MeshModelingToolsExp/Public/RemoveOccludedTrianglesTool.h
+++ b/Engine/Plugins/Experimental/MeshModelingToolsetExp/Source/MeshModelingToolsExp/Public/RemoveOccludedTrianglesTool.h
@@ -205,7 +205,7 @@ protected:
 	// AABB trees and winding trees for every mesh target, with repeated instances as pointers to the same data
 	TArray<TSharedPtr<UE::Geometry::FDynamicMeshAABBTree3, ESPMode::ThreadSafe>> OccluderTrees;
 	TArray<TSharedPtr<UE::Geometry::TFastWindingTree<FDynamicMesh3>, ESPMode::ThreadSafe>> OccluderWindings;
-	TArray<UE::Geometry::FTransform3d> OccluderTransforms;
+	TArray<UE::Geometry::FTransformSRT3d> OccluderTransforms;
 
 	TArray<TArray<int32>> PreviewToCopyIdx;
 	TArray<int32> PreviewToTargetIdx;
diff --git a/Engine/Plugins/Experimental/MeshModelingToolsetExp/Source/MeshModelingToolsExp/Public/Sculpting/MeshSculptToolBase.h b/Engine/Plugins/Experimental/MeshModelingToolsetExp/Source/MeshModelingToolsExp/Public/Sculpting/MeshSculptToolBase.h
--- a/Engine/Plugins/Experimental/MeshModelingToolsetExp/Source/MeshModelingToolsExp/Public/Sculpting/MeshSculptToolBase.h
+++ b/Engine/Plugins/Experimental/MeshModelingToolsetExp/Source/MeshModelingToolsExp/Public/Sculpting/MeshSculptToolBase.h
@@ -224,7 +224,6 @@ class MESHMODELINGTOOLSEXP_API UMeshSculptToolBase : public UMeshSurfacePointToo
 {
 	GENERATED_BODY()
 protected:
-	using FTransform3d = UE::Geometry::FTransform3d;
 	using FFrame3d = UE::Geometry::FFrame3d;
 public:
 
@@ -266,9 +265,9 @@ protected:
 	FViewCameraState CameraState;
 
 	/** Initial transformation on target mesh */
-	UE::Geometry::FTransform3d InitialTargetTransform;
+	UE::Geometry::FTransformSRT3d InitialTargetTransform;
 	/** Active transformation on target mesh, includes baked scale */
-	UE::Geometry::FTransform3d CurTargetTransform;
+	UE::Geometry::FTransformSRT3d CurTargetTransform;
 
 	FRay3d GetLocalRay(const FRay& WorldRay) const;
 
diff --git a/Engine/Plugins/Experimental/MeshModelingToolsetExp/Source/MeshModelingToolsExp/Public/SeamSculptTool.h b/Engine/Plugins/Experimental/MeshModelingToolsetExp/Source/MeshModelingToolsExp/Public/SeamSculptTool.h
--- a/Engine/Plugins/Experimental/MeshModelingToolsetExp/Source/MeshModelingToolsExp/Public/SeamSculptTool.h
+++ b/Engine/Plugins/Experimental/MeshModelingToolsetExp/Source/MeshModelingToolsExp/Public/SeamSculptTool.h
@@ -86,7 +86,7 @@ protected:
 	TObjectPtr<UPreviewGeometry> PreviewGeom;
 
 	TSharedPtr<FDynamicMesh3, ESPMode::ThreadSafe> InputMesh;
-	UE::Geometry::FTransform3d MeshTransform;
+	FTransform3d MeshTransform;
 	double NormalOffset = 0;
 
 	void InitPreviewGeometry();
diff --git a/Engine/Plugins/Experimental/PlanarCutPlugin/Source/PlanarCut/Private/GeometryMeshConversion.cpp b/Engine/Plugins/Experimental/PlanarCutPlugin/Source/PlanarCut/Private/GeometryMeshConversion.cpp
--- a/Engine/Plugins/Experimental/PlanarCutPlugin/Source/PlanarCut/Private/GeometryMeshConversion.cpp
+++ b/Engine/Plugins/Experimental/PlanarCutPlugin/Source/PlanarCut/Private/GeometryMeshConversion.cpp
@@ -1184,7 +1184,7 @@ FCellMeshes::FCellMeshes(int32 NumUVLayersIn, FDynamicMesh3& SingleCutter, const
 
 	if (Transform.IsSet())
 	{
-		MeshTransforms::ApplyTransform(SingleCutter, UE::Geometry::FTransform3d(Transform.GetValue()));
+		MeshTransforms::ApplyTransform(SingleCutter, FTransformSRT3d(Transform.GetValue()));
 	}
 
 	// Mesh should already be augmented
@@ -2126,8 +2126,7 @@ void FDynamicMeshCollection::Init(const FGeometryCollection* Collection, const T
 			continue;
 		}
 
-		using FTransform3d = UE::Geometry::FTransform3d;
-		FTransform3d CollectionToLocal = FTransform3d(GeometryCollectionAlgo::GlobalMatrix(Transforms, Collection->Parent, TransformIdx) * TransformCollection);
+		FTransformSRT3d CollectionToLocal = FTransformSRT3d(GeometryCollectionAlgo::GlobalMatrix(Transforms, Collection->Parent, TransformIdx) * TransformCollection);
 
 		int32 AddedMeshIdx = Meshes.Add(new FMeshData(NumUVLayers));
 		FMeshData& MeshData = Meshes[AddedMeshIdx];
diff --git a/Engine/Plugins/Experimental/PlanarCutPlugin/Source/PlanarCut/Private/PlanarCut.cpp b/Engine/Plugins/Experimental/PlanarCutPlugin/Source/PlanarCut/Private/PlanarCut.cpp
--- a/Engine/Plugins/Experimental/PlanarCutPlugin/Source/PlanarCut/Private/PlanarCut.cpp
+++ b/Engine/Plugins/Experimental/PlanarCutPlugin/Source/PlanarCut/Private/PlanarCut.cpp
@@ -908,8 +908,6 @@ int32 MergeBones(
 	UE::PlanarCut::ENeighborSelectionMethod NeighborSelectionMethod
 )
 {
-	using UE::Geometry::FTransform3d;
-
 	FTransform CellsToWorld = FTransform::Identity;
 
 	FGeometryCollectionProximityUtility ProximityUtility(&Collection);
@@ -1188,7 +1186,7 @@ int32 MergeBones(
 	for (int32 UpMeshIdx = 0; UpMeshIdx < UpdateCollection.Meshes.Num(); UpMeshIdx++)
 	{
 		FMeshData& UpMeshData = UpdateCollection.Meshes[UpMeshIdx];
-		FTransform3d UpMeshXF = (FTransform3d)UpMeshData.FromCollection;
+		FTransformSRT3d UpMeshXF = (FTransformSRT3d)UpMeshData.FromCollection;
 		int32 UpGeoIdx = Collection.TransformToGeometryIndex[UpMeshData.TransformIndex];
 		FRemoveGroup& Group = RemoveGroups[GeomIdxToRemoveGroupIdx[UpGeoIdx]];
 		if (!ensure(Group.IsValid()))
@@ -1199,7 +1197,7 @@ int32 MergeBones(
 		for (int32 RmGeoIdx : Group.ToRemove)
 		{
 			FMeshData& RmMeshData = RemoveCollection.Meshes[GeoIdxToRmMeshIdx[RmGeoIdx]];
-			FTransform3d RmMeshXF = (FTransform3d)RmMeshData.FromCollection;
+			FTransformSRT3d RmMeshXF = (FTransformSRT3d)RmMeshData.FromCollection;
 			if (bUnionJoinedPieces)
 			{
 				FMeshBoolean Boolean(&UpMeshData.AugMesh, &RmMeshData.AugMesh, &UpMeshData.AugMesh, FMeshBoolean::EBooleanOp::Union);
diff --git a/Engine/Plugins/Experimental/UVEditor/Source/UVEditorTools/Private/UVEditorRecomputeUVsTool.cpp b/Engine/Plugins/Experimental/UVEditor/Source/UVEditorTools/Private/UVEditorRecomputeUVsTool.cpp
--- a/Engine/Plugins/Experimental/UVEditor/Source/UVEditorTools/Private/UVEditorRecomputeUVsTool.cpp
+++ b/Engine/Plugins/Experimental/UVEditor/Source/UVEditorTools/Private/UVEditorRecomputeUVsTool.cpp
@@ -74,7 +74,7 @@ void UUVEditorRecomputeUVsTool::Setup()
 	{
 		TObjectPtr<UUVEditorToolMeshInput> Target = Targets[TargetIndex];
 		Factories[TargetIndex] = NewObject<URecomputeUVsOpFactory>();
-		Factories[TargetIndex]->TargetTransform = (UE::Geometry::FTransform3d)Target->AppliedPreview->PreviewMesh->GetTransform();
+		Factories[TargetIndex]->TargetTransform = (FTransform3d)Target->AppliedPreview->PreviewMesh->GetTransform();
 		Factories[TargetIndex]->Settings = Settings;
 		Factories[TargetIndex]->OriginalMesh = Target->AppliedCanonical;
 		Factories[TargetIndex]->InputGroups = ActiveGroupSet;
diff --git a/Engine/Plugins/Experimental/UVEditor/Source/UVEditorTools/Private/UVSelectTool.cpp b/Engine/Plugins/Experimental/UVEditor/Source/UVEditorTools/Private/UVSelectTool.cpp
--- a/Engine/Plugins/Experimental/UVEditor/Source/UVEditorTools/Private/UVSelectTool.cpp
+++ b/Engine/Plugins/Experimental/UVEditor/Source/UVEditorTools/Private/UVSelectTool.cpp
@@ -888,7 +888,7 @@ void UUVSelectTool::ApplyGizmoTransform()
 {
 	if (bGizmoTransformNeedsApplication)
 	{
-		UE::Geometry::FTransform3d TransformToApply(UnappliedGizmoTransform);
+		FTransformSRT3d TransformToApply(UnappliedGizmoTransform);
 
 		// TODO: The division here is a bit of a hack. Properly-speaking, the scaling handles should act relative to
 		// gizmo size, not the visible space across which we drag, otherwise it becomes dependent on the units we
diff --git a/Engine/Plugins/Experimental/Water/Source/Runtime/Private/WaterBooleanUtils.cpp b/Engine/Plugins/Experimental/Water/Source/Runtime/Private/WaterBooleanUtils.cpp
--- a/Engine/Plugins/Experimental/Water/Source/Runtime/Private/WaterBooleanUtils.cpp
+++ b/Engine/Plugins/Experimental/Water/Source/Runtime/Private/WaterBooleanUtils.cpp
@@ -131,7 +131,7 @@ void FWaterBooleanUtils::ExtractMesh(AVolume* Volume, FDynamicMesh3& Mesh)
 	Mesh.DiscardAttributes();
 
 	UModel* Model = Volume->Brush;
-	UE::Geometry::FTransform3d XForm(Volume->GetTransform());
+	FTransformSRT3d XForm(Volume->GetTransform());
 
 	// Each "BspNode" is a planar polygon, triangulate each polygon and accumulate in a mesh.
 	// Note that this does not make any attempt to weld vertices/edges
@@ -237,10 +237,10 @@ void FWaterBooleanUtils::ApplyBooleanRepairs(FDynamicMesh3& Mesh, double MergeTo
 
 void FWaterBooleanUtils::MakeNormalizationTransform(
 	const FAxisAlignedBox3d& Bounds,
-	UE::Geometry::FTransform3d& ToNormalizedOut, UE::Geometry::FTransform3d& FromNormalizedOut)
+	FTransformSRT3d& ToNormalizedOut, FTransformSRT3d& FromNormalizedOut)
 {
 	double WorldSize = Bounds.MaxDim();
-	FromNormalizedOut = UE::Geometry::FTransform3d::Identity();
+	FromNormalizedOut = FTransformSRT3d::Identity();
 	FromNormalizedOut.SetTranslation(Bounds.Center());
 	FromNormalizedOut.SetScale(WorldSize * FVector3d::One());
 	ToNormalizedOut = FromNormalizedOut.Inverse();
@@ -261,12 +261,12 @@ void FWaterBooleanUtils::SetToFaceNormals(FDynamicMesh3& Mesh)
 	}
 }
 
-FDynamicMesh3 FWaterBooleanUtils::AccumulateExtrusionVolumes(const TArray<AWaterBodyExclusionVolume*>& ExclusionVolumes, const UE::Geometry::FTransform3d& Transform)
+FDynamicMesh3 FWaterBooleanUtils::AccumulateExtrusionVolumes(const TArray<AWaterBodyExclusionVolume*>& ExclusionVolumes, const FTransformSRT3d& Transform)
 {
 	// First step is to boolean-union all the volumes into this mesh
 	FDynamicMesh3 CombinedVolumes;
 
-	UE::Geometry::FTransform3d IdentityTransform;
+	FTransformSRT3d IdentityTransform;
 	for (AWaterBodyExclusionVolume* Volume : ExclusionVolumes)
 	{
 		// convert volume to mesh
@@ -301,10 +301,8 @@ FDynamicMesh3 FWaterBooleanUtils::AccumulateExtrusionVolumes(const TArray<AWater
 
 FDynamicMesh3 FWaterBooleanUtils::MakeBoxCollisionMesh(FAxisAlignedBox3d WorldBoxBounds, const TArray<AWaterBodyExclusionVolume*>& ExclusionVolumes)
 {
-	using FTransform3d = UE::Geometry::FTransform3d;
-
 	// construct a transform that scales the input world-space box to unit bounds, to improve numerical precision
-	FTransform3d ToUnitTransform, FromUnitTransform;
+	FTransformSRT3d ToUnitTransform, FromUnitTransform;
 	MakeNormalizationTransform(WorldBoxBounds, ToUnitTransform, FromUnitTransform);
 
 	// First step is to boolean-union all the volumes into this mesh
@@ -1093,7 +1091,7 @@ void FWaterBooleanUtils::BuildOceanCollisionComponents(
 	TArray<TArray<AWaterBodyExclusionVolume*>> ClusteredVolumes;
 	ClusterExclusionVolumes(IntersectingVolumes, ClusteredVolumes, WorldMeshBufferWidth);
 
-	UE::Geometry::FTransform3d ActorTransformd(ActorTransform);
+	FTransform3d ActorTransformd(ActorTransform);
 
 	// for each cluster, clip ocean box to bounds-plus-buffer of cluster and then
 	// compute Boolean subtraction of exclusion volumes from that box
diff --git a/Engine/Plugins/Experimental/Water/Source/Runtime/Private/WaterBooleanUtils.h b/Engine/Plugins/Experimental/Water/Source/Runtime/Private/WaterBooleanUtils.h
--- a/Engine/Plugins/Experimental/Water/Source/Runtime/Private/WaterBooleanUtils.h
+++ b/Engine/Plugins/Experimental/Water/Source/Runtime/Private/WaterBooleanUtils.h
@@ -83,7 +83,7 @@ private:
 	 * Construct transforms to and from a normalized space, for a given bounds.
 	 * ("Normalized space" is such that scaled box fits in origin-centered unit box)
 	 */
-	static void MakeNormalizationTransform(const FAxisAlignedBox3d& Bounds, UE::Geometry::FTransform3d& ToNormalizedOut, UE::Geometry::FTransform3d& FromNormalizedOut);
+	static void MakeNormalizationTransform(const FAxisAlignedBox3d& Bounds, FTransformSRT3d& ToNormalizedOut, FTransformSRT3d& FromNormalizedOut);
 
 	/**
 	 * Utility function to set per-triangle attribute normals on a DynamicMesh
@@ -94,7 +94,7 @@ private:
 	 * Mesh a set of AVolumes and then Boolean-Union them into a FDynamicMesh3.
 	 * @param Transform this transform is applied to each volume before the Boolean is computed
 	 */
-	static FDynamicMesh3 AccumulateExtrusionVolumes(const TArray<AWaterBodyExclusionVolume*>& ExclusionVolumes, const UE::Geometry::FTransform3d& Transform);
+	static FDynamicMesh3 AccumulateExtrusionVolumes(const TArray<AWaterBodyExclusionVolume*>& ExclusionVolumes, const FTransformSRT3d& Transform);
 
 	/**
 	 * Generate a mesh that is a Box with a set of Volumes Boolean-subtracted.
diff --git a/Engine/Plugins/Runtime/GeometryProcessing/Source/DynamicMesh/Private/DynamicMeshEditor.cpp b/Engine/Plugins/Runtime/GeometryProcessing/Source/DynamicMesh/Private/DynamicMeshEditor.cpp
--- a/Engine/Plugins/Runtime/GeometryProcessing/Source/DynamicMesh/Private/DynamicMeshEditor.cpp
+++ b/Engine/Plugins/Runtime/GeometryProcessing/Source/DynamicMesh/Private/DynamicMeshEditor.cpp
@@ -1158,7 +1158,7 @@ void FDynamicMeshEditor::SetQuadUVsFromProjection(const FIndex2i& QuadTris, cons
 }
 
 
-void FDynamicMeshEditor::RescaleAttributeUVs(float UVScale, bool bWorldSpace, int UVLayerIndex, TOptional<UE::Geometry::FTransform3d> ToWorld)
+void FDynamicMeshEditor::RescaleAttributeUVs(float UVScale, bool bWorldSpace, int UVLayerIndex, TOptional<FTransformSRT3d> ToWorld)
 {
 	check(Mesh->HasAttributes() && Mesh->Attributes()->NumUVLayers() > UVLayerIndex );
 	FDynamicMeshUVOverlay* UVs = Mesh->Attributes()->GetUVLayer(UVLayerIndex);
diff --git a/Engine/Plugins/Runtime/GeometryProcessing/Source/DynamicMesh/Private/Operations/MeshBoolean.cpp b/Engine/Plugins/Runtime/GeometryProcessing/Source/DynamicMesh/Private/Operations/MeshBoolean.cpp
--- a/Engine/Plugins/Runtime/GeometryProcessing/Source/DynamicMesh/Private/Operations/MeshBoolean.cpp
+++ b/Engine/Plugins/Runtime/GeometryProcessing/Source/DynamicMesh/Private/Operations/MeshBoolean.cpp
@@ -221,7 +221,7 @@ bool FMeshBoolean::Compute()
 	double ScaleFactor = 1.0 / FMath::Clamp(CombinedAABB.MaxDim(), 0.01, 1000000.0);
 	for (int MeshIdx = 0; MeshIdx < 2; MeshIdx++)
 	{
-		FTransform3d CenteredTransform = Transforms[MeshIdx];
+		FTransformSRT3d CenteredTransform = Transforms[MeshIdx];
 		CenteredTransform.SetTranslation(ScaleFactor*(CenteredTransform.GetTranslation() - CombinedAABB.Center()));
 		CenteredTransform.SetScale(ScaleFactor*CenteredTransform.GetScale());
 		MeshTransforms::ApplyTransform(*CutMesh[MeshIdx], CenteredTransform);
@@ -230,7 +230,7 @@ bool FMeshBoolean::Compute()
 			CutMesh[MeshIdx]->ReverseOrientation(false);
 		}
 	}
-	ResultTransform = FTransform3d(CombinedAABB.Center());
+	ResultTransform = FTransformSRT3d(CombinedAABB.Center());
 	ResultTransform.SetScale(FVector3d::One() * (1.0 / ScaleFactor));
 
 	if (Cancelled())
@@ -757,7 +757,7 @@ bool FMeshBoolean::Compute()
 	if (bPutResultInInputSpace)
 	{
 		MeshTransforms::ApplyTransform(*Result, ResultTransform);
-		ResultTransform = FTransform3d::Identity();
+		ResultTransform = FTransformSRT3d::Identity();
 	}
 
 	return bSuccess;
diff --git a/Engine/Plugins/Runtime/GeometryProcessing/Source/DynamicMesh/Private/Operations/MeshSelfUnion.cpp b/Engine/Plugins/Runtime/GeometryProcessing/Source/DynamicMesh/Private/Operations/MeshSelfUnion.cpp
--- a/Engine/Plugins/Runtime/GeometryProcessing/Source/DynamicMesh/Private/Operations/MeshSelfUnion.cpp
+++ b/Engine/Plugins/Runtime/GeometryProcessing/Source/DynamicMesh/Private/Operations/MeshSelfUnion.cpp
@@ -24,11 +24,11 @@ bool FMeshSelfUnion::Compute()
 	// transform the mesh to a shared space (centered at the origin, and scaled to a unit cube)
 	FAxisAlignedBox3d AABB = Mesh->GetBounds(true);
 	double ScaleFactor = 1.0 / FMath::Clamp(AABB.MaxDim(), 0.01, 1000000.0);
-	FTransform3d TransformToCenteredBox = FTransform3d::Identity();
+	FTransformSRT3d TransformToCenteredBox = FTransformSRT3d::Identity();
 	TransformToCenteredBox.SetTranslation(ScaleFactor * (TransformToCenteredBox.GetTranslation() - AABB.Center()));
 	TransformToCenteredBox.SetScale(ScaleFactor * FVector3d::One());
 	MeshTransforms::ApplyTransform(*Mesh, TransformToCenteredBox);
-	FTransform3d ResultTransform(AABB.Center());
+	FTransformSRT3d ResultTransform(AABB.Center());
 	ResultTransform.SetScale((1.0 / ScaleFactor) * FVector3d::One());
 
 	// build spatial data and use it to find intersections
diff --git a/Engine/Plugins/Runtime/GeometryProcessing/Source/DynamicMesh/Private/Operations/RemoveOccludedTriangles.cpp b/Engine/Plugins/Runtime/GeometryProcessing/Source/DynamicMesh/Private/Operations/RemoveOccludedTriangles.cpp
--- a/Engine/Plugins/Runtime/GeometryProcessing/Source/DynamicMesh/Private/Operations/RemoveOccludedTriangles.cpp
+++ b/Engine/Plugins/Runtime/GeometryProcessing/Source/DynamicMesh/Private/Operations/RemoveOccludedTriangles.cpp
@@ -94,11 +94,11 @@ namespace UE
 					Spatials.Add(new TMeshAABBTree3<FComponentMesh>(&ComponentMeshes[ComponentIdx]));
 					Windings.Add(new TFastWindingTree<FComponentMesh>(&Spatials[ComponentIdx]));
 				}
-				FTransform3d Id = FTransform3d::Identity();
-				TArrayView<FTransform3d> TransformsView(&Id, 1);
+				FTransformSRT3d Id = FTransformSRT3d::Identity();
+				TArrayView<FTransformSRT3d> TransformsView(&Id, 1);
 				TArrayView<TMeshAABBTree3<FComponentMesh>*> SpatialsView(Spatials.GetData(), Spatials.Num());
 				TArrayView<TFastWindingTree<FComponentMesh>*> WindingsView(Windings.GetData(), Windings.Num());
-				return Remover.Apply(TransformsView, SpatialsView, WindingsView, TArrayView<const FTransform3d>(), bTestOccludedByAny);
+				return Remover.Apply(TransformsView, SpatialsView, WindingsView, TArrayView<const FTransformSRT3d>(), bTestOccludedByAny);
 			}
 			else
 			{
@@ -107,7 +107,7 @@ namespace UE
 				Remover.TriangleSamplingMethod = SamplingMethod;
 				Remover.AddTriangleSamples = RandomSamplesPerTri;
 				FDynamicMeshAABBTree3 Spatial(&Mesh, true);
-				return Remover.Apply(FTransform3d::Identity(), &Spatial);
+				return Remover.Apply(FTransformSRT3d::Identity(), &Spatial);
 			}
 		}
 	}
diff --git a/Engine/Plugins/Runtime/GeometryProcessing/Source/DynamicMesh/Private/ShapeApproximation/SimpleShapeSet3.cpp b/Engine/Plugins/Runtime/GeometryProcessing/Source/DynamicMesh/Private/ShapeApproximation/SimpleShapeSet3.cpp
--- a/Engine/Plugins/Runtime/GeometryProcessing/Source/DynamicMesh/Private/ShapeApproximation/SimpleShapeSet3.cpp
+++ b/Engine/Plugins/Runtime/GeometryProcessing/Source/DynamicMesh/Private/ShapeApproximation/SimpleShapeSet3.cpp
@@ -306,13 +306,13 @@ void FSimpleShapeSet3d::RemoveContainedGeometry()
 }
 
 
-static void TransformSphereShape(FSphereShape3d& SphereShape, const UE::Geometry::FTransform3d& Transform)
+static void TransformSphereShape(FSphereShape3d& SphereShape, const FTransform3d& Transform)
 {
-	double RadiusScale = Transform.GetScale().Length() / FVector3d::One().Length();
+	double RadiusScale = Transform.GetScale3D().Length() / FVector3d::One().Length();
 	SphereShape.Sphere.Center = Transform.TransformPosition(SphereShape.Sphere.Center);
 	SphereShape.Sphere.Radius *= RadiusScale;
 }
-static void TransformBoxShape(FBoxShape3d& BoxShape, const UE::Geometry::FTransform3d& Transform)
+static void TransformBoxShape(FBoxShape3d& BoxShape, const FTransform3d& Transform)
 {
 	FVector3d CornerVec = BoxShape.Box.Frame.PointAt(BoxShape.Box.Extents) - BoxShape.Box.Frame.Origin;
 	BoxShape.Box.Frame.Transform(Transform);
@@ -321,7 +321,7 @@ static void TransformBoxShape(FBoxShape3d& BoxShape, const UE::Geometry::FTransf
 	BoxShape.Box.Extents.Y = CornerVec.Dot(BoxShape.Box.AxisY());
 	BoxShape.Box.Extents.Z = CornerVec.Dot(BoxShape.Box.AxisZ());
 }
-static void TransformCapsuleShape(UE::Geometry::FCapsuleShape3d& CapsuleShape, const UE::Geometry::FTransform3d& Transform)
+static void TransformCapsuleShape(UE::Geometry::FCapsuleShape3d& CapsuleShape, const FTransform3d& Transform)
 {
 	FVector3d P0 = Transform.TransformPosition(CapsuleShape.Capsule.Segment.StartPoint());
 	FVector3d P1 = Transform.TransformPosition(CapsuleShape.Capsule.Segment.EndPoint());
@@ -341,19 +341,19 @@ static void TransformCapsuleShape(UE::Geometry::FCapsuleShape3d& CapsuleShape, c
 
 
 
-static void TransformSphereShape(FSphereShape3d& SphereShape, const TArray<UE::Geometry::FTransform3d>& TransformSequence)
+static void TransformSphereShape(FSphereShape3d& SphereShape, const TArray<FTransform3d>& TransformSequence)
 {
-	for (const UE::Geometry::FTransform3d& XForm : TransformSequence)
+	for (const FTransform3d& XForm : TransformSequence)
 	{
 		SphereShape.Sphere.Center = XForm.TransformPosition(SphereShape.Sphere.Center);
-		double RadiusScale = XForm.GetScale().Length() / FVector3d::One().Length();
+		double RadiusScale = XForm.GetScale3D().Length() / FVector3d::One().Length();
 		SphereShape.Sphere.Radius *= RadiusScale;
 	}
 }
-static void TransformBoxShape(FBoxShape3d& BoxShape, const TArray<UE::Geometry::FTransform3d>& TransformSequence)
+static void TransformBoxShape(FBoxShape3d& BoxShape, const TArray<FTransform3d>& TransformSequence)
 {
 	FVector3d CornerVec = BoxShape.Box.Frame.PointAt(BoxShape.Box.Extents) - BoxShape.Box.Frame.Origin;
-	for (const UE::Geometry::FTransform3d& XForm : TransformSequence)
+	for (const FTransform3d& XForm : TransformSequence)
 	{
 		BoxShape.Box.Frame.Transform(XForm);
 		CornerVec = XForm.TransformVector(CornerVec);
@@ -362,7 +362,7 @@ static void TransformBoxShape(FBoxShape3d& BoxShape, const TArray<UE::Geometry::
 	BoxShape.Box.Extents.Y = CornerVec.Dot(BoxShape.Box.AxisY());
 	BoxShape.Box.Extents.Z = CornerVec.Dot(BoxShape.Box.AxisZ());
 }
-static void TransformCapsuleShape(UE::Geometry::FCapsuleShape3d& CapsuleShape, const TArray<UE::Geometry::FTransform3d>& TransformSequence)
+static void TransformCapsuleShape(UE::Geometry::FCapsuleShape3d& CapsuleShape, const TArray<FTransform3d>& TransformSequence)
 {
 	FVector3d P0 = CapsuleShape.Capsule.Segment.StartPoint();
 	FVector3d P1 = CapsuleShape.Capsule.Segment.EndPoint();
@@ -372,7 +372,7 @@ static void TransformCapsuleShape(UE::Geometry::FCapsuleShape3d& CapsuleShape, c
 	FVector3d InitialSideVec = CapsuleFrame.PointAt(FVector3d(CurRadius, CurRadius, 0)) - CapsuleFrame.Origin;
 	FVector3d NewSideVec = InitialSideVec;
 
-	for (const UE::Geometry::FTransform3d& XForm : TransformSequence)
+	for (const FTransform3d& XForm : TransformSequence)
 	{
 		P0 = XForm.TransformPosition(P0);
 		P1 = XForm.TransformPosition(P1);
@@ -414,7 +414,7 @@ void FSimpleShapeSet3d::Append(const FSimpleShapeSet3d& OtherShapeSet)
 }
 
 
-void FSimpleShapeSet3d::Append(const FSimpleShapeSet3d& OtherShapeSet, const UE::Geometry::FTransform3d& Transform)
+void FSimpleShapeSet3d::Append(const FSimpleShapeSet3d& OtherShapeSet, const FTransform3d& Transform)
 {
 	for (FSphereShape3d SphereShape : OtherShapeSet.Spheres)
 	{
@@ -444,7 +444,7 @@ void FSimpleShapeSet3d::Append(const FSimpleShapeSet3d& OtherShapeSet, const UE:
 
 
 
-void FSimpleShapeSet3d::Append(const FSimpleShapeSet3d& OtherShapeSet, const TArray<UE::Geometry::FTransform3d>& TransformSequence)
+void FSimpleShapeSet3d::Append(const FSimpleShapeSet3d& OtherShapeSet, const TArray<FTransform3d>& TransformSequence)
 {
 	for (FSphereShape3d SphereShape : OtherShapeSet.Spheres)
 	{
@@ -517,7 +517,7 @@ void FSimpleShapeSet3d::FilterByVolume(int32 MaximumCount)
 
 
 
-void FSimpleShapeSet3d::ApplyTransform(const UE::Geometry::FTransform3d& Transform)
+void FSimpleShapeSet3d::ApplyTransform(const FTransform3d& Transform)
 {
 	for (FSphereShape3d& SphereShape : Spheres)
 	{
diff --git a/Engine/Plugins/Runtime/GeometryProcessing/Source/DynamicMesh/Public/DynamicMeshEditor.h b/Engine/Plugins/Runtime/GeometryProcessing/Source/DynamicMesh/Public/DynamicMeshEditor.h
--- a/Engine/Plugins/Runtime/GeometryProcessing/Source/DynamicMesh/Public/DynamicMeshEditor.h
+++ b/Engine/Plugins/Runtime/GeometryProcessing/Source/DynamicMesh/Public/DynamicMeshEditor.h
@@ -367,7 +367,7 @@ public:
 	 * @param UVLayerIndex which UV layer to operate on (must exist)
 	 * @param ToWorld Optionally transform vertices for world space scaling
 	 */
-	void RescaleAttributeUVs(float UVScale = 1.0f, bool bWorldSpace = false, int UVLayerIndex = 0, TOptional<FTransform3d> ToWorld = TOptional<FTransform3d>());
+	void RescaleAttributeUVs(float UVScale = 1.0f, bool bWorldSpace = false, int UVLayerIndex = 0, TOptional<FTransformSRT3d> ToWorld = TOptional<FTransformSRT3d>());
 
 
 	//////////////////////////////////////////////////////////////////////////
