rename UE::Geometry::TTransform3 to TTransformSRT3, update references

diff --git a/Engine/Plugins/Runtime/GeometryProcessing/Source/DynamicMesh/Public/Operations/MeshBoolean.h b/Engine/Plugins/Runtime/GeometryProcessing/Source/DynamicMesh/Public/Operations/MeshBoolean.h
--- a/Engine/Plugins/Runtime/GeometryProcessing/Source/DynamicMesh/Public/Operations/MeshBoolean.h
+++ b/Engine/Plugins/Runtime/GeometryProcessing/Source/DynamicMesh/Public/Operations/MeshBoolean.h
@@ -34,7 +34,7 @@ public:
 	// Inputs
 	//
 	const FDynamicMesh3* Meshes[2];
-	const FTransform3d Transforms[2];
+	const FTransformSRT3d Transforms[2];
 
 	enum class EBooleanOp
 	{
@@ -111,7 +111,7 @@ public:
 	//
 
 	/** Transform taking the result mesh back to the original space of the inputs */
-	FTransform3d ResultTransform;
+	FTransformSRT3d ResultTransform;
 
 	/** Boundary edges created by the mesh boolean algorithm failing to cleanly weld (doesn't include boundaries that already existed in source meshes) */
 	TArray<int> CreatedBoundaryEdges;
@@ -145,7 +145,7 @@ public:
 	 * @param OutputMesh Mesh to store output
 	 * @param Operation How to combine meshes
 	 */
-	FMeshBoolean(const FDynamicMesh3* MeshA, const FTransform3d& TransformA, const FDynamicMesh3* MeshB, const FTransform3d& TransformB,
+	FMeshBoolean(const FDynamicMesh3* MeshA, const FTransformSRT3d& TransformA, const FDynamicMesh3* MeshB, const FTransformSRT3d& TransformB,
 				 FDynamicMesh3* OutputMesh, EBooleanOp Operation)
 		: Meshes{ MeshA, MeshB }, Transforms{ TransformA, TransformB }, Operation(Operation), Result(OutputMesh)
 	{
@@ -153,7 +153,7 @@ public:
 	}
 
 	FMeshBoolean(const FDynamicMesh3* MeshA, const FDynamicMesh3* MeshB, FDynamicMesh3* OutputMesh, EBooleanOp Operation)
-		: Meshes{ MeshA, MeshB }, Transforms{ FTransform3d::Identity(), FTransform3d::Identity() }, Operation(Operation), Result(OutputMesh)
+		: Meshes{ MeshA, MeshB }, Transforms{ FTransformSRT3d::Identity(), FTransformSRT3d::Identity() }, Operation(Operation), Result(OutputMesh)
 	{
 		check(MeshA != nullptr && MeshB != nullptr && OutputMesh != nullptr);
 	}
diff --git a/Engine/Plugins/Runtime/GeometryProcessing/Source/DynamicMesh/Public/Operations/RemoveOccludedTriangles.h b/Engine/Plugins/Runtime/GeometryProcessing/Source/DynamicMesh/Public/Operations/RemoveOccludedTriangles.h
--- a/Engine/Plugins/Runtime/GeometryProcessing/Source/DynamicMesh/Public/Operations/RemoveOccludedTriangles.h
+++ b/Engine/Plugins/Runtime/GeometryProcessing/Source/DynamicMesh/Public/Operations/RemoveOccludedTriangles.h
@@ -106,9 +106,9 @@ public:
 	 *							 Otherwise, we test if it's occluded by the combination of all occluders.
 	 * @return true on success
 	 */
-	virtual bool Select(const TArrayView<const FTransform3d> MeshLocalToOccluderSpaces,
+	virtual bool Select(const TArrayView<const FTransformSRT3d> MeshLocalToOccluderSpaces,
 		const TArrayView<TMeshAABBTree3<OccluderTriangleMeshType>*> Spatials, const TArrayView<TFastWindingTree<OccluderTriangleMeshType>*> FastWindingTrees,
-		const TArrayView<const FTransform3d> SpatialTransforms = TArrayView<const FTransform3d>(), bool bTestOccludedByAny = false)
+		const TArrayView<const FTransformSRT3d> SpatialTransforms = TArrayView<const FTransformSRT3d>(), bool bTestOccludedByAny = false)
 	{
 		if (Cancelled())
 		{
@@ -149,14 +149,14 @@ public:
 		// Helper struct to track both the transform and its inverse
 		struct FTransformWithInv
 		{
-			FTransform3d Transform;
+			FTransformSRT3d Transform;
 			FQuaterniond InvRot;
 			FVector3d InvScale;
 
 			FTransformWithInv() : 
-				Transform(FTransform3d::Identity()), InvRot(FQuaterniond::Identity()), InvScale(FVector3d::One())
+				Transform(FTransformSRT3d::Identity()), InvRot(FQuaterniond::Identity()), InvScale(FVector3d::One())
 			{}
-			FTransformWithInv(const FTransform3d& Transform) : Transform(Transform)
+			FTransformWithInv(const FTransformSRT3d& Transform) : Transform(Transform)
 			{
 				ComputeInv();
 			}
@@ -191,7 +191,7 @@ public:
 			}
 
 			// checks if the transform is an exact match
-			inline bool IsSameTransform(const FTransform3d& Other) const
+			inline bool IsSameTransform(const FTransformSRT3d& Other) const
 			{
 				return Other.GetTranslation() == Transform.GetTranslation()
 					&& Other.GetScale() == Transform.GetScale()
@@ -200,7 +200,7 @@ public:
 
 			// apply the inverse transform unless the point was originally transformed by this same transform; then just return the original point
 			// use to avoid floating point error shoving a sample into the surface that it came from
-			inline FVector3d InverseTransformUnlessMatch(const FVector3d& Pt, const FVector3d& OriginalPt, const FTransform3d& OriginalXF) const
+			inline FVector3d InverseTransformUnlessMatch(const FVector3d& Pt, const FVector3d& OriginalPt, const FTransformSRT3d& OriginalXF) const
 			{
 				if (IsSameTransform(OriginalXF))
 				{
@@ -214,12 +214,12 @@ public:
 		};
 
 		typedef TFunctionRef<bool(const TArrayView<TMeshAABBTree3<OccluderTriangleMeshType>*> Spatials, const TArrayView<TFastWindingTree<OccluderTriangleMeshType>*> FastWindingTrees,
-			const FVector3d& Pt, const FVector3d& OriginalPt, const FTransform3d& OriginalXF, const TArray<FTransformWithInv>& SpatialTransforms)> FIsOccludedFn;
+			const FVector3d& Pt, const FVector3d& OriginalPt, const FTransformSRT3d& OriginalXF, const TArray<FTransformWithInv>& SpatialTransforms)> FIsOccludedFn;
 
 		// test if point is occluded by the combination of all spatials, according to winding number
 		FIsOccludedFn IsOccludedFWNTotal =
 			[this](const TArrayView<TMeshAABBTree3<OccluderTriangleMeshType>*> Spatials, const TArrayView<TFastWindingTree<OccluderTriangleMeshType>*> FastWindingTrees,
-			const FVector3d& Pt, const FVector3d& OriginalPt, const FTransform3d& OriginalXF, const TArray<FTransformWithInv>& SpatialTransforms) -> bool
+			const FVector3d& Pt, const FVector3d& OriginalPt, const FTransformSRT3d& OriginalXF, const TArray<FTransformWithInv>& SpatialTransforms) -> bool
 		{
 			double WindingSum = 0;
 			for (int Idx = 0, NumPts = FastWindingTrees.Num(); Idx < NumPts; ++Idx)
@@ -232,7 +232,7 @@ public:
 		// test if point is occluded by the any of the spatials, according to winding number
 		FIsOccludedFn IsOccludedFWNAny =
 			[this](const TArrayView<TMeshAABBTree3<OccluderTriangleMeshType>*> Spatials, const TArrayView<TFastWindingTree<OccluderTriangleMeshType>*> FastWindingTrees,
-				const FVector3d& Pt, const FVector3d& OriginalPt, const FTransform3d& OriginalXF, const TArray<FTransformWithInv>& SpatialTransforms) -> bool
+				const FVector3d& Pt, const FVector3d& OriginalPt, const FTransformSRT3d& OriginalXF, const TArray<FTransformWithInv>& SpatialTransforms) -> bool
 		{
 			for (int Idx = 0, NumPts = FastWindingTrees.Num(); Idx < NumPts; ++Idx)
 			{
@@ -247,7 +247,7 @@ public:
 		// test if point is occluded by the combination of all spatials, according to raycasts
 		FIsOccludedFn IsOccludedSimpleTotal = [this, &RayDirs, &NR]
 		(const TArrayView<TMeshAABBTree3<OccluderTriangleMeshType>*> Spatials, const TArrayView<TFastWindingTree<OccluderTriangleMeshType>*> FastWindingTrees,
-			const FVector3d& Pt, const FVector3d& OriginalPt, const FTransform3d& OriginalXF, const TArray<FTransformWithInv>& SpatialTransforms) -> bool
+			const FVector3d& Pt, const FVector3d& OriginalPt, const FTransformSRT3d& OriginalXF, const TArray<FTransformWithInv>& SpatialTransforms) -> bool
 		{
 			FRay3d Ray;
 			
@@ -276,7 +276,7 @@ public:
 		// test if point is occluded by the any of the spatials, according to raycasts
 		FIsOccludedFn IsOccludedSimpleAny = [this, &RayDirs, &NR]
 		(const TArrayView<TMeshAABBTree3<OccluderTriangleMeshType>*> Spatials, const TArrayView<TFastWindingTree<OccluderTriangleMeshType>*> FastWindingTrees,
-			const FVector3d& Pt, const FVector3d& OriginalPt, const FTransform3d& OriginalXF, const TArray<FTransformWithInv>& SpatialTransforms) -> bool
+			const FVector3d& Pt, const FVector3d& OriginalPt, const FTransformSRT3d& OriginalXF, const TArray<FTransformWithInv>& SpatialTransforms) -> bool
 		{
 			FRay3d Ray;
 
@@ -352,7 +352,7 @@ public:
 				bool bAllOccluded = true;
 				for (int32 TransformIdx = 0, TransformNum = MeshLocalToOccluderSpaces.Num(); TransformIdx < TransformNum; TransformIdx++)
 				{
-					const FTransform3d& OriginalTransform = MeshLocalToOccluderSpaces[TransformIdx];
+					const FTransformSRT3d& OriginalTransform = MeshLocalToOccluderSpaces[TransformIdx];
 					FVector3d XFPos = OriginalTransform.TransformPosition(SamplePos);
 
 					bAllOccluded = bAllOccluded && IsOccludedF(Spatials, FastWindingTrees, 
@@ -397,7 +397,7 @@ public:
 					FVector3d SamplePos = V0 * BaryCoords.X + V1 * BaryCoords.Y + V2 * BaryCoords.Z + Normal * NormalOffset;
 					for (int32 TransformIdx = 0, TransformNum = MeshLocalToOccluderSpaces.Num(); TransformIdx < TransformNum; TransformIdx++)
 					{
-						const FTransform3d& OriginalTransform = MeshLocalToOccluderSpaces[TransformIdx];
+						const FTransformSRT3d& OriginalTransform = MeshLocalToOccluderSpaces[TransformIdx];
 						FVector3d XFPos = OriginalTransform.TransformPosition(SamplePos);
 
 						bInside = bInside && IsOccludedF(Spatials, FastWindingTrees,
@@ -461,9 +461,9 @@ public:
 	 *							 Otherwise, we test if it's occluded by the combination of all occluders.
 	 * @return true on success
 	 */
-	virtual bool Apply(const TArrayView<const FTransform3d> MeshLocalToOccluderSpaces, 
+	virtual bool Apply(const TArrayView<const FTransformSRT3d> MeshLocalToOccluderSpaces, 
 		const TArrayView<TMeshAABBTree3<OccluderTriangleMeshType>*> Spatials, const TArrayView<TFastWindingTree<OccluderTriangleMeshType>*> FastWindingTrees,
-		const TArrayView<const FTransform3d> SpatialTransforms = TArrayView<const FTransform3d>(), bool bTestOccludedByAny = false)
+		const TArrayView<const FTransformSRT3d> SpatialTransforms = TArrayView<const FTransformSRT3d>(), bool bTestOccludedByAny = false)
 	{
 		if (!Select(MeshLocalToOccluderSpaces, Spatials, FastWindingTrees, SpatialTransforms, bTestOccludedByAny))
 		{
@@ -483,12 +483,12 @@ public:
 	 * @param FastWindingTrees Precomputed fast winding trees for occluders
 	 * @return true on success
 	 */
-	virtual bool Select(const TArrayView<const FTransform3d> MeshLocalToOccluderSpaces,
+	virtual bool Select(const TArrayView<const FTransformSRT3d> MeshLocalToOccluderSpaces,
 		TMeshAABBTree3<OccluderTriangleMeshType>* Spatial, TFastWindingTree<OccluderTriangleMeshType>* FastWindingTree)
 	{
 		TArrayView<TMeshAABBTree3<OccluderTriangleMeshType>*> Spatials(&Spatial, 1);
 		TArrayView<TFastWindingTree<OccluderTriangleMeshType>*> FastWindingTrees(&FastWindingTree, 1);
-		return Select(MeshLocalToOccluderSpaces, Spatials, FastWindingTrees, TArrayView<const FTransform3d>());
+		return Select(MeshLocalToOccluderSpaces, Spatials, FastWindingTrees, TArrayView<const FTransformSRT3d>());
 	}
 
 
@@ -501,12 +501,12 @@ public:
 	 * @param FastWindingTrees Precomputed fast winding trees for occluders
 	 * @return true on success
 	 */
-	virtual bool Apply(const TArrayView<const FTransform3d> MeshLocalToOccluderSpaces,
+	virtual bool Apply(const TArrayView<const FTransformSRT3d> MeshLocalToOccluderSpaces,
 		TMeshAABBTree3<OccluderTriangleMeshType>* Spatial, TFastWindingTree<OccluderTriangleMeshType>* FastWindingTree)
 	{
 		TArrayView<TMeshAABBTree3<OccluderTriangleMeshType>*> Spatials(&Spatial, 1);
 		TArrayView<TFastWindingTree<OccluderTriangleMeshType>*> FastWindingTrees(&FastWindingTree, 1);
-		return Apply(MeshLocalToOccluderSpaces, Spatials, FastWindingTrees, TArrayView<const FTransform3d>());
+		return Apply(MeshLocalToOccluderSpaces, Spatials, FastWindingTrees, TArrayView<const FTransformSRT3d>());
 	}
 
 	/**
@@ -516,9 +516,9 @@ public:
 	 * @param Occluder AABB tree of occluding geometry
 	 * @return true on success
 	 */
-	virtual bool Select(const FTransform3d& MeshLocalToOccluderSpace, TMeshAABBTree3<OccluderTriangleMeshType>* Spatial, TFastWindingTree<OccluderTriangleMeshType>* FastWindingTree)
+	virtual bool Select(const FTransformSRT3d& MeshLocalToOccluderSpace, TMeshAABBTree3<OccluderTriangleMeshType>* Spatial, TFastWindingTree<OccluderTriangleMeshType>* FastWindingTree)
 	{
-		TArrayView<const FTransform3d> MeshLocalToOccluderSpaces(&MeshLocalToOccluderSpace, 1); // array view of the single transform
+		TArrayView<const FTransformSRT3d> MeshLocalToOccluderSpaces(&MeshLocalToOccluderSpace, 1); // array view of the single transform
 		return Select(MeshLocalToOccluderSpaces, Spatial, FastWindingTree);
 	}
 
@@ -530,9 +530,9 @@ public:
 	 * @param Occluder AABB tree of occluding geometry
 	 * @return true on success
 	 */
-	virtual bool Apply(const FTransform3d& MeshLocalToOccluderSpace, TMeshAABBTree3<OccluderTriangleMeshType>* Spatial, TFastWindingTree<OccluderTriangleMeshType>* FastWindingTree)
+	virtual bool Apply(const FTransformSRT3d& MeshLocalToOccluderSpace, TMeshAABBTree3<OccluderTriangleMeshType>* Spatial, TFastWindingTree<OccluderTriangleMeshType>* FastWindingTree)
 	{
-		TArrayView<const FTransform3d> MeshLocalToOccluderSpaces(&MeshLocalToOccluderSpace, 1); // array view of the single transform
+		TArrayView<const FTransformSRT3d> MeshLocalToOccluderSpaces(&MeshLocalToOccluderSpace, 1); // array view of the single transform
 		return Apply(MeshLocalToOccluderSpaces, Spatial, FastWindingTree);
 	}
 
@@ -544,7 +544,7 @@ public:
 	 * @param Occluder AABB tree of occluding geometry
 	 * @return true on success
 	 */
-	virtual bool Select(const FTransform3d& MeshLocalToOccluderSpace, TMeshAABBTree3<OccluderTriangleMeshType>* Occluder)
+	virtual bool Select(const FTransformSRT3d& MeshLocalToOccluderSpace, TMeshAABBTree3<OccluderTriangleMeshType>* Occluder)
 	{
 		TFastWindingTree<OccluderTriangleMeshType> FastWindingTree(Occluder, InsideMode == EOcclusionCalculationMode::FastWindingNumber);
 		return Select(MeshLocalToOccluderSpace, Occluder, &FastWindingTree);
@@ -558,7 +558,7 @@ public:
 	 * @param Occluder AABB tree of occluding geometry
 	 * @return true on success
 	 */
-	virtual bool Apply(const FTransform3d& MeshLocalToOccluderSpace, TMeshAABBTree3<OccluderTriangleMeshType>* Occluder)
+	virtual bool Apply(const FTransformSRT3d& MeshLocalToOccluderSpace, TMeshAABBTree3<OccluderTriangleMeshType>* Occluder)
 	{
 		TFastWindingTree<OccluderTriangleMeshType> FastWindingTree(Occluder, InsideMode == EOcclusionCalculationMode::FastWindingNumber);
 		return Apply(MeshLocalToOccluderSpace, Occluder, &FastWindingTree);
diff --git a/Engine/Plugins/Runtime/GeometryProcessing/Source/DynamicMesh/Public/ProjectionTargets.h b/Engine/Plugins/Runtime/GeometryProcessing/Source/DynamicMesh/Public/ProjectionTargets.h
--- a/Engine/Plugins/Runtime/GeometryProcessing/Source/DynamicMesh/Public/ProjectionTargets.h
+++ b/Engine/Plugins/Runtime/GeometryProcessing/Source/DynamicMesh/Public/ProjectionTargets.h
@@ -58,8 +58,8 @@ public:
 
 	FWorldSpaceProjectionTarget(const FDynamicMesh3* MeshIn, 
 								FDynamicMeshAABBTree3* SpatialIn,
-								const FTransform3d& TargetMeshLocalToWorldIn,
-								const FTransform3d& ToolMeshLocalToWorldIn) :
+								const FTransformSRT3d& TargetMeshLocalToWorldIn,
+								const FTransformSRT3d& ToolMeshLocalToWorldIn) :
 		FMeshProjectionTarget(MeshIn, SpatialIn),
 		TargetMeshLocalToWorld(TargetMeshLocalToWorldIn),
 		ToolMeshLocalToWorld(ToolMeshLocalToWorldIn)
@@ -107,8 +107,8 @@ public:
 	}
 
 protected:
-	FTransform3d TargetMeshLocalToWorld;
-	FTransform3d ToolMeshLocalToWorld;
+	FTransformSRT3d TargetMeshLocalToWorld;
+	FTransformSRT3d ToolMeshLocalToWorld;
 
 };
 
diff --git a/Engine/Plugins/Runtime/MeshModelingToolset/Source/MeshModelingTools/Private/CombineMeshesTool.cpp b/Engine/Plugins/Runtime/MeshModelingToolset/Source/MeshModelingTools/Private/CombineMeshesTool.cpp
--- a/Engine/Plugins/Runtime/MeshModelingToolset/Source/MeshModelingTools/Private/CombineMeshesTool.cpp
+++ b/Engine/Plugins/Runtime/MeshModelingToolset/Source/MeshModelingTools/Private/CombineMeshesTool.cpp
@@ -204,7 +204,7 @@ void UCombineMeshesTool::CreateNewAsset()
 				AccumulateDMesh.Attributes()->SetNumUVLayers(ComponentDMesh.Attributes()->NumUVLayers());
 			}
 
-			UE::Geometry::FTransform3d XF = (UE::Geometry::FTransform3d)((FTransform)UE::ToolTarget::GetLocalToWorldTransform(Targets[ComponentIdx]) * ToAccum);
+			FTransformSRT3d XF = (UE::ToolTarget::GetLocalToWorldTransform(Targets[ComponentIdx]) * ToAccum);
 			if (XF.GetDeterminant() < 0)
 			{
 				ComponentDMesh.ReverseOrientation(false);
@@ -227,7 +227,7 @@ void UCombineMeshesTool::CreateNewAsset()
 				if (UE::Geometry::ComponentTypeSupportsCollision(PrimitiveComponent))
 				{
 					CollisionSettings = UE::Geometry::GetCollisionSettings(PrimitiveComponent);
-					UE::Geometry::AppendSimpleCollision(PrimitiveComponent, &SimpleCollision, UE::Geometry::FTransform3d::Identity());
+					UE::Geometry::AppendSimpleCollision(PrimitiveComponent, &SimpleCollision, FTransform3d::Identity);
 				}
 			}
 			else
@@ -339,8 +339,8 @@ void UCombineMeshesTool::UpdateExistingAsset()
 	UPrimitiveComponent* UpdateComponent = UE::ToolTarget::GetTargetComponent(Targets[SkipIndex]);
 	SkipActor = UE::ToolTarget::GetTargetActor(Targets[SkipIndex]);
 
-	UE::Geometry::FTransform3d TargetToWorld = (UE::Geometry::FTransform3d)UE::ToolTarget::GetLocalToWorldTransform(Targets[SkipIndex]);
-	UE::Geometry::FTransform3d WorldToTarget = TargetToWorld.Inverse();
+	FTransform3d TargetToWorld = UE::ToolTarget::GetLocalToWorldTransform(Targets[SkipIndex]);
+	FTransform3d WorldToTarget = TargetToWorld.Inverse();
 
 	FSimpleShapeSet3d SimpleCollision;
 	UE::Geometry::FComponentCollisionSettings CollisionSettings;
@@ -349,7 +349,7 @@ void UCombineMeshesTool::UpdateExistingAsset()
 	{
 		CollisionSettings = UE::Geometry::GetCollisionSettings(UpdateComponent);
 	}
-	TArray<UE::Geometry::FTransform3d> Transforms;
+	TArray<FTransform3d> Transforms;
 	Transforms.SetNum(2);
 
 	{
@@ -381,7 +381,7 @@ void UCombineMeshesTool::UpdateExistingAsset()
 
 			if (ComponentIdx != SkipIndex)
 			{
-				UE::Geometry::FTransform3d ComponentToWorld = (UE::Geometry::FTransform3d)UE::ToolTarget::GetLocalToWorldTransform(Targets[ComponentIdx]);
+				FTransform3d ComponentToWorld = (FTransform3d)UE::ToolTarget::GetLocalToWorldTransform(Targets[ComponentIdx]);
 				MeshTransforms::ApplyTransform(ComponentDMesh, ComponentToWorld);
 				if (ComponentToWorld.GetDeterminant() < 0)
 				{
@@ -403,7 +403,7 @@ void UCombineMeshesTool::UpdateExistingAsset()
 			{
 				if (bOutputComponentSupportsCollision && UE::Geometry::ComponentTypeSupportsCollision(PrimitiveComponent))
 				{
-					UE::Geometry::AppendSimpleCollision(PrimitiveComponent, &SimpleCollision, UE::Geometry::FTransform3d::Identity());
+					UE::Geometry::AppendSimpleCollision(PrimitiveComponent, &SimpleCollision, FTransform3d::Identity);
 				}
 			}
 
diff --git a/Engine/Plugins/Runtime/MeshModelingToolset/Source/MeshModelingTools/Private/ToolActivities/PolyEditBevelEdgeActivity.cpp b/Engine/Plugins/Runtime/MeshModelingToolset/Source/MeshModelingTools/Private/ToolActivities/PolyEditBevelEdgeActivity.cpp
--- a/Engine/Plugins/Runtime/MeshModelingToolset/Source/MeshModelingTools/Private/ToolActivities/PolyEditBevelEdgeActivity.cpp
+++ b/Engine/Plugins/Runtime/MeshModelingToolset/Source/MeshModelingTools/Private/ToolActivities/PolyEditBevelEdgeActivity.cpp
@@ -28,9 +28,9 @@ public:
 	TArray<int32> BevelGroupFaces;
 	double BevelDistance = 1.0;
 
-	void SetTransform(const FTransform& Transform)
+	void SetTransform(const FTransformSRT3d& Transform)
 	{
-		ResultTransform = (UE::Geometry::FTransform3d)Transform;
+		ResultTransform = Transform;
 	}
 
 	// FDynamicMeshOperator implementation 
@@ -103,7 +103,7 @@ TUniquePtr<FDynamicMeshOperator> UPolyEditBevelEdgeActivity::MakeNewOperator()
 		UE_LOG(LogGeometry, Warning, TEXT("UPolyEditBevelEdgeActivity::MakeNewOperator : empty selection"));
 	}
 
-	UE::Geometry::FTransform3d WorldTransform(ActivityContext->Preview->PreviewMesh->GetTransform());
+	FTransform3d WorldTransform(ActivityContext->Preview->PreviewMesh->GetTransform());
 	Op->SetResultTransform(WorldTransform);
 	
 	return Op;
diff --git a/Engine/Plugins/Runtime/MeshModelingToolset/Source/MeshModelingTools/Private/ToolActivities/PolyEditCutFacesActivity.cpp b/Engine/Plugins/Runtime/MeshModelingToolset/Source/MeshModelingTools/Private/ToolActivities/PolyEditCutFacesActivity.cpp
--- a/Engine/Plugins/Runtime/MeshModelingToolset/Source/MeshModelingTools/Private/ToolActivities/PolyEditCutFacesActivity.cpp
+++ b/Engine/Plugins/Runtime/MeshModelingToolset/Source/MeshModelingTools/Private/ToolActivities/PolyEditCutFacesActivity.cpp
@@ -100,7 +100,7 @@ void UPolyEditCutFacesActivity::BeginCutFaces()
 	TArray<int32> ActiveTriangleSelection;
 	ActivityContext->CurrentTopology->GetSelectedTriangles(ActiveSelection, ActiveTriangleSelection);
 
-	UE::Geometry::FTransform3d WorldTransform(ActivityContext->Preview->PreviewMesh->GetTransform());
+	FTransform3d WorldTransform(ActivityContext->Preview->PreviewMesh->GetTransform());
 
 	EditPreview = PolyEditActivityUtil::CreatePolyEditPreviewMesh(*ParentTool, *ActivityContext);
 	EditPreview->InitializeStaticType(ActivityContext->CurrentMesh.Get(), ActiveTriangleSelection, &WorldTransform);
@@ -147,7 +147,7 @@ void UPolyEditCutFacesActivity::ApplyCutFaces()
 	}
 	FVector3d PlaneOrigin = 0.5 * (Point0.Origin + Point1.Origin);
 	// map into local space of target mesh
-	UE::Geometry::FTransform3d WorldTransform(ActivityContext->Preview->PreviewMesh->GetTransform());
+	FTransformSRT3d WorldTransform(ActivityContext->Preview->PreviewMesh->GetTransform());
 	PlaneOrigin = WorldTransform.InverseTransformPosition(PlaneOrigin);
 	PlaneNormal = WorldTransform.InverseTransformNormal(PlaneNormal);
 	UE::Geometry::Normalize(PlaneNormal);
diff --git a/Engine/Plugins/Runtime/MeshModelingToolset/Source/MeshModelingTools/Private/ToolActivities/PolyEditExtrudeActivity.cpp b/Engine/Plugins/Runtime/MeshModelingToolset/Source/MeshModelingTools/Private/ToolActivities/PolyEditExtrudeActivity.cpp
--- a/Engine/Plugins/Runtime/MeshModelingToolset/Source/MeshModelingTools/Private/ToolActivities/PolyEditExtrudeActivity.cpp
+++ b/Engine/Plugins/Runtime/MeshModelingToolset/Source/MeshModelingTools/Private/ToolActivities/PolyEditExtrudeActivity.cpp
@@ -42,7 +42,6 @@ namespace PolyEditExtrudeActivityLocals
 TUniquePtr<FDynamicMeshOperator> UPolyEditExtrudeActivity::MakeNewOperator()
 {
 	using namespace PolyEditExtrudeActivityLocals;
-	using UE::Geometry::FTransform3d;
 
 	TUniquePtr<FExtrudeOp> Op = MakeUnique<FExtrudeOp>();
 	Op->OriginalMesh = ActivityContext->CurrentMesh;
@@ -256,7 +255,6 @@ EToolActivityStartResult UPolyEditExtrudeActivity::Start()
 void UPolyEditExtrudeActivity::BeginExtrude()
 {
 	using namespace PolyEditExtrudeActivityLocals;
-	using UE::Geometry::FTransform3d;
 
 	ActiveSelection = ActivityContext->SelectionMechanic->GetActiveSelection();
 
@@ -291,8 +289,6 @@ void UPolyEditExtrudeActivity::BeginExtrude()
 // The height mechanics has to get reinitialized whenever extrude direction changes
 void UPolyEditExtrudeActivity::ReinitializeExtrudeHeightMechanic()
 {
-	using UE::Geometry::FTransform3d;
-
 	ExtrusionFrameWorld = ActiveSelectionFrameWorld;
 	ExtrusionFrameWorld.AlignAxis(2, GetExtrudeDirection());
 
@@ -463,8 +459,6 @@ void UPolyEditExtrudeActivity::ApplyExtrude()
 // Gets used to set the direction along which we measure the distance to extrude.
 FVector3d UPolyEditExtrudeActivity::GetExtrudeDirection() const
 {
-	using UE::Geometry::FTransform3d;
-
 	EPolyEditExtrudeDirection DirectionToUse = EPolyEditExtrudeDirection::SelectionNormal;
 	switch (PropertySetToUse)
 	{
@@ -492,11 +486,11 @@ FVector3d UPolyEditExtrudeActivity::GetExtrudeDirection() const
 	case EPolyEditExtrudeDirection::WorldZ:
 		return FVector3d::UnitZ();
 	case EPolyEditExtrudeDirection::LocalX:
-		return FTransform3d(ActivityContext->Preview->PreviewMesh->GetTransform()).GetRotation().AxisX();
+		return FTransformSRT3d(ActivityContext->Preview->PreviewMesh->GetTransform()).GetRotation().AxisX();
 	case EPolyEditExtrudeDirection::LocalY:
-		return FTransform3d(ActivityContext->Preview->PreviewMesh->GetTransform()).GetRotation().AxisY();
+		return FTransformSRT3d(ActivityContext->Preview->PreviewMesh->GetTransform()).GetRotation().AxisY();
 	case EPolyEditExtrudeDirection::LocalZ:
-		return FTransform3d(ActivityContext->Preview->PreviewMesh->GetTransform()).GetRotation().AxisZ();
+		return FTransformSRT3d(ActivityContext->Preview->PreviewMesh->GetTransform()).GetRotation().AxisZ();
 	}
 	return ActiveSelectionFrameWorld.Z();
 }
diff --git a/Engine/Plugins/Runtime/MeshModelingToolset/Source/MeshModelingTools/Private/ToolActivities/PolyEditInsetOutsetActivity.cpp b/Engine/Plugins/Runtime/MeshModelingToolset/Source/MeshModelingTools/Private/ToolActivities/PolyEditInsetOutsetActivity.cpp
--- a/Engine/Plugins/Runtime/MeshModelingToolset/Source/MeshModelingTools/Private/ToolActivities/PolyEditInsetOutsetActivity.cpp
+++ b/Engine/Plugins/Runtime/MeshModelingToolset/Source/MeshModelingTools/Private/ToolActivities/PolyEditInsetOutsetActivity.cpp
@@ -118,7 +118,7 @@ bool UPolyEditInsetOutsetActivity::BeginInset()
 	TArray<int32> ActiveTriangleSelection;
 	ActivityContext->CurrentTopology->GetSelectedTriangles(ActiveSelection, ActiveTriangleSelection);
 
-	UE::Geometry::FTransform3d WorldTransform(ActivityContext->Preview->PreviewMesh->GetTransform());
+	FTransform3d WorldTransform(ActivityContext->Preview->PreviewMesh->GetTransform());
 
 	EditPreview = PolyEditActivityUtil::CreatePolyEditPreviewMesh(*ParentTool, *ActivityContext);
 	EditPreview->InitializeInsetType(ActivityContext->CurrentMesh.Get(), ActiveTriangleSelection, &WorldTransform);
@@ -148,7 +148,7 @@ bool UPolyEditInsetOutsetActivity::BeginInset()
 	}
 	TArray<FVector3d> LoopVertices;
 	Loops.Loops[0].GetVertices(LoopVertices);
-	CurveDistMechanic->InitializePolyLoop(LoopVertices, UE::Geometry::FTransform3d::Identity());
+	CurveDistMechanic->InitializePolyLoop(LoopVertices, FTransformSRT3d::Identity());
 
 	SetToolPropertySourceEnabled(Settings, true);
 
diff --git a/Engine/Plugins/Runtime/MeshModelingToolset/Source/MeshModelingTools/Private/ToolActivities/PolyEditPlanarProjectionUVActivity.cpp b/Engine/Plugins/Runtime/MeshModelingToolset/Source/MeshModelingTools/Private/ToolActivities/PolyEditPlanarProjectionUVActivity.cpp
--- a/Engine/Plugins/Runtime/MeshModelingToolset/Source/MeshModelingTools/Private/ToolActivities/PolyEditPlanarProjectionUVActivity.cpp
+++ b/Engine/Plugins/Runtime/MeshModelingToolset/Source/MeshModelingTools/Private/ToolActivities/PolyEditPlanarProjectionUVActivity.cpp
@@ -101,7 +101,7 @@ void UPolyEditPlanarProjectionUVActivity::BeginSetUVs()
 	const FGroupTopologySelection& ActiveSelection = ActivityContext->SelectionMechanic->GetActiveSelection();
 	ActivityContext->CurrentTopology->GetSelectedTriangles(ActiveSelection, ActiveTriangleSelection);
 
-	UE::Geometry::FTransform3d WorldTransform(ActivityContext->Preview->PreviewMesh->GetTransform());
+	FTransform3d WorldTransform(ActivityContext->Preview->PreviewMesh->GetTransform());
 
 	EditPreview = PolyEditActivityUtil::CreatePolyEditPreviewMesh(*ParentTool, *ActivityContext);
 	EditPreview->InitializeStaticType(ActivityContext->CurrentMesh.Get(), ActiveTriangleSelection, &WorldTransform);
@@ -181,9 +181,9 @@ void UPolyEditPlanarProjectionUVActivity::ApplySetUVs()
 	PlanarFrame.ConstrainedAlignAxis(0, Delta, PlanarFrame.Z());
 
 	// transform to local, use 3D point to transfer UV scale value
-	UE::Geometry::FTransform3d WorldTransform(ActivityContext->Preview->PreviewMesh->GetTransform());
+	FTransform3d WorldTransform(ActivityContext->Preview->PreviewMesh->GetTransform());
 	FVector3d ScalePt = PlanarFrame.Origin + UVScale * PlanarFrame.Z();
-	UE::Geometry::FTransform3d ToLocalXForm(WorldTransform.Inverse());
+	FTransform3d ToLocalXForm(WorldTransform.Inverse());
 	PlanarFrame.Transform(ToLocalXForm);
 	ScalePt = ToLocalXForm.TransformPosition(ScalePt);
 	UVScale = Distance(ScalePt, PlanarFrame.Origin);
diff --git a/Engine/Plugins/Runtime/MeshModelingToolset/Source/MeshModelingTools/Private/UVProjectionTool.cpp b/Engine/Plugins/Runtime/MeshModelingToolset/Source/MeshModelingTools/Private/UVProjectionTool.cpp
--- a/Engine/Plugins/Runtime/MeshModelingToolset/Source/MeshModelingTools/Private/UVProjectionTool.cpp
+++ b/Engine/Plugins/Runtime/MeshModelingToolset/Source/MeshModelingTools/Private/UVProjectionTool.cpp
@@ -102,7 +102,7 @@ void UUVProjectionTool::Setup()
 	ProjectionShapeVisualizer.LineThickness = 2.0;
 
 	// initialize Gizmo
-	UE::Geometry::FTransform3d GizmoPositionWorld(WorldBounds.Center());
+	FTransform3d GizmoPositionWorld(WorldBounds.Center());
 	TransformProxy = NewObject<UTransformProxy>(this);
 	TransformProxy->SetTransform((FTransform)GizmoPositionWorld);
 	TransformProxy->OnTransformChanged.AddUObject(this, &UUVProjectionTool::TransformChanged);
diff --git a/Engine/Plugins/Runtime/MeshModelingToolset/Source/MeshModelingTools/Public/EditMeshPolygonsTool.h b/Engine/Plugins/Runtime/MeshModelingToolset/Source/MeshModelingTools/Public/EditMeshPolygonsTool.h
--- a/Engine/Plugins/Runtime/MeshModelingToolset/Source/MeshModelingTools/Public/EditMeshPolygonsTool.h
+++ b/Engine/Plugins/Runtime/MeshModelingToolset/Source/MeshModelingTools/Public/EditMeshPolygonsTool.h
@@ -565,7 +565,7 @@ protected:
 	FFrame3d LockedTransfomerFrame;
 	bool bInGizmoDrag = false;
 
-	UE::Geometry::FTransform3d WorldTransform;
+	UE::Geometry::FTransformSRT3d WorldTransform;
 
 	FFrame3d InitialGizmoFrame;
 	FVector3d InitialGizmoScale;
diff --git a/Engine/Plugins/Runtime/MeshModelingToolset/Source/ModelingComponents/Private/BaseTools/BaseCreateFromSelectedTool.cpp b/Engine/Plugins/Runtime/MeshModelingToolset/Source/ModelingComponents/Private/BaseTools/BaseCreateFromSelectedTool.cpp
--- a/Engine/Plugins/Runtime/MeshModelingToolset/Source/ModelingComponents/Private/BaseTools/BaseCreateFromSelectedTool.cpp
+++ b/Engine/Plugins/Runtime/MeshModelingToolset/Source/ModelingComponents/Private/BaseTools/BaseCreateFromSelectedTool.cpp
@@ -196,10 +196,10 @@ void UBaseCreateFromSelectedTool::GenerateAsset(const FDynamicMeshOpResult& OpRe
 		FVector3d Center = OpResult.Mesh->GetBounds().Center();
 		double Rescale = OpResult.Transform.GetScale().X;
 		FTransform3d LocalTransform(-Center * Rescale);
-		LocalTransform.SetScale(FVector3d(Rescale, Rescale, Rescale));
+		LocalTransform.SetScale3D(FVector3d(Rescale, Rescale, Rescale));
 		MeshTransforms::ApplyTransform(*OpResult.Mesh, LocalTransform);
 		NewTransform = OpResult.Transform;
-		NewTransform.SetScale(FVector3d::One());
+		NewTransform.SetScale3D(FVector3d::One());
 		NewTransform.SetTranslation(NewTransform.GetTranslation() + NewTransform.TransformVector(Center * Rescale));
 	}
 
diff --git a/Engine/Plugins/Runtime/MeshModelingToolset/Source/ModelingComponents/Private/Components/OctreeDynamicMeshComponent.cpp b/Engine/Plugins/Runtime/MeshModelingToolset/Source/ModelingComponents/Private/Components/OctreeDynamicMeshComponent.cpp
--- a/Engine/Plugins/Runtime/MeshModelingToolset/Source/ModelingComponents/Private/Components/OctreeDynamicMeshComponent.cpp
+++ b/Engine/Plugins/Runtime/MeshModelingToolset/Source/ModelingComponents/Private/Components/OctreeDynamicMeshComponent.cpp
@@ -75,7 +75,7 @@ void UOctreeDynamicMeshComponent::SetMesh(UE::Geometry::FDynamicMesh3&& MoveMesh
 
 
 
-void UOctreeDynamicMeshComponent::ApplyTransform(const UE::Geometry::FTransform3d& Transform, bool bInvert)
+void UOctreeDynamicMeshComponent::ApplyTransform(const FTransform3d& Transform, bool bInvert)
 {
 	if (bInvert)
 	{
diff --git a/Engine/Plugins/Runtime/MeshModelingToolset/Source/ModelingComponents/Private/ConversionUtils/VolumeToDynamicMesh.cpp b/Engine/Plugins/Runtime/MeshModelingToolset/Source/ModelingComponents/Private/ConversionUtils/VolumeToDynamicMesh.cpp
--- a/Engine/Plugins/Runtime/MeshModelingToolset/Source/ModelingComponents/Private/ConversionUtils/VolumeToDynamicMesh.cpp
+++ b/Engine/Plugins/Runtime/MeshModelingToolset/Source/ModelingComponents/Private/ConversionUtils/VolumeToDynamicMesh.cpp
@@ -38,7 +38,7 @@ void VolumeToDynamicMesh(AVolume* Volume, FDynamicMesh3& Mesh,
 
 	if (Options.bInWorldSpace)
 	{
-		UE::Geometry::FTransform3d XForm(Volume->GetTransform());
+		FTransform3d XForm(Volume->GetTransform());
 		MeshTransforms::ApplyTransform(Mesh, XForm);
 	}
 }
@@ -56,7 +56,7 @@ void BrushComponentToDynamicMesh(UBrushComponent* Component, UE::Geometry::FDyna
 
 	if (Options.bInWorldSpace)
 	{
-		UE::Geometry::FTransform3d XForm(Component->GetComponentTransform());
+		FTransform3d XForm(Component->GetComponentTransform());
 		MeshTransforms::ApplyTransform(Mesh, XForm);
 	}
 }
diff --git a/Engine/Plugins/Runtime/MeshModelingToolset/Source/ModelingComponents/Private/Drawing/UVLayoutPreview.cpp b/Engine/Plugins/Runtime/MeshModelingToolset/Source/ModelingComponents/Private/Drawing/UVLayoutPreview.cpp
--- a/Engine/Plugins/Runtime/MeshModelingToolset/Source/ModelingComponents/Private/Drawing/UVLayoutPreview.cpp
+++ b/Engine/Plugins/Runtime/MeshModelingToolset/Source/ModelingComponents/Private/Drawing/UVLayoutPreview.cpp
@@ -214,7 +214,7 @@ void UUVLayoutPreview::RecalculatePosition()
 
 	CurrentWorldFrame = ObjFrame;
 
-	UE::Geometry::FTransform3d Transform(ObjFrame.Rotation, ObjFrame.Origin);
+	FTransformSRT3d Transform(ObjFrame.Rotation, ObjFrame.Origin);
 	Transform.SetScale(UseScale * FVector3d::One());
 
 	SetTransform((FTransform)Transform);
diff --git a/Engine/Plugins/Runtime/MeshModelingToolset/Source/ModelingComponents/Private/Mechanics/DragAlignmentMechanic.cpp b/Engine/Plugins/Runtime/MeshModelingToolset/Source/ModelingComponents/Private/Mechanics/DragAlignmentMechanic.cpp
--- a/Engine/Plugins/Runtime/MeshModelingToolset/Source/ModelingComponents/Private/Mechanics/DragAlignmentMechanic.cpp
+++ b/Engine/Plugins/Runtime/MeshModelingToolset/Source/ModelingComponents/Private/Mechanics/DragAlignmentMechanic.cpp
@@ -175,7 +175,7 @@ void UDragAlignmentMechanic::OnGizmoTransformChanged(FTransform NewTransform)
 }
 
 void UDragAlignmentMechanic::InitializeDeformedMeshRayCast(TFunction<FDynamicMeshAABBTree3* ()> GetSpatialIn, 
-	const UE::Geometry::FTransform3d &TargetTransform, const FGroupTopologyDeformer* LinearDeformer)
+	const FTransform3d &TargetTransform, const FGroupTopologyDeformer* LinearDeformer)
 {
 	if (LinearDeformer)
 	{
diff --git a/Engine/Plugins/Runtime/MeshModelingToolset/Source/ModelingComponents/Private/Mechanics/LatticeControlPointsMechanic.cpp b/Engine/Plugins/Runtime/MeshModelingToolset/Source/ModelingComponents/Private/Mechanics/LatticeControlPointsMechanic.cpp
--- a/Engine/Plugins/Runtime/MeshModelingToolset/Source/ModelingComponents/Private/Mechanics/LatticeControlPointsMechanic.cpp
+++ b/Engine/Plugins/Runtime/MeshModelingToolset/Source/ModelingComponents/Private/Mechanics/LatticeControlPointsMechanic.cpp
@@ -317,7 +317,7 @@ void ULatticeControlPointsMechanic::GizmoTransformChanged(UTransformProxy* Proxy
 	FQuaterniond DeltaRotation = FQuaterniond(Transform.GetRotation() * GizmoStartRotation.Inverse());
 	FVector DeltaScale = Transform.GetScale3D() / GizmoStartScale;
 
-	FTransform3d DeltaTransform;
+	FTransformSRT3d DeltaTransform;
 	DeltaTransform.SetScale((FVector3d)DeltaScale);
 	DeltaTransform.SetRotation((FQuaterniond)DeltaRotation);
 	DeltaTransform.SetTranslation((FVector3d)Transform.GetTranslation());
diff --git a/Engine/Plugins/Runtime/MeshModelingToolset/Source/ModelingComponents/Private/MeshOpPreviewHelpers.cpp b/Engine/Plugins/Runtime/MeshModelingToolset/Source/ModelingComponents/Private/MeshOpPreviewHelpers.cpp
--- a/Engine/Plugins/Runtime/MeshModelingToolset/Source/ModelingComponents/Private/MeshOpPreviewHelpers.cpp
+++ b/Engine/Plugins/Runtime/MeshModelingToolset/Source/ModelingComponents/Private/MeshOpPreviewHelpers.cpp
@@ -42,7 +42,7 @@ FDynamicMeshOpResult UMeshOpPreviewWithBackgroundCompute::Shutdown()
 
 	FDynamicMeshOpResult Result;
 	Result.Mesh = PreviewMesh->ExtractPreviewMesh();
-	Result.Transform = UE::Geometry::FTransform3d(PreviewMesh->GetTransform());
+	Result.Transform = FTransformSRT3d(PreviewMesh->GetTransform());
 
 	PreviewMesh->SetVisible(false);
 	PreviewMesh->Disconnect();
diff --git a/Engine/Plugins/Runtime/MeshModelingToolset/Source/ModelingComponents/Private/ModelingToolTargetUtil.cpp b/Engine/Plugins/Runtime/MeshModelingToolset/Source/ModelingComponents/Private/ModelingToolTargetUtil.cpp
--- a/Engine/Plugins/Runtime/MeshModelingToolset/Source/ModelingComponents/Private/ModelingToolTargetUtil.cpp
+++ b/Engine/Plugins/Runtime/MeshModelingToolset/Source/ModelingComponents/Private/ModelingToolTargetUtil.cpp
@@ -106,15 +106,15 @@ bool UE::ToolTarget::SetSourceObjectVisible(UToolTarget* Target, bool bVisible)
 }
 
 
-UE::Geometry::FTransform3d UE::ToolTarget::GetLocalToWorldTransform(UToolTarget* Target)
+FTransform3d UE::ToolTarget::GetLocalToWorldTransform(UToolTarget* Target)
 {
 	IPrimitiveComponentBackedTarget* TargetComponent = Cast<IPrimitiveComponentBackedTarget>(Target);
 	if (TargetComponent)
 	{
-		return (UE::Geometry::FTransform3d)TargetComponent->GetWorldTransform();
+		return (FTransform3d)TargetComponent->GetWorldTransform();
 	}
 	ensure(false);
-	return UE::Geometry::FTransform3d();
+	return FTransform3d();
 }
 
 FComponentMaterialSet UE::ToolTarget::GetMaterialSet(UToolTarget* Target, bool bPreferAssetMaterials)
diff --git a/Engine/Plugins/Runtime/MeshModelingToolset/Source/ModelingComponents/Private/PreviewMesh.cpp b/Engine/Plugins/Runtime/MeshModelingToolset/Source/ModelingComponents/Private/PreviewMesh.cpp
--- a/Engine/Plugins/Runtime/MeshModelingToolset/Source/ModelingComponents/Private/PreviewMesh.cpp
+++ b/Engine/Plugins/Runtime/MeshModelingToolset/Source/ModelingComponents/Private/PreviewMesh.cpp
@@ -402,7 +402,7 @@ bool UPreviewMesh::FindRayIntersection(const FRay3d& WorldRay, FHitResult& HitOu
 {
 	if (IsVisible() && TemporaryParentActor != nullptr && bBuildSpatialDataStructure)
 	{
-		UE::Geometry::FTransform3d Transform(TemporaryParentActor->GetActorTransform());
+		FTransform3d Transform(TemporaryParentActor->GetActorTransform());
 		FRay3d LocalRay(Transform.InverseTransformPosition(WorldRay.Origin),
 			Transform.InverseTransformVector(WorldRay.Direction));
 		UE::Geometry::Normalize(LocalRay.Direction);
diff --git a/Engine/Plugins/Runtime/MeshModelingToolset/Source/ModelingComponents/Private/Scene/MeshSceneAdapter.cpp b/Engine/Plugins/Runtime/MeshModelingToolset/Source/ModelingComponents/Private/Scene/MeshSceneAdapter.cpp
--- a/Engine/Plugins/Runtime/MeshModelingToolset/Source/ModelingComponents/Private/Scene/MeshSceneAdapter.cpp
+++ b/Engine/Plugins/Runtime/MeshModelingToolset/Source/ModelingComponents/Private/Scene/MeshSceneAdapter.cpp
@@ -1087,9 +1087,8 @@ void ConstructUniqueScalesMapping(
 	TArray<FTransformSequence3d>& UniqueScaleTransformsOut,
 	double ScaleComponentTolerance = 0.01)
 {
-	using FTransform3d = UE::Geometry::FTransform3d;
 	// two transforms are "the same up to scaling" if this returns true
-	auto CompareScales = [ScaleComponentTolerance](const FTransform3d& T1, const FTransform3d& T2)
+	auto CompareScales = [ScaleComponentTolerance](const FTransformSRT3d& T1, const FTransformSRT3d& T2)
 	{
 		return (T1.GetScale() - T2.GetScale()).GetAbsMax() < ScaleComponentTolerance;
 	};
diff --git a/Engine/Plugins/Runtime/MeshModelingToolset/Source/ModelingComponents/Private/Scene/SceneGeometrySpatialCache.cpp b/Engine/Plugins/Runtime/MeshModelingToolset/Source/ModelingComponents/Private/Scene/SceneGeometrySpatialCache.cpp
--- a/Engine/Plugins/Runtime/MeshModelingToolset/Source/ModelingComponents/Private/Scene/SceneGeometrySpatialCache.cpp
+++ b/Engine/Plugins/Runtime/MeshModelingToolset/Source/ModelingComponents/Private/Scene/SceneGeometrySpatialCache.cpp
@@ -61,7 +61,7 @@ public:
 public:
 	TWeakObjectPtr<UBrushComponent> Component;
 	FAxisAlignedBox3d LocalBounds;
-	UE::Geometry::FTransform3d Transform;
+	FTransformSRT3d Transform;
 	FColliderMesh ColliderMesh;
 	bool bDeferredRebuildPending;
 };
@@ -159,7 +159,7 @@ void FSceneVolumeSpatial::OnGeometryModified(bool bDeferSpatialRebuild)
 		return;
 	}
 
-	Transform = (UE::Geometry::FTransform3d)Component->GetComponentTransform();
+	Transform = (FTransformSRT3d)Component->GetComponentTransform();
 	LocalBounds = FAxisAlignedBox3d( Component->CalcBounds(FTransform::Identity).GetBox() );
 
 	if (bDeferSpatialRebuild)
@@ -189,7 +189,7 @@ void FSceneVolumeSpatial::OnTransformModified()
 {
 	if (IsValid())
 	{
-		Transform = (UE::Geometry::FTransform3d)Component->GetComponentTransform();
+		Transform = (FTransformSRT3d)Component->GetComponentTransform();
 	}
 }
 
@@ -221,7 +221,7 @@ public:
 public:
 	TWeakObjectPtr<UDynamicMeshComponent> Component;
 	FAxisAlignedBox3d LocalBounds;
-	UE::Geometry::FTransform3d Transform;
+	FTransformSRT3d Transform;
 	FDynamicMeshAABBTree3 AABBTree;
 	bool bDeferredRebuildPending;
 };
@@ -329,7 +329,7 @@ void FSceneDynamicMeshSpatial::OnGeometryModified(bool bDeferSpatialRebuild)
 		return;
 	}
 
-	Transform = (UE::Geometry::FTransform3d)DynamicMeshComponent->GetComponentTransform();
+	Transform = (FTransformSRT3d)DynamicMeshComponent->GetComponentTransform();
 	// may be reasonable to use USceneComponent::GetLocalBounds() here. Unclear if it will have already been updated on mesh changes...
 	DynamicMeshComponent->ProcessMesh([&](const FDynamicMesh3& Mesh)
 	{
@@ -355,7 +355,7 @@ void FSceneDynamicMeshSpatial::OnTransformModified()
 {
 	if (IsValid())
 	{
-		Transform = (UE::Geometry::FTransform3d)Component->GetComponentTransform();
+		Transform = (FTransformSRT3d)Component->GetComponentTransform();
 	}
 }
 
diff --git a/Engine/Plugins/Runtime/MeshModelingToolset/Source/ModelingComponents/Private/Selection/PersistentMeshSelectionManager.cpp b/Engine/Plugins/Runtime/MeshModelingToolset/Source/ModelingComponents/Private/Selection/PersistentMeshSelectionManager.cpp
--- a/Engine/Plugins/Runtime/MeshModelingToolset/Source/ModelingComponents/Private/Selection/PersistentMeshSelectionManager.cpp
+++ b/Engine/Plugins/Runtime/MeshModelingToolset/Source/ModelingComponents/Private/Selection/PersistentMeshSelectionManager.cpp
@@ -98,7 +98,7 @@ void UPersistentMeshSelectionManager::OnSelectionModified()
 		//const float ROIBorderDepthBias = 0.1f * (float)(WorldBounds.DiagonalLength() * 0.01);
 		const float ROIBorderDepthBias = 0.01f;
 
-		UE::Geometry::FTransform3d Transform(ActiveSelection->GetTargetComponent()->GetComponentToWorld());
+		FTransform3d Transform(ActiveSelection->GetTargetComponent()->GetComponentToWorld());
 
 		const TArray<UE::Geometry::FSegment3d>& Lines = SelectionData->RenderEdges;
 		SelectionDisplay->CreateOrUpdateLineSet(TEXT("SelectionEdges"),  Lines.Num(),
diff --git a/Engine/Plugins/Runtime/MeshModelingToolset/Source/ModelingComponents/Private/Selection/PolygonSelectionMechanic.cpp b/Engine/Plugins/Runtime/MeshModelingToolset/Source/ModelingComponents/Private/Selection/PolygonSelectionMechanic.cpp
--- a/Engine/Plugins/Runtime/MeshModelingToolset/Source/ModelingComponents/Private/Selection/PolygonSelectionMechanic.cpp
+++ b/Engine/Plugins/Runtime/MeshModelingToolset/Source/ModelingComponents/Private/Selection/PolygonSelectionMechanic.cpp
@@ -104,14 +104,14 @@ void UPolygonSelectionMechanic::LoadSelection(const UPersistentMeshSelection& Se
 
 void UPolygonSelectionMechanic::Initialize(
 	const FDynamicMesh3* MeshIn,
-	FTransform TargetTransformIn,
+	FTransform3d TargetTransformIn,
 	UWorld* WorldIn,
 	const FGroupTopology* TopologyIn,
 	TFunction<FDynamicMeshAABBTree3 * ()> GetSpatialSourceFuncIn)
 {
 	this->Mesh = MeshIn;
 	this->Topology = TopologyIn;
-	this->TargetTransform = FTransform3d(TargetTransformIn);
+	this->TargetTransform = TargetTransformIn;
 
 	TopoSelector->Initialize(Mesh, Topology);
 	this->GetSpatialFunc = GetSpatialSourceFuncIn;
diff --git a/Engine/Plugins/Runtime/MeshModelingToolset/Source/ModelingComponents/Public/Components/OctreeDynamicMeshComponent.h b/Engine/Plugins/Runtime/MeshModelingToolset/Source/ModelingComponents/Public/Components/OctreeDynamicMeshComponent.h
--- a/Engine/Plugins/Runtime/MeshModelingToolset/Source/ModelingComponents/Public/Components/OctreeDynamicMeshComponent.h
+++ b/Engine/Plugins/Runtime/MeshModelingToolset/Source/ModelingComponents/Public/Components/OctreeDynamicMeshComponent.h
@@ -67,7 +67,7 @@ public:
 	 * Apply transform to internal mesh. Invalidates RenderProxy.
 	 * @param bInvert if true, inverse tranform is applied instead of forward transform
 	 */
-	virtual void ApplyTransform(const UE::Geometry::FTransform3d& Transform, bool bInvert) override;
+	virtual void ApplyTransform(const FTransform3d& Transform, bool bInvert) override;
 
 	//
 	// change tracking/etc
diff --git a/Engine/Plugins/Runtime/MeshModelingToolset/Source/ModelingComponents/Public/Drawing/PolyEditPreviewMesh.h b/Engine/Plugins/Runtime/MeshModelingToolset/Source/ModelingComponents/Public/Drawing/PolyEditPreviewMesh.h
--- a/Engine/Plugins/Runtime/MeshModelingToolset/Source/ModelingComponents/Public/Drawing/PolyEditPreviewMesh.h
+++ b/Engine/Plugins/Runtime/MeshModelingToolset/Source/ModelingComponents/Public/Drawing/PolyEditPreviewMesh.h
@@ -17,7 +17,6 @@ UCLASS(Transient)
 class MODELINGCOMPONENTS_API UPolyEditPreviewMesh : public UPreviewMesh
 {
 	GENERATED_BODY()
-	using FTransform3d = UE::Geometry::FTransform3d;	
 public:
 
 	//
@@ -78,6 +77,6 @@ protected:
 
 	FVector3d InputDirection;
 
-	UE::Geometry::FTransform3d MeshTransform;
+	FTransform3d MeshTransform;
 	bool bHaveMeshTransform;
 };
\ No newline at end of file
diff --git a/Engine/Plugins/Runtime/MeshModelingToolset/Source/ModelingComponents/Public/Mechanics/DragAlignmentMechanic.h b/Engine/Plugins/Runtime/MeshModelingToolset/Source/ModelingComponents/Public/Mechanics/DragAlignmentMechanic.h
--- a/Engine/Plugins/Runtime/MeshModelingToolset/Source/ModelingComponents/Public/Mechanics/DragAlignmentMechanic.h
+++ b/Engine/Plugins/Runtime/MeshModelingToolset/Source/ModelingComponents/Public/Mechanics/DragAlignmentMechanic.h
@@ -54,7 +54,7 @@ public:
 	 */
 	void InitializeDeformedMeshRayCast(
 		TFunction<UE::Geometry::FDynamicMeshAABBTree3* ()> GetSpatialIn,
-		const UE::Geometry::FTransform3d& TargetTransform, const FGroupTopologyDeformer* LinearDeformer);
+		const FTransform3d& TargetTransform, const FGroupTopologyDeformer* LinearDeformer);
 
 	virtual void Setup(UInteractiveTool* ParentToolIn) override;
 	virtual void Shutdown() override;
diff --git a/Engine/Plugins/Runtime/MeshModelingToolset/Source/ModelingComponents/Public/MeshOpPreviewHelpers.h b/Engine/Plugins/Runtime/MeshModelingToolset/Source/ModelingComponents/Public/MeshOpPreviewHelpers.h
--- a/Engine/Plugins/Runtime/MeshModelingToolset/Source/ModelingComponents/Public/MeshOpPreviewHelpers.h
+++ b/Engine/Plugins/Runtime/MeshModelingToolset/Source/ModelingComponents/Public/MeshOpPreviewHelpers.h
@@ -24,7 +24,7 @@ using FBackgroundDynamicMeshComputeSource = UE::Geometry::TBackgroundModelingCom
 struct MODELINGCOMPONENTS_API FDynamicMeshOpResult
 {
 	TUniquePtr<UE::Geometry::FDynamicMesh3> Mesh;
-	UE::Geometry::FTransform3d Transform;
+	UE::Geometry::FTransformSRT3d Transform;
 };
 
 
diff --git a/Engine/Plugins/Runtime/MeshModelingToolset/Source/ModelingComponents/Public/ModelingToolTargetUtil.h b/Engine/Plugins/Runtime/MeshModelingToolset/Source/ModelingComponents/Public/ModelingToolTargetUtil.h
--- a/Engine/Plugins/Runtime/MeshModelingToolset/Source/ModelingComponents/Public/ModelingToolTargetUtil.h
+++ b/Engine/Plugins/Runtime/MeshModelingToolset/Source/ModelingComponents/Public/ModelingToolTargetUtil.h
@@ -67,7 +67,7 @@ MODELINGCOMPONENTS_API bool SetSourceObjectVisible(UToolTarget* Target, bool bVi
 /**
  * @return the local-to-world Transform underlying a ToolTarget, eg the Component or Actor transform
  */
-MODELINGCOMPONENTS_API UE::Geometry::FTransform3d GetLocalToWorldTransform(UToolTarget* Target);
+MODELINGCOMPONENTS_API FTransform3d GetLocalToWorldTransform(UToolTarget* Target);
 
 /**
  * Fetch the Material Set on the object underlying a ToolTarget. In cases where there are (eg) 
diff --git a/Engine/Plugins/Runtime/MeshModelingToolset/Source/ModelingComponents/Public/Selection/PolygonSelectionMechanic.h b/Engine/Plugins/Runtime/MeshModelingToolset/Source/ModelingComponents/Public/Selection/PolygonSelectionMechanic.h
--- a/Engine/Plugins/Runtime/MeshModelingToolset/Source/ModelingComponents/Public/Selection/PolygonSelectionMechanic.h
+++ b/Engine/Plugins/Runtime/MeshModelingToolset/Source/ModelingComponents/Public/Selection/PolygonSelectionMechanic.h
@@ -86,7 +86,6 @@ class MODELINGCOMPONENTS_API UPolygonSelectionMechanic : public UInteractionMech
 {
 	GENERATED_BODY()
 	using FFrame3d = UE::Geometry::FFrame3d;
-	using FTransform3d = UE::Geometry::FTransform3d;	
 	using FAxisAlignedBox3d = UE::Geometry::FAxisAlignedBox3d;
 public:
 
@@ -111,7 +110,7 @@ public:
 	 * @param GetSpatialSourceFunc Function that returns an AABB tree for the mesh.
 	 */
 	void Initialize(const FDynamicMesh3* Mesh,
-		FTransform TargetTransform,
+		FTransform3d TargetTransform,
 		UWorld* World,
 		const FGroupTopology* Topology,
 		TFunction<FDynamicMeshAABBTree3*()> GetSpatialSourceFunc
diff --git a/Engine/Plugins/Runtime/MeshModelingToolset/Source/ModelingOperators/Private/BaseOps/SimpleMeshProcessingBaseOp.cpp b/Engine/Plugins/Runtime/MeshModelingToolset/Source/ModelingOperators/Private/BaseOps/SimpleMeshProcessingBaseOp.cpp
--- a/Engine/Plugins/Runtime/MeshModelingToolset/Source/ModelingOperators/Private/BaseOps/SimpleMeshProcessingBaseOp.cpp
+++ b/Engine/Plugins/Runtime/MeshModelingToolset/Source/ModelingOperators/Private/BaseOps/SimpleMeshProcessingBaseOp.cpp
@@ -40,7 +40,7 @@ FSimpleMeshProcessingBaseOp::FSimpleMeshProcessingBaseOp(const FDynamicMesh3* Me
 	}
 }
 
-void FSimpleMeshProcessingBaseOp::SetTransform(const UE::Geometry::FTransform3d& XForm)
+void FSimpleMeshProcessingBaseOp::SetTransform(const FTransformSRT3d& XForm)
 {
 	ResultTransform = XForm;
 }
diff --git a/Engine/Plugins/Runtime/MeshModelingToolset/Source/ModelingOperators/Private/CleaningOps/EditNormalsOp.cpp b/Engine/Plugins/Runtime/MeshModelingToolset/Source/ModelingOperators/Private/CleaningOps/EditNormalsOp.cpp
--- a/Engine/Plugins/Runtime/MeshModelingToolset/Source/ModelingOperators/Private/CleaningOps/EditNormalsOp.cpp
+++ b/Engine/Plugins/Runtime/MeshModelingToolset/Source/ModelingOperators/Private/CleaningOps/EditNormalsOp.cpp
@@ -10,8 +10,9 @@
 
 using namespace UE::Geometry;
 
-void FEditNormalsOp::SetTransform(const FTransform& Transform) {
-	ResultTransform = (FTransform3d)Transform;
+void FEditNormalsOp::SetTransform(const FTransformSRT3d& Transform) 
+{
+	ResultTransform = Transform;
 }
 
 void FEditNormalsOp::CalculateResult(FProgressCancel* Progress)
diff --git a/Engine/Plugins/Runtime/MeshModelingToolset/Source/ModelingOperators/Private/CleaningOps/RemeshMeshOp.cpp b/Engine/Plugins/Runtime/MeshModelingToolset/Source/ModelingOperators/Private/CleaningOps/RemeshMeshOp.cpp
--- a/Engine/Plugins/Runtime/MeshModelingToolset/Source/ModelingOperators/Private/CleaningOps/RemeshMeshOp.cpp
+++ b/Engine/Plugins/Runtime/MeshModelingToolset/Source/ModelingOperators/Private/CleaningOps/RemeshMeshOp.cpp
@@ -50,9 +50,9 @@ TUniquePtr<FRemesher> FRemeshMeshOp::CreateRemesher(ERemeshType Type, FDynamicMe
 
 }
 
-void FRemeshMeshOp::SetTransform(const FTransform& Transform)
+void FRemeshMeshOp::SetTransform(const FTransformSRT3d& Transform)
 {
-	ResultTransform = (FTransform3d)Transform;
+	ResultTransform = Transform;
 }
 
 void FRemeshMeshOp::CalculateResult(FProgressCancel* Progress)
diff --git a/Engine/Plugins/Runtime/MeshModelingToolset/Source/ModelingOperators/Private/CleaningOps/RemoveOccludedTrianglesOp.cpp b/Engine/Plugins/Runtime/MeshModelingToolset/Source/ModelingOperators/Private/CleaningOps/RemoveOccludedTrianglesOp.cpp
--- a/Engine/Plugins/Runtime/MeshModelingToolset/Source/ModelingOperators/Private/CleaningOps/RemoveOccludedTrianglesOp.cpp
+++ b/Engine/Plugins/Runtime/MeshModelingToolset/Source/ModelingOperators/Private/CleaningOps/RemoveOccludedTrianglesOp.cpp
@@ -15,7 +15,7 @@ using namespace UE::Geometry;
 
 void FRemoveOccludedTrianglesOp::SetTransform(const FTransform& Transform)
 {
-	ResultTransform = (FTransform3d)Transform;
+	ResultTransform = (FTransformSRT3d)Transform;
 }
 
 void FRemoveOccludedTrianglesOp::CalculateResult(FProgressCancel* Progress)
diff --git a/Engine/Plugins/Runtime/MeshModelingToolset/Source/ModelingOperators/Private/CompositionOps/BooleanMeshesOp.cpp b/Engine/Plugins/Runtime/MeshModelingToolset/Source/ModelingOperators/Private/CompositionOps/BooleanMeshesOp.cpp
--- a/Engine/Plugins/Runtime/MeshModelingToolset/Source/ModelingOperators/Private/CompositionOps/BooleanMeshesOp.cpp
+++ b/Engine/Plugins/Runtime/MeshModelingToolset/Source/ModelingOperators/Private/CompositionOps/BooleanMeshesOp.cpp
@@ -8,8 +8,9 @@
 
 using namespace UE::Geometry;
 
-void FBooleanMeshesOp::SetTransform(const FTransform& Transform) {
-	ResultTransform = (FTransform3d)Transform;
+void FBooleanMeshesOp::SetTransform(const FTransformSRT3d& Transform) 
+{
+	ResultTransform = Transform;
 }
 
 void FBooleanMeshesOp::CalculateResult(FProgressCancel* Progress)
@@ -64,7 +65,7 @@ void FBooleanMeshesOp::CalculateResult(FProgressCancel* Progress)
 		}
 	}
 
-	FMeshBoolean MeshBoolean(Meshes[FirstIdx].Get(), (FTransform3d)Transforms[FirstIdx], Meshes[OtherIdx].Get(), (FTransform3d)Transforms[OtherIdx], ResultMesh.Get(), Op);
+	FMeshBoolean MeshBoolean(Meshes[FirstIdx].Get(), Transforms[FirstIdx], Meshes[OtherIdx].Get(), Transforms[OtherIdx], ResultMesh.Get(), Op);
 	if (Progress && Progress->Cancelled())
 	{
 		return;
diff --git a/Engine/Plugins/Runtime/MeshModelingToolset/Source/ModelingOperators/Private/CompositionOps/CubeGridBooleanOp.cpp b/Engine/Plugins/Runtime/MeshModelingToolset/Source/ModelingOperators/Private/CompositionOps/CubeGridBooleanOp.cpp
--- a/Engine/Plugins/Runtime/MeshModelingToolset/Source/ModelingOperators/Private/CompositionOps/CubeGridBooleanOp.cpp
+++ b/Engine/Plugins/Runtime/MeshModelingToolset/Source/ModelingOperators/Private/CompositionOps/CubeGridBooleanOp.cpp
@@ -434,7 +434,7 @@ void FCubeGridBooleanOp::CalculateResult(FProgressCancel* Progress)
 	
 	// Perform the boolean operation.
 	FMeshBoolean MeshBoolean(ResultMesh.Get(), InputTransform,
-		OpMesh.Get(), FTransform3d::Identity(), ResultMesh.Get(),
+		OpMesh.Get(), FTransformSRT3d::Identity(), ResultMesh.Get(),
 		bSubtract ? FMeshBoolean::EBooleanOp::Difference : FMeshBoolean::EBooleanOp::Union);
 
 	MeshBoolean.bPutResultInInputSpace = true;
@@ -455,7 +455,7 @@ void FCubeGridBooleanOp::CalculateResult(FProgressCancel* Progress)
 	{
 		// Set the transform to be in the center.
 		FVector3d Center = ResultMesh->GetBounds().Center();
-		ResultTransform = FTransform3d(Center);
+		ResultTransform = FTransformSRT3d(Center);
 	}
 	MeshTransforms::ApplyTransformInverse(*ResultMesh, InputTransform);
 
diff --git a/Engine/Plugins/Runtime/MeshModelingToolset/Source/ModelingOperators/Private/CompositionOps/CurveSweepOp.cpp b/Engine/Plugins/Runtime/MeshModelingToolset/Source/ModelingOperators/Private/CompositionOps/CurveSweepOp.cpp
--- a/Engine/Plugins/Runtime/MeshModelingToolset/Source/ModelingOperators/Private/CompositionOps/CurveSweepOp.cpp
+++ b/Engine/Plugins/Runtime/MeshModelingToolset/Source/ModelingOperators/Private/CompositionOps/CurveSweepOp.cpp
@@ -204,7 +204,7 @@ void FCurveSweepOp::CalculateResult(FProgressCancel* Progress)
 	// Set the pivot of the mesh to the bounding box center
 	FVector3d PivotLocation = ResultMesh->GetBounds(true).Center();
 	MeshTransforms::Translate(*ResultMesh, -PivotLocation);
-	ResultTransform = FTransform3d::Identity();
+	ResultTransform = FTransformSRT3d::Identity();
 	ResultTransform.SetTranslation(PivotLocation);
 }
 
diff --git a/Engine/Plugins/Runtime/MeshModelingToolset/Source/ModelingOperators/Private/CompositionOps/MirrorOp.cpp b/Engine/Plugins/Runtime/MeshModelingToolset/Source/ModelingOperators/Private/CompositionOps/MirrorOp.cpp
--- a/Engine/Plugins/Runtime/MeshModelingToolset/Source/ModelingOperators/Private/CompositionOps/MirrorOp.cpp
+++ b/Engine/Plugins/Runtime/MeshModelingToolset/Source/ModelingOperators/Private/CompositionOps/MirrorOp.cpp
@@ -7,8 +7,9 @@
 
 using namespace UE::Geometry;
 
-void FMirrorOp::SetTransform(const FTransform& Transform) {
-	ResultTransform = (FTransform3d)Transform;
+void FMirrorOp::SetTransform(const FTransformSRT3d& Transform) 
+{
+	ResultTransform = Transform;
 }
 
 void FMirrorOp::CalculateResult(FProgressCancel* Progress)
diff --git a/Engine/Plugins/Runtime/MeshModelingToolset/Source/ModelingOperators/Private/CompositionOps/SelfUnionMeshesOp.cpp b/Engine/Plugins/Runtime/MeshModelingToolset/Source/ModelingOperators/Private/CompositionOps/SelfUnionMeshesOp.cpp
--- a/Engine/Plugins/Runtime/MeshModelingToolset/Source/ModelingOperators/Private/CompositionOps/SelfUnionMeshesOp.cpp
+++ b/Engine/Plugins/Runtime/MeshModelingToolset/Source/ModelingOperators/Private/CompositionOps/SelfUnionMeshesOp.cpp
@@ -8,8 +8,9 @@
 
 using namespace UE::Geometry;
 
-void FSelfUnionMeshesOp::SetTransform(const FTransform& Transform) {
-	ResultTransform = (FTransform3d)Transform;
+void FSelfUnionMeshesOp::SetTransform(const FTransformSRT3d& Transform) 
+{
+	ResultTransform = Transform;
 }
 
 void FSelfUnionMeshesOp::CalculateResult(FProgressCancel* Progress)
diff --git a/Engine/Plugins/Runtime/MeshModelingToolset/Source/ModelingOperators/Private/CompositionOps/VoxelBlendMeshesOp.cpp b/Engine/Plugins/Runtime/MeshModelingToolset/Source/ModelingOperators/Private/CompositionOps/VoxelBlendMeshesOp.cpp
--- a/Engine/Plugins/Runtime/MeshModelingToolset/Source/ModelingOperators/Private/CompositionOps/VoxelBlendMeshesOp.cpp
+++ b/Engine/Plugins/Runtime/MeshModelingToolset/Source/ModelingOperators/Private/CompositionOps/VoxelBlendMeshesOp.cpp
@@ -14,8 +14,9 @@
 
 using namespace UE::Geometry;
 
-void FVoxelBlendMeshesOp::SetTransform(const FTransform& Transform) {
-	ResultTransform = (FTransform3d)Transform;
+void FVoxelBlendMeshesOp::SetTransform(const FTransformSRT3d& Transform) 
+{
+	ResultTransform = Transform;
 }
 
 void FVoxelBlendMeshesOp::CalculateResult(FProgressCancel* Progress)
@@ -47,7 +48,7 @@ void FVoxelBlendMeshesOp::CalculateResult(FProgressCancel* Progress)
 		{
 			TransformedMeshes[MeshIdx].ReverseOrientation(false);
 		}
-		MeshTransforms::ApplyTransform(TransformedMeshes[MeshIdx], (FTransform3d)Transforms[MeshIdx]);
+		MeshTransforms::ApplyTransform(TransformedMeshes[MeshIdx], (FTransformSRT3d)Transforms[MeshIdx]);
 
 		if (bVoxWrap)
 		{
diff --git a/Engine/Plugins/Runtime/MeshModelingToolset/Source/ModelingOperators/Private/CompositionOps/VoxelMorphologyMeshesOp.cpp b/Engine/Plugins/Runtime/MeshModelingToolset/Source/ModelingOperators/Private/CompositionOps/VoxelMorphologyMeshesOp.cpp
--- a/Engine/Plugins/Runtime/MeshModelingToolset/Source/ModelingOperators/Private/CompositionOps/VoxelMorphologyMeshesOp.cpp
+++ b/Engine/Plugins/Runtime/MeshModelingToolset/Source/ModelingOperators/Private/CompositionOps/VoxelMorphologyMeshesOp.cpp
@@ -14,8 +14,9 @@
 
 using namespace UE::Geometry;
 
-void FVoxelMorphologyMeshesOp::SetTransform(const FTransform& Transform) {
-	ResultTransform = (FTransform3d)Transform;
+void FVoxelMorphologyMeshesOp::SetTransform(const FTransformSRT3d& Transform) 
+{
+	ResultTransform = Transform;
 }
 
 void FVoxelMorphologyMeshesOp::CalculateResult(FProgressCancel* Progress)
@@ -59,7 +60,7 @@ void FVoxelMorphologyMeshesOp::CalculateResult(FProgressCancel* Progress)
 	FDynamicMeshEditor AppendEditor(&CombinedMesh);
 	for (int MeshIdx = 0; MeshIdx < Meshes.Num(); MeshIdx++)
 	{
-		FTransform3d MeshTransform = (FTransform3d)Transforms[MeshIdx];
+		FTransformSRT3d MeshTransform = Transforms[MeshIdx];
 		bool bReverseOrientation = MeshTransform.GetDeterminant() < 0;
 		FMeshIndexMappings IndexMaps;
 		AppendEditor.AppendMesh(Meshes[MeshIdx].Get(), IndexMaps,
diff --git a/Engine/Plugins/Runtime/MeshModelingToolset/Source/ModelingOperators/Private/CompositionOps/VoxelSolidifyMeshesOp.cpp b/Engine/Plugins/Runtime/MeshModelingToolset/Source/ModelingOperators/Private/CompositionOps/VoxelSolidifyMeshesOp.cpp
--- a/Engine/Plugins/Runtime/MeshModelingToolset/Source/ModelingOperators/Private/CompositionOps/VoxelSolidifyMeshesOp.cpp
+++ b/Engine/Plugins/Runtime/MeshModelingToolset/Source/ModelingOperators/Private/CompositionOps/VoxelSolidifyMeshesOp.cpp
@@ -13,8 +13,9 @@
 
 using namespace UE::Geometry;
 
-void FVoxelSolidifyMeshesOp::SetTransform(const FTransform& Transform) {
-	ResultTransform = (FTransform3d)Transform;
+void FVoxelSolidifyMeshesOp::SetTransform(const FTransformSRT3d& Transform) 
+{
+	ResultTransform = Transform;
 }
 
 void FVoxelSolidifyMeshesOp::CalculateResult(FProgressCancel* Progress)
@@ -35,7 +36,7 @@ void FVoxelSolidifyMeshesOp::CalculateResult(FProgressCancel* Progress)
 	FDynamicMeshEditor AppendEditor(&CombinedMesh);
 	for (int MeshIdx = 0; MeshIdx < Meshes.Num(); MeshIdx++)
 	{
-		FTransform3d MeshTransform = (FTransform3d)Transforms[MeshIdx];
+		FTransformSRT3d MeshTransform = Transforms[MeshIdx];
 		bool bReverseOrientation = MeshTransform.GetDeterminant() < 0;
 		FMeshIndexMappings IndexMaps;
 		AppendEditor.AppendMesh(Meshes[MeshIdx].Get(), IndexMaps,
diff --git a/Engine/Plugins/Runtime/MeshModelingToolset/Source/ModelingOperators/Private/CuttingOps/EdgeLoopInsertionOp.cpp b/Engine/Plugins/Runtime/MeshModelingToolset/Source/ModelingOperators/Private/CuttingOps/EdgeLoopInsertionOp.cpp
--- a/Engine/Plugins/Runtime/MeshModelingToolset/Source/ModelingOperators/Private/CuttingOps/EdgeLoopInsertionOp.cpp
+++ b/Engine/Plugins/Runtime/MeshModelingToolset/Source/ModelingOperators/Private/CuttingOps/EdgeLoopInsertionOp.cpp
@@ -6,8 +6,9 @@
 
 using namespace UE::Geometry;
 
-void FEdgeLoopInsertionOp::SetTransform(const FTransform& Transform) {
-	ResultTransform = (FTransform3d)Transform;
+void FEdgeLoopInsertionOp::SetTransform(const FTransformSRT3d& Transform) 
+{
+	ResultTransform = Transform;
 }
 
 void FEdgeLoopInsertionOp::GetLoopEdgeLocations(TArray<TPair<FVector3d, FVector3d>>& EndPointPairsOut) const
diff --git a/Engine/Plugins/Runtime/MeshModelingToolset/Source/ModelingOperators/Private/CuttingOps/GroupEdgeInsertionOp.cpp b/Engine/Plugins/Runtime/MeshModelingToolset/Source/ModelingOperators/Private/CuttingOps/GroupEdgeInsertionOp.cpp
--- a/Engine/Plugins/Runtime/MeshModelingToolset/Source/ModelingOperators/Private/CuttingOps/GroupEdgeInsertionOp.cpp
+++ b/Engine/Plugins/Runtime/MeshModelingToolset/Source/ModelingOperators/Private/CuttingOps/GroupEdgeInsertionOp.cpp
@@ -7,8 +7,9 @@
 
 using namespace UE::Geometry;
 
-void FGroupEdgeInsertionOp::SetTransform(const FTransform& Transform) {
-	ResultTransform = (FTransform3d)Transform;
+void FGroupEdgeInsertionOp::SetTransform(const FTransformSRT3d& Transform) 
+{
+	ResultTransform = Transform;
 }
 
 void FGroupEdgeInsertionOp::GetEdgeLocations(TArray<TPair<FVector3d, FVector3d>>& EndPointPairsOut) const
diff --git a/Engine/Plugins/Runtime/MeshModelingToolset/Source/ModelingOperators/Private/CuttingOps/PlaneCutOp.cpp b/Engine/Plugins/Runtime/MeshModelingToolset/Source/ModelingOperators/Private/CuttingOps/PlaneCutOp.cpp
--- a/Engine/Plugins/Runtime/MeshModelingToolset/Source/ModelingOperators/Private/CuttingOps/PlaneCutOp.cpp
+++ b/Engine/Plugins/Runtime/MeshModelingToolset/Source/ModelingOperators/Private/CuttingOps/PlaneCutOp.cpp
@@ -12,8 +12,9 @@ using namespace UE::Geometry;
 
 const FName FPlaneCutOp::ObjectIndexAttribute = "ObjectIndexAttribute";
 
-void FPlaneCutOp::SetTransform(const FTransform& Transform) {
-	ResultTransform = (FTransform3d)Transform;
+void FPlaneCutOp::SetTransform(const FTransformSRT3d& Transform) 
+{
+	ResultTransform = Transform;
 }
 
 void FPlaneCutOp::CalculateResult(FProgressCancel* Progress)
diff --git a/Engine/Plugins/Runtime/MeshModelingToolset/Source/ModelingOperators/Private/DeformationOps/ExtrudeOp.cpp b/Engine/Plugins/Runtime/MeshModelingToolset/Source/ModelingOperators/Private/DeformationOps/ExtrudeOp.cpp
--- a/Engine/Plugins/Runtime/MeshModelingToolset/Source/ModelingOperators/Private/DeformationOps/ExtrudeOp.cpp
+++ b/Engine/Plugins/Runtime/MeshModelingToolset/Source/ModelingOperators/Private/DeformationOps/ExtrudeOp.cpp
@@ -175,7 +175,7 @@ bool FExtrudeOp::BooleanExtrude(FProgressCancel* Progress)
 
 	// Now perform a boolean operation with our result.
 	FMeshBoolean::EBooleanOp Op = ExtrudeDistance > 0 ? FMeshBoolean::EBooleanOp::Union : FMeshBoolean::EBooleanOp::Difference;
-	FMeshBoolean MeshBoolean(ResultMesh.Get(), FTransform3d::Identity(), &Submesh->GetSubmesh(), FTransform3d::Identity(), ResultMesh.Get(), Op);
+	FMeshBoolean MeshBoolean(ResultMesh.Get(), FTransformSRT3d::Identity(), &Submesh->GetSubmesh(), FTransformSRT3d::Identity(), ResultMesh.Get(), Op);
 	MeshBoolean.bPutResultInInputSpace = true;
 	MeshBoolean.bSimplifyAlongNewEdges = true;
 	MeshBoolean.bPopulateSecondMeshGroupMap = true;
diff --git a/Engine/Plugins/Runtime/MeshModelingToolset/Source/ModelingOperators/Private/MeshSpaceDeformerOps/MeshSpaceDeformerOp.cpp b/Engine/Plugins/Runtime/MeshModelingToolset/Source/ModelingOperators/Private/MeshSpaceDeformerOps/MeshSpaceDeformerOp.cpp
--- a/Engine/Plugins/Runtime/MeshModelingToolset/Source/ModelingOperators/Private/MeshSpaceDeformerOps/MeshSpaceDeformerOp.cpp
+++ b/Engine/Plugins/Runtime/MeshModelingToolset/Source/ModelingOperators/Private/MeshSpaceDeformerOps/MeshSpaceDeformerOp.cpp
@@ -5,9 +5,9 @@
 
 using namespace UE::Geometry;
 
-void FMeshSpaceDeformerOp::SetTransform(const FTransform& Transform)
+void FMeshSpaceDeformerOp::SetTransform(const FTransformSRT3d& Transform)
 {
-	ResultTransform = (FTransform3d)Transform;
+	ResultTransform = Transform;
 }
 
 void FMeshSpaceDeformerOp::CalculateResult(FProgressCancel* Progress)
diff --git a/Engine/Plugins/Runtime/MeshModelingToolset/Source/ModelingOperators/Private/ParameterizationOps/UVLayoutOp.cpp b/Engine/Plugins/Runtime/MeshModelingToolset/Source/ModelingOperators/Private/ParameterizationOps/UVLayoutOp.cpp
--- a/Engine/Plugins/Runtime/MeshModelingToolset/Source/ModelingOperators/Private/ParameterizationOps/UVLayoutOp.cpp
+++ b/Engine/Plugins/Runtime/MeshModelingToolset/Source/ModelingOperators/Private/ParameterizationOps/UVLayoutOp.cpp
@@ -10,8 +10,9 @@
 
 using namespace UE::Geometry;
 
-void FUVLayoutOp::SetTransform(const FTransform& Transform) {
-	ResultTransform = (FTransform3d)Transform;
+void FUVLayoutOp::SetTransform(const FTransformSRT3d& Transform) 
+{
+	ResultTransform = Transform;
 }
 
 
diff --git a/Engine/Plugins/Runtime/MeshModelingToolset/Source/ModelingOperators/Private/SmoothingOps/SmoothingOpBase.cpp b/Engine/Plugins/Runtime/MeshModelingToolset/Source/ModelingOperators/Private/SmoothingOps/SmoothingOpBase.cpp
--- a/Engine/Plugins/Runtime/MeshModelingToolset/Source/ModelingOperators/Private/SmoothingOps/SmoothingOpBase.cpp
+++ b/Engine/Plugins/Runtime/MeshModelingToolset/Source/ModelingOperators/Private/SmoothingOps/SmoothingOpBase.cpp
@@ -22,7 +22,7 @@ FSmoothingOpBase::FSmoothingOpBase(const FDynamicMesh3* Mesh, const FOptions& Op
 	}
 }
 
-void FSmoothingOpBase::SetTransform(const UE::Geometry::FTransform3d& XForm)
+void FSmoothingOpBase::SetTransform(const FTransformSRT3d& XForm)
 {
 	ResultTransform = XForm;
 }
diff --git a/Engine/Plugins/Runtime/MeshModelingToolset/Source/ModelingOperators/Public/BaseOps/SimpleMeshProcessingBaseOp.h b/Engine/Plugins/Runtime/MeshModelingToolset/Source/ModelingOperators/Public/BaseOps/SimpleMeshProcessingBaseOp.h
--- a/Engine/Plugins/Runtime/MeshModelingToolset/Source/ModelingOperators/Public/BaseOps/SimpleMeshProcessingBaseOp.h
+++ b/Engine/Plugins/Runtime/MeshModelingToolset/Source/ModelingOperators/Public/BaseOps/SimpleMeshProcessingBaseOp.h
@@ -38,7 +38,7 @@ public:
 	virtual ~FSimpleMeshProcessingBaseOp() {}
 
 	// set ability on protected transform.
-	void SetTransform(const FTransform3d& XForm);
+	void SetTransform(const FTransformSRT3d& XForm);
 
 	// copy the PositionBuffer locations back to the ResultMesh
 	void UpdateResultMeshPositions();
diff --git a/Engine/Plugins/Runtime/MeshModelingToolset/Source/ModelingOperators/Public/CleaningOps/EditNormalsOp.h b/Engine/Plugins/Runtime/MeshModelingToolset/Source/ModelingOperators/Public/CleaningOps/EditNormalsOp.h
--- a/Engine/Plugins/Runtime/MeshModelingToolset/Source/ModelingOperators/Public/CleaningOps/EditNormalsOp.h
+++ b/Engine/Plugins/Runtime/MeshModelingToolset/Source/ModelingOperators/Public/CleaningOps/EditNormalsOp.h
@@ -57,7 +57,7 @@ public:
 	bool bAllowSharpVertices;
 	float NormalSplitThreshold;
 
-	void SetTransform(const FTransform& Transform);
+	void SetTransform(const FTransformSRT3d& Transform);
 
 	//
 	// FDynamicMeshOperator implementation
diff --git a/Engine/Plugins/Runtime/MeshModelingToolset/Source/ModelingOperators/Public/CleaningOps/RemeshMeshOp.h b/Engine/Plugins/Runtime/MeshModelingToolset/Source/ModelingOperators/Public/CleaningOps/RemeshMeshOp.h
--- a/Engine/Plugins/Runtime/MeshModelingToolset/Source/ModelingOperators/Public/CleaningOps/RemeshMeshOp.h
+++ b/Engine/Plugins/Runtime/MeshModelingToolset/Source/ModelingOperators/Public/CleaningOps/RemeshMeshOp.h
@@ -73,8 +73,8 @@ public:
 	FDynamicMesh3* ProjectionTarget = nullptr;
 	FDynamicMeshAABBTree3* ProjectionTargetSpatial = nullptr;
 
-	FTransform3d TargetMeshLocalToWorld;
-	FTransform3d ToolMeshLocalToWorld;
+	FTransformSRT3d TargetMeshLocalToWorld;
+	FTransformSRT3d ToolMeshLocalToWorld;
 	bool bUseWorldSpace = false;
 	bool bParallel = true;
 
@@ -102,7 +102,7 @@ public:
 	// End normal flow only
 
 
-	void SetTransform(const FTransform& Transform);
+	void SetTransform(const FTransformSRT3d& Transform);
 
 	//
 	// FDynamicMeshOperator implementation
diff --git a/Engine/Plugins/Runtime/MeshModelingToolset/Source/ModelingOperators/Public/CleaningOps/RemoveOccludedTrianglesOp.h b/Engine/Plugins/Runtime/MeshModelingToolset/Source/ModelingOperators/Public/CleaningOps/RemoveOccludedTrianglesOp.h
--- a/Engine/Plugins/Runtime/MeshModelingToolset/Source/ModelingOperators/Public/CleaningOps/RemoveOccludedTrianglesOp.h
+++ b/Engine/Plugins/Runtime/MeshModelingToolset/Source/ModelingOperators/Public/CleaningOps/RemoveOccludedTrianglesOp.h
@@ -24,9 +24,9 @@ public:
 	
 	TArray<TSharedPtr<UE::Geometry::FDynamicMeshAABBTree3, ESPMode::ThreadSafe>> OccluderTrees;
 	TArray<TSharedPtr<UE::Geometry::TFastWindingTree<FDynamicMesh3>, ESPMode::ThreadSafe>> OccluderWindings;
-	TArray<FTransform3d> OccluderTransforms;
+	TArray<FTransformSRT3d> OccluderTransforms;
 
-	TArray<FTransform3d> MeshTransforms;
+	TArray<FTransformSRT3d> MeshTransforms;
 
 	UE::Geometry::EOcclusionTriangleSampling TriangleSamplingMethod =
 		UE::Geometry::EOcclusionTriangleSampling::Centroids;
diff --git a/Engine/Plugins/Runtime/MeshModelingToolset/Source/ModelingOperators/Public/CompositionOps/BooleanMeshesOp.h b/Engine/Plugins/Runtime/MeshModelingToolset/Source/ModelingOperators/Public/CompositionOps/BooleanMeshesOp.h
--- a/Engine/Plugins/Runtime/MeshModelingToolset/Source/ModelingOperators/Public/CompositionOps/BooleanMeshesOp.h
+++ b/Engine/Plugins/Runtime/MeshModelingToolset/Source/ModelingOperators/Public/CompositionOps/BooleanMeshesOp.h
@@ -62,7 +62,7 @@ public:
 	ETrimSide TrimSide;
 	bool bTrimMode = false; // if true, do a trim operation instead of a boolean
 	TArray<TSharedPtr<const FDynamicMesh3, ESPMode::ThreadSafe>> Meshes;
-	TArray<FTransform> Transforms; // 1:1 with Meshes
+	TArray<FTransformSRT3d> Transforms; // 1:1 with Meshes
 	bool bAttemptFixHoles = false;
 	double WindingThreshold = 0.5;
 
@@ -71,7 +71,7 @@ public:
 	/** Angle threshold in degrees used for testing if two triangles should be considered coplanar, or two lines collinear */
 	float TryCollapseExtraEdgesPlanarThresh = 0.01f;
 
-	void SetTransform(const FTransform& Transform);
+	void SetTransform(const FTransformSRT3d& Transform);
 
 	//
 	// FDynamicMeshOperator implementation
diff --git a/Engine/Plugins/Runtime/MeshModelingToolset/Source/ModelingOperators/Public/CompositionOps/CubeGridBooleanOp.h b/Engine/Plugins/Runtime/MeshModelingToolset/Source/ModelingOperators/Public/CompositionOps/CubeGridBooleanOp.h
--- a/Engine/Plugins/Runtime/MeshModelingToolset/Source/ModelingOperators/Public/CompositionOps/CubeGridBooleanOp.h
+++ b/Engine/Plugins/Runtime/MeshModelingToolset/Source/ModelingOperators/Public/CompositionOps/CubeGridBooleanOp.h
@@ -27,7 +27,7 @@ public:
 
 	// The mesh to add to or subtract from
 	TSharedPtr<const FDynamicMesh3, ESPMode::ThreadSafe> InputMesh;
-	FTransform3d InputTransform = FTransform3d::Identity();
+	FTransformSRT3d InputTransform = FTransformSRT3d::Identity();
 
 	// Box from which to generate the second mesh
 	FOrientedBox3d WorldBox;
diff --git a/Engine/Plugins/Runtime/MeshModelingToolset/Source/ModelingOperators/Public/CompositionOps/MirrorOp.h b/Engine/Plugins/Runtime/MeshModelingToolset/Source/ModelingOperators/Public/CompositionOps/MirrorOp.h
--- a/Engine/Plugins/Runtime/MeshModelingToolset/Source/ModelingOperators/Public/CompositionOps/MirrorOp.h
+++ b/Engine/Plugins/Runtime/MeshModelingToolset/Source/ModelingOperators/Public/CompositionOps/MirrorOp.h
@@ -36,7 +36,7 @@ public:
 	/** Whether, when welding, new bowtie vertex creation should be allowed. */
 	bool bAllowBowtieVertexCreation = false;
 
-	void SetTransform(const FTransform& Transform);
+	void SetTransform(const FTransformSRT3d& Transform);
 
 	//
 	// FDynamicMeshOperator implementation
diff --git a/Engine/Plugins/Runtime/MeshModelingToolset/Source/ModelingOperators/Public/CompositionOps/VoxelMorphologyMeshesOp.h b/Engine/Plugins/Runtime/MeshModelingToolset/Source/ModelingOperators/Public/CompositionOps/VoxelMorphologyMeshesOp.h
--- a/Engine/Plugins/Runtime/MeshModelingToolset/Source/ModelingOperators/Public/CompositionOps/VoxelMorphologyMeshesOp.h
+++ b/Engine/Plugins/Runtime/MeshModelingToolset/Source/ModelingOperators/Public/CompositionOps/VoxelMorphologyMeshesOp.h
@@ -41,7 +41,7 @@ public:
 
 	// inputs
 	TArray<TSharedPtr<const FDynamicMesh3, ESPMode::ThreadSafe>> Meshes;
-	TArray<FTransform> Transforms; // 1:1 with Meshes
+	TArray<FTransformSRT3d> Transforms; // 1:1 with Meshes
 
 	double Distance = 1.0;
 	EMorphologyOperation Operation;
@@ -50,7 +50,7 @@ public:
 	bool bRemoveInternalsAfterVoxWrap = false;
 	double ThickenShells = 0.0;
 
-	void SetTransform(const FTransform& Transform);
+	void SetTransform(const FTransformSRT3d& Transform);
 
 	//
 	// FDynamicMeshOperator implementation
diff --git a/Engine/Plugins/Runtime/MeshModelingToolset/Source/ModelingOperators/Public/DeformationOps/ExtrudeOp.h b/Engine/Plugins/Runtime/MeshModelingToolset/Source/ModelingOperators/Public/DeformationOps/ExtrudeOp.h
--- a/Engine/Plugins/Runtime/MeshModelingToolset/Source/ModelingOperators/Public/DeformationOps/ExtrudeOp.h
+++ b/Engine/Plugins/Runtime/MeshModelingToolset/Source/ModelingOperators/Public/DeformationOps/ExtrudeOp.h
@@ -66,9 +66,9 @@ public:
 	// Outputs
 	TArray<int32> ExtrudedFaceNewGids;
 
-	void SetTransform(const FTransform& Transform)
+	void SetTransform(const FTransformSRT3d& Transform)
 	{
-		ResultTransform = (FTransform3d)Transform;
+		ResultTransform = Transform;
 	}
 
 	// FDynamicMeshOperator implementation 
diff --git a/Engine/Plugins/Runtime/MeshModelingToolset/Source/ModelingOperators/Public/ModelingOperators.h b/Engine/Plugins/Runtime/MeshModelingToolset/Source/ModelingOperators/Public/ModelingOperators.h
--- a/Engine/Plugins/Runtime/MeshModelingToolset/Source/ModelingOperators/Public/ModelingOperators.h
+++ b/Engine/Plugins/Runtime/MeshModelingToolset/Source/ModelingOperators/Public/ModelingOperators.h
@@ -14,21 +14,21 @@ namespace Geometry
 
 /**
  * FDynamicMeshOperator is a base interface for operator implementations that can produce 
- * a FDynamicMesh3 with associated FTransform3d. This class exists so we can write generic code 
+ * a FDynamicMesh3 with associated FTransformSRT3d. This class exists so we can write generic code 
  * that works with any type of mesh operation of this style. 
  */
 class FDynamicMeshOperator
 {
 protected:
 	TUniquePtr<FDynamicMesh3> ResultMesh;
-	FTransform3d ResultTransform;
+	FTransformSRT3d ResultTransform;
 	FGeometryResult ResultInfo;
 
 public:
 	FDynamicMeshOperator()
 	{
 		ResultMesh = MakeUnique<FDynamicMesh3>();
-		ResultTransform = FTransform3d::Identity();
+		ResultTransform = FTransformSRT3d::Identity();
 	}
 	virtual ~FDynamicMeshOperator()
 	{
@@ -37,7 +37,7 @@ public:
 	/**
 	 * Set the output transform
 	 */
-	virtual void SetResultTransform(const FTransform3d& Transform)
+	virtual void SetResultTransform(const FTransformSRT3d& Transform)
 	{
 		ResultTransform = Transform;
 	}
@@ -61,7 +61,7 @@ public:
 	/**
 	 * @return the transform applied to the mesh produced by CalculateResult()
 	 */
-	const FTransform3d& GetResultTransform() const
+	const FTransformSRT3d& GetResultTransform() const
 	{
 		return ResultTransform;
 	}
diff --git a/Engine/Plugins/Runtime/MeshModelingToolset/Source/ModelingOperators/Public/ParameterizationOps/RecomputeUVsOp.h b/Engine/Plugins/Runtime/MeshModelingToolset/Source/ModelingOperators/Public/ParameterizationOps/RecomputeUVsOp.h
--- a/Engine/Plugins/Runtime/MeshModelingToolset/Source/ModelingOperators/Public/ParameterizationOps/RecomputeUVsOp.h
+++ b/Engine/Plugins/Runtime/MeshModelingToolset/Source/ModelingOperators/Public/ParameterizationOps/RecomputeUVsOp.h
@@ -81,7 +81,7 @@ public:
 	double MaxNormalDeviationDeg = 45.0;
 
 	// set ability on protected transform.
-	void SetTransform(const FTransform3d& XForm)
+	void SetTransform(const FTransformSRT3d& XForm)
 	{
 		ResultTransform = XForm;
 	}
@@ -129,5 +129,5 @@ public:
 	TSharedPtr<UE::Geometry::FPolygroupSet, ESPMode::ThreadSafe> InputGroups;
 	TSharedPtr<UE::Geometry::FDynamicMesh3, ESPMode::ThreadSafe> OriginalMesh;
 	TUniqueFunction<int32()> GetSelectedUVChannel = []() { return 0; };
-	UE::Geometry::FTransform3d TargetTransform;
+	UE::Geometry::FTransformSRT3d TargetTransform;
 };
\ No newline at end of file
diff --git a/Engine/Plugins/Runtime/MeshModelingToolset/Source/ModelingOperators/Public/ParameterizationOps/UVProjectionOp.h b/Engine/Plugins/Runtime/MeshModelingToolset/Source/ModelingOperators/Public/ParameterizationOps/UVProjectionOp.h
--- a/Engine/Plugins/Runtime/MeshModelingToolset/Source/ModelingOperators/Public/ParameterizationOps/UVProjectionOp.h
+++ b/Engine/Plugins/Runtime/MeshModelingToolset/Source/ModelingOperators/Public/ParameterizationOps/UVProjectionOp.h
@@ -38,7 +38,7 @@ public:
 	int32 UseUVLayer = 0;
 
 	// transform from Mesh into space of Projection Geometry
-	FTransform3d MeshToProjectionSpace;
+	FTransformSRT3d MeshToProjectionSpace;
 
 	EUVProjectionMethod ProjectionMethod = EUVProjectionMethod::Box;
 	FOrientedBox3d ProjectionBox;
diff --git a/Engine/Plugins/Runtime/MeshModelingToolset/Source/ModelingOperators/Public/SmoothingOps/SmoothingOpBase.h b/Engine/Plugins/Runtime/MeshModelingToolset/Source/ModelingOperators/Public/SmoothingOps/SmoothingOpBase.h
--- a/Engine/Plugins/Runtime/MeshModelingToolset/Source/ModelingOperators/Public/SmoothingOps/SmoothingOpBase.h
+++ b/Engine/Plugins/Runtime/MeshModelingToolset/Source/ModelingOperators/Public/SmoothingOps/SmoothingOpBase.h
@@ -64,7 +64,7 @@ public:
 	virtual ~FSmoothingOpBase() override {}
 
 	// set ability on protected transform.
-	void SetTransform(const FTransform3d& XForm);
+	void SetTransform(const FTransformSRT3d& XForm);
 
 	// base class overrides this.  Results in updated ResultMesh.
 	virtual void CalculateResult(FProgressCancel* Progress) override = 0;
diff --git a/Engine/Plugins/Runtime/MeshModelingToolset/Source/ModelingOperatorsEditorOnly/Private/ParameterizationOps/ParameterizeMeshOp.cpp b/Engine/Plugins/Runtime/MeshModelingToolset/Source/ModelingOperatorsEditorOnly/Private/ParameterizationOps/ParameterizeMeshOp.cpp
--- a/Engine/Plugins/Runtime/MeshModelingToolset/Source/ModelingOperatorsEditorOnly/Private/ParameterizationOps/ParameterizeMeshOp.cpp
+++ b/Engine/Plugins/Runtime/MeshModelingToolset/Source/ModelingOperatorsEditorOnly/Private/ParameterizationOps/ParameterizeMeshOp.cpp
@@ -489,7 +489,7 @@ TUniquePtr<FDynamicMeshOperator> UParameterizeMeshOperatorFactory::MakeNewOperat
 
 	Op->UVLayer = GetSelectedUVChannel();	
 
-	UE::Geometry::FTransform3d Transform(TargetTransform);
+	FTransformSRT3d Transform(TargetTransform);
 	Op->SetTransform(Transform);
 
 	return Op;
diff --git a/Engine/Plugins/Runtime/MeshModelingToolset/Source/ModelingOperatorsEditorOnly/Public/ParameterizationOps/ParameterizeMeshOp.h b/Engine/Plugins/Runtime/MeshModelingToolset/Source/ModelingOperatorsEditorOnly/Public/ParameterizationOps/ParameterizeMeshOp.h
--- a/Engine/Plugins/Runtime/MeshModelingToolset/Source/ModelingOperatorsEditorOnly/Public/ParameterizationOps/ParameterizeMeshOp.h
+++ b/Engine/Plugins/Runtime/MeshModelingToolset/Source/ModelingOperatorsEditorOnly/Public/ParameterizationOps/ParameterizeMeshOp.h
@@ -68,7 +68,7 @@ public:
 	EParamOpBackend Method = EParamOpBackend::UVAtlas;
 
 	// set ability on protected transform.
-	void SetTransform(const FTransform3d& XForm)
+	void SetTransform(const FTransformSRT3d& XForm)
 	{
 		ResultTransform = XForm;
 	}
diff --git a/Engine/Source/Runtime/GeometryCore/Private/DynamicMesh/MeshAdapterUtil.cpp b/Engine/Source/Runtime/GeometryCore/Private/DynamicMesh/MeshAdapterUtil.cpp
--- a/Engine/Source/Runtime/GeometryCore/Private/DynamicMesh/MeshAdapterUtil.cpp
+++ b/Engine/Source/Runtime/GeometryCore/Private/DynamicMesh/MeshAdapterUtil.cpp
@@ -89,13 +89,12 @@ FPointSetAdapterd UE::Geometry::MakeBoundaryEdgeMidpointsAdapter(const FDynamicM
 }
 
 
-FTriangleMeshAdapterd UE::Geometry::MakeTransformedDynamicMeshAdapter(const FDynamicMesh3* Mesh, FTransform Transform)
+FTriangleMeshAdapterd UE::Geometry::MakeTransformedDynamicMeshAdapter(const FDynamicMesh3* Mesh, FTransform3d Transform)
 {
 	FTriangleMeshAdapterd Adapter = MakeDynamicMeshAdapter(Mesh);
-	UE::Geometry::FTransform3d Transformd(Transform);
-	Adapter.GetVertex = [Mesh, Transformd](int32 Idx)
+	Adapter.GetVertex = [Mesh, Transform](int32 Idx)
 	{
-		return FVector3d(Transformd.TransformPosition(Mesh->GetVertex(Idx)));
+		return FVector3d(Transform.TransformPosition(Mesh->GetVertex(Idx)));
 	};
 	return Adapter;
 }
diff --git a/Engine/Source/Runtime/GeometryCore/Private/DynamicMesh/MeshTransforms.cpp b/Engine/Source/Runtime/GeometryCore/Private/DynamicMesh/MeshTransforms.cpp
--- a/Engine/Source/Runtime/GeometryCore/Private/DynamicMesh/MeshTransforms.cpp
+++ b/Engine/Source/Runtime/GeometryCore/Private/DynamicMesh/MeshTransforms.cpp
@@ -112,7 +112,7 @@ void MeshTransforms::FrameCoordsToWorld(FDynamicMesh3& Mesh, const FFrame3d& Fra
 
 
 
-void MeshTransforms::ApplyTransform(FDynamicMesh3& Mesh, const FTransform3d& Transform)
+void MeshTransforms::ApplyTransform(FDynamicMesh3& Mesh, const FTransformSRT3d& Transform)
 {
 	bool bVertexNormals = Mesh.HasVertexNormals();
 
@@ -143,7 +143,7 @@ void MeshTransforms::ApplyTransform(FDynamicMesh3& Mesh, const FTransform3d& Tra
 
 
 
-void MeshTransforms::ApplyTransformInverse(FDynamicMesh3& Mesh, const FTransform3d& Transform)
+void MeshTransforms::ApplyTransformInverse(FDynamicMesh3& Mesh, const FTransformSRT3d& Transform)
 {
 	bool bVertexNormals = Mesh.HasVertexNormals();
 
diff --git a/Engine/Source/Runtime/GeometryCore/Public/BoxTypes.h b/Engine/Source/Runtime/GeometryCore/Public/BoxTypes.h
--- a/Engine/Source/Runtime/GeometryCore/Public/BoxTypes.h
+++ b/Engine/Source/Runtime/GeometryCore/Public/BoxTypes.h
@@ -274,7 +274,7 @@ struct TAxisAlignedBox3
 		}
 	}
 
-	TAxisAlignedBox3(const TAxisAlignedBox3& Box, const FTransform3d& Transform)
+	TAxisAlignedBox3(const TAxisAlignedBox3& Box, const FTransformSRT3d& Transform)
 	{
 		TVector<RealType> C0 = Transform.TransformPosition(Box.GetCorner(0));
 		Min = C0;
diff --git a/Engine/Source/Runtime/GeometryCore/Public/DynamicMesh/MeshAdapterUtil.h b/Engine/Source/Runtime/GeometryCore/Public/DynamicMesh/MeshAdapterUtil.h
--- a/Engine/Source/Runtime/GeometryCore/Public/DynamicMesh/MeshAdapterUtil.h
+++ b/Engine/Source/Runtime/GeometryCore/Public/DynamicMesh/MeshAdapterUtil.h
@@ -18,18 +18,13 @@ namespace Geometry
 	/**
 	 * @return Transformed adapter of a FDynamicMesh
 	 */
-	FTriangleMeshAdapterd GEOMETRYCORE_API MakeTransformedDynamicMeshAdapter(const FDynamicMesh3* Mesh, FTransform Transform);
+	FTriangleMeshAdapterd GEOMETRYCORE_API MakeTransformedDynamicMeshAdapter(const FDynamicMesh3* Mesh, FTransform3d Transform);
 
 	/**
 	 * @return 1:1 adapter of a FDynamicMesh; can be used as a starting point to create other adapters
 	 */
 	FTriangleMeshAdapterd GEOMETRYCORE_API MakeDynamicMeshAdapter(const FDynamicMesh3* Mesh);
 
-	/**
-	* @return Transformed adapter of a FDynamicMesh
-	*/
-	FTriangleMeshAdapterd GEOMETRYCORE_API MakeTransformedDynamicMeshAdapter(const FDynamicMesh3* Mesh, FTransform Transform);
-
 	/**
 	 * @return Mesh vertices as a point set
 	 */
diff --git a/Engine/Source/Runtime/GeometryCore/Public/DynamicMesh/MeshTransforms.h b/Engine/Source/Runtime/GeometryCore/Public/DynamicMesh/MeshTransforms.h
--- a/Engine/Source/Runtime/GeometryCore/Public/DynamicMesh/MeshTransforms.h
+++ b/Engine/Source/Runtime/GeometryCore/Public/DynamicMesh/MeshTransforms.h
@@ -13,7 +13,6 @@
 namespace MeshTransforms
 {
 	using namespace UE::Geometry;
-	using UE::Geometry::FTransform3d;
 
 	/**
 	 * Apply Translation to vertex positions of Mesh. Does not modify any other attributes.
@@ -40,14 +39,14 @@ namespace MeshTransforms
 	 * Apply given Transform to a Mesh.
 	 * Modifies Vertex Positions and Normals, and any Per-Triangle Normal Overlays
 	 */
-	GEOMETRYCORE_API void ApplyTransform(FDynamicMesh3& Mesh, const FTransform3d& Transform);
+	GEOMETRYCORE_API void ApplyTransform(FDynamicMesh3& Mesh, const FTransformSRT3d& Transform);
 
 
 	/**
 	 * Apply inverse of given Transform to a Mesh.
 	 * Modifies Vertex Positions and Normals, and any Per-Triangle Normal Overlays
 	 */
-	GEOMETRYCORE_API void ApplyTransformInverse(FDynamicMesh3& Mesh, const FTransform3d& Transform);
+	GEOMETRYCORE_API void ApplyTransformInverse(FDynamicMesh3& Mesh, const FTransformSRT3d& Transform);
 
 
 	/**
diff --git a/Engine/Source/Runtime/GeometryCore/Public/FrameTypes.h b/Engine/Source/Runtime/GeometryCore/Public/FrameTypes.h
--- a/Engine/Source/Runtime/GeometryCore/Public/FrameTypes.h
+++ b/Engine/Source/Runtime/GeometryCore/Public/FrameTypes.h
@@ -177,9 +177,9 @@ struct TFrame3
 	}
 
 	/** @return conversion of this Frame to TTransform */
-	TTransform3<RealType> ToTransform() const
+	TTransformSRT3<RealType> ToTransform() const
 	{
-		return TTransform3<RealType>(Rotation, Origin);
+		return TTransformSRT3<RealType>(Rotation, Origin);
 	}
 
 	/** @return point at distances along frame axes */
@@ -350,7 +350,7 @@ struct TFrame3
 	/**
 	 * transform this frame by the given transform
 	 */
-	void Transform(const TTransform3<RealType>& XForm)
+	void Transform(const TTransformSRT3<RealType>& XForm)
 	{
 		Origin = XForm.TransformPosition(Origin);
 		Rotate(XForm.GetRotation());
diff --git a/Engine/Source/Runtime/GeometryCore/Public/Generators/LineSegmentGenerators.h b/Engine/Source/Runtime/GeometryCore/Public/Generators/LineSegmentGenerators.h
--- a/Engine/Source/Runtime/GeometryCore/Public/Generators/LineSegmentGenerators.h
+++ b/Engine/Source/Runtime/GeometryCore/Public/Generators/LineSegmentGenerators.h
@@ -29,7 +29,7 @@ namespace UE
 			const TVector<RealType>& AxisX, 
 			const TVector<RealType>& AxisY, 
 			const TVector<RealType>& AxisZ,
-			const TTransform3<RealType>& Transform,
+			const TTransformSRT3<RealType>& Transform,
 			TFunctionRef<void(const TVector<RealType>& A, const TVector<RealType>& B)> EmitLineFunc);
 
 		/**
@@ -48,7 +48,7 @@ namespace UE
 			const TVector<RealType>& Center,
 			const TVector<RealType>& AxisX,
 			const TVector<RealType>& AxisY,
-			const TTransform3<RealType>& Transform,
+			const TTransformSRT3<RealType>& Transform,
 			TFunctionRef<void(const TVector<RealType>& A, const TVector<RealType>& B)> EmitLineFunc);
 
 		/**
@@ -70,7 +70,7 @@ namespace UE
 			const TVector<RealType>& Center,
 			const TVector<RealType>& AxisX,
 			const TVector<RealType>& AxisY,
-			const TTransform3<RealType>& Transform,
+			const TTransformSRT3<RealType>& Transform,
 			TFunctionRef<void(const TVector<RealType>& A, const TVector<RealType>& B)> EmitLineFunc);
 	}
 }
@@ -86,7 +86,7 @@ void UE::Geometry::GenerateBoxSegments(
 	const TVector<RealType>& AxisX,
 	const TVector<RealType>& AxisY,
 	const TVector<RealType>& AxisZ,
-	const TTransform3<RealType>& Transform,
+	const TTransformSRT3<RealType>& Transform,
 	TFunctionRef<void(const TVector<RealType>& A, const TVector<RealType>& B)> EmitLineFunc)
 {
 	// B is box max/min, P and Q are used to store the start and endpoints of the segments we create
@@ -125,7 +125,7 @@ void UE::Geometry::GenerateCircleSegments(
 	const TVector<RealType>& Center,
 	const TVector<RealType>& AxisX,
 	const TVector<RealType>& AxisY,
-	const TTransform3<RealType>& Transform,
+	const TTransformSRT3<RealType>& Transform,
 	TFunctionRef<void(const TVector<RealType>& A, const TVector<RealType>& B)> EmitLineFunc)
 {
 	TVector<RealType> PrevPos = TVector<RealType>::Zero();
@@ -155,7 +155,7 @@ void UE::Geometry::GenerateArcSegments(
 	const TVector<RealType>& Center,
 	const TVector<RealType>& AxisX,
 	const TVector<RealType>& AxisY,
-	const TTransform3<RealType>& Transform,
+	const TTransformSRT3<RealType>& Transform,
 	TFunctionRef<void(const TVector<RealType>& A, const TVector<RealType>& B)> EmitLineFunc)
 {
 	TVector<RealType> PrevPos = TVector<RealType>::Zero();
diff --git a/Engine/Source/Runtime/GeometryCore/Public/MeshAdapterTransforms.h b/Engine/Source/Runtime/GeometryCore/Public/MeshAdapterTransforms.h
--- a/Engine/Source/Runtime/GeometryCore/Public/MeshAdapterTransforms.h
+++ b/Engine/Source/Runtime/GeometryCore/Public/MeshAdapterTransforms.h
@@ -22,7 +22,6 @@
 namespace MeshAdapterTransforms
 {
 	using namespace UE::Geometry;
-	using FTransform3d = UE::Geometry::FTransform3d;
 	using FFrame3d = UE::Geometry::FFrame3d;
 
 	/**
@@ -107,7 +106,7 @@ namespace MeshAdapterTransforms
 	* Modifies Vertex Positions and Normals, and any Per-Triangle Normal Overlays
 	*/
 	template<class TriangleMeshType>
-	void ApplyTransform(TriangleMeshType& Mesh, const FTransform3d& Transform)
+	void ApplyTransform(TriangleMeshType& Mesh, const FTransformSRT3d& Transform)
 	{
 		int NumVertices = Mesh.MaxVertexID();
 		ParallelFor(NumVertices, [&](int vid)
@@ -140,7 +139,7 @@ namespace MeshAdapterTransforms
 	* Modifies Vertex Positions and Normals, and any Per-Triangle Normal Overlays
 	*/
 	template<class TriangleMeshType>
-	void ApplyTransformInverse(TriangleMeshType& Mesh, const FTransform3d& Transform)
+	void ApplyTransformInverse(TriangleMeshType& Mesh, const FTransformSRT3d& Transform)
 	{
 		int NumVertices = Mesh.MaxVertexID();
 		ParallelFor(NumVertices, [&](int vid)
diff --git a/Engine/Source/Runtime/GeometryCore/Public/TransformSequence.h b/Engine/Source/Runtime/GeometryCore/Public/TransformSequence.h
--- a/Engine/Source/Runtime/GeometryCore/Public/TransformSequence.h
+++ b/Engine/Source/Runtime/GeometryCore/Public/TransformSequence.h
@@ -19,14 +19,14 @@ template<typename RealType>
 class TTransformSequence3
 {
 protected:
-	TArray<TTransform3<RealType>, TInlineAllocator<2>> Transforms;
+	TArray<TTransformSRT3<RealType>, TInlineAllocator<2>> Transforms;
 
 public:
 
 	/**
 	 * Add Transform to the end of the sequence, ie Seq(P) becomes NewTransform * Seq(P)
 	 */
-	void Append(const TTransform3<RealType>& Transform)
+	void Append(const TTransformSRT3<RealType>& Transform)
 	{
 		Transforms.Add(Transform);
 	}
@@ -36,7 +36,7 @@ public:
 	 */
 	void Append(const FTransform& Transform)
 	{
-		Transforms.Add( TTransform3<RealType>(Transform) );
+		Transforms.Add( TTransformSRT3<RealType>(Transform) );
 	}
 
 	/**
@@ -47,7 +47,7 @@ public:
 	/**
 	 * @return transforms in the sequence
 	 */
-	const TArray<TTransform3<RealType>,TInlineAllocator<2>>& GetTransforms() const { return Transforms; }
+	const TArray<TTransformSRT3<RealType>,TInlineAllocator<2>>& GetTransforms() const { return Transforms; }
 
 
 	/**
@@ -55,7 +55,7 @@ public:
 	 */
 	bool HasNonUniformScale(RealType Tolerance = TMathUtil<RealType>::ZeroTolerance) const
 	{
-		for (const TTransform3<RealType>& Transform : Transforms)
+		for (const TTransformSRT3<RealType>& Transform : Transforms)
 		{
 			if (Transform.HasNonUniformScale(Tolerance))
 			{
@@ -71,7 +71,7 @@ public:
 	TVector<RealType> GetAccumulatedScale() const
 	{
 		TVector<RealType> FinalScale = TVector<RealType>::One();
-		for (const TTransform3<RealType>& Transform : Transforms)
+		for (const TTransformSRT3<RealType>& Transform : Transforms)
 		{
 			FinalScale = FinalScale * Transform.GetScale();
 		}
@@ -83,7 +83,7 @@ public:
 	 */
 	void ClearScales()
 	{
-		for (TTransform3<RealType>& Transform : Transforms)
+		for (TTransformSRT3<RealType>& Transform : Transforms)
 		{
 			Transform.SetScale(TVector<RealType>::One());
 		}
@@ -94,7 +94,7 @@ public:
 	 */
 	TVector<RealType> TransformPosition(UE::Math::TVector<RealType> P) const
 	{
-		for (const TTransform3<RealType>& Transform : Transforms)
+		for (const TTransformSRT3<RealType>& Transform : Transforms)
 		{
 			P = Transform.TransformPosition(P);
 		}
@@ -119,7 +119,7 @@ public:
 	 */
 	TVector<RealType> TransformVector(UE::Math::TVector<RealType> V) const
 	{
-		for (const TTransform3<RealType>& Transform : Transforms)
+		for (const TTransformSRT3<RealType>& Transform : Transforms)
 		{
 			V = Transform.TransformVector(V);
 		}
@@ -132,7 +132,7 @@ public:
 	 */
 	TVector<RealType> TransformNormal(UE::Math::TVector<RealType> Normal) const
 	{
-		for (const TTransform3<RealType>& Transform : Transforms)
+		for (const TTransformSRT3<RealType>& Transform : Transforms)
 		{
 			Normal = Transform.TransformNormal(Normal);
 		}
diff --git a/Engine/Source/Runtime/GeometryCore/Public/TransformTypes.h b/Engine/Source/Runtime/GeometryCore/Public/TransformTypes.h
--- a/Engine/Source/Runtime/GeometryCore/Public/TransformTypes.h
+++ b/Engine/Source/Runtime/GeometryCore/Public/TransformTypes.h
@@ -14,10 +14,19 @@ namespace Geometry
 using namespace UE::Math;
 
 /**
- * TTransform3 is a double/float templated version of standard UE FTransform.
+ * TTransformSRT3 is a variant of the standard UE FTransform/TTransform that uses the 
+ * UE::Geometry::TQuaternion instead of the UE::Math::TQuat. 
+ * 
+ * Transform order is Scale, then Rotate, then Translate.
+ * So mathematically (T * R * S) * v , assuming traditional matrix-vector multiplication order
+ * (note that the UE::Math library uses the opposite vector*Matrix ordering for TMatrix!)
+ * 
+ * Implicit constructors/casts are defined to/from FTransform3d and FTransform3f.
+ * This will likely be revised in future so that conversions that lose precision are explicit.
+ * 
  */
 template<typename RealType>
-class TTransform3
+class TTransformSRT3
 {
 protected:
 	TQuaternion<RealType> Rotation;
@@ -26,35 +35,42 @@ protected:
 
 public:
 
-	TTransform3()
+	TTransformSRT3()
 	{
 		Rotation = TQuaternion<RealType>::Identity();
 		Translation = TVector<RealType>::Zero();
 		Scale3D = TVector<RealType>::One();
 	}
 
-	TTransform3(const TQuaternion<RealType>& RotationIn, const UE::Math::TVector<RealType>& TranslationIn, const UE::Math::TVector<RealType>& ScaleIn)
+	TTransformSRT3(const TQuaternion<RealType>& RotationIn, const UE::Math::TVector<RealType>& TranslationIn, const UE::Math::TVector<RealType>& ScaleIn)
 	{
 		Rotation = RotationIn;
 		Translation = TranslationIn;
 		Scale3D = ScaleIn;
 	}
 
-	explicit TTransform3(const TQuaternion<RealType>& RotationIn, const UE::Math::TVector<RealType>& TranslationIn)
+	explicit TTransformSRT3(const TQuaternion<RealType>& RotationIn, const UE::Math::TVector<RealType>& TranslationIn)
 	{
 		Rotation = RotationIn;
 		Translation = TranslationIn;
 		Scale3D = TVector<RealType>::One();
 	}
 
-	explicit TTransform3(const UE::Math::TVector<RealType>& TranslationIn)
+	explicit TTransformSRT3(const UE::Math::TVector<RealType>& TranslationIn)
 	{
 		Rotation = TQuaternion<RealType>::Identity();
 		Translation = TranslationIn;
 		Scale3D = TVector<RealType>::One();
 	}
 
-	explicit TTransform3(const FTransform& Transform)
+	TTransformSRT3(const FTransform3f& Transform)
+	{
+		Rotation = TQuaternion<RealType>(Transform.GetRotation());
+		Translation = TVector<RealType>(Transform.GetTranslation());
+		Scale3D = TVector<RealType>(Transform.GetScale3D());
+	}
+
+	TTransformSRT3(const FTransform3d& Transform)
 	{
 		Rotation = TQuaternion<RealType>(Transform.GetRotation());
 		Translation = TVector<RealType>(Transform.GetTranslation());
@@ -64,19 +80,26 @@ public:
 	/**
 	 * @return identity transform, IE no rotation, zero origin, unit scale
 	 */
-	static TTransform3<RealType> Identity()
+	static TTransformSRT3<RealType> Identity()
 	{
-		return TTransform3<RealType>(TQuaternion<RealType>::Identity(), TVector<RealType>::Zero(), TVector<RealType>::One());
+		return TTransformSRT3<RealType>(TQuaternion<RealType>::Identity(), TVector<RealType>::Zero(), TVector<RealType>::One());
 	}
 
 	/**
 	 * @return this transform converted to FTransform 
 	 */
-	explicit operator FTransform() const
+	operator FTransform3f() const
 	{
-		return FTransform((FQuat)Rotation, (FVector)Translation, (FVector)Scale3D);
+		return FTransform3f((FQuat4f)Rotation, (FVector3f)Translation, (FVector3f)Scale3D);
 	}
 
+	/**
+	* @return this transform converted to FTransform3d
+	*/
+	operator FTransform3d() const
+	{
+		return FTransform((FQuat4d)Rotation, (FVector3d)Translation, (FVector3d)Scale3D);
+	}
 
 	/**
 	 * @return Rotation portion of Transform, as Quaternion
@@ -147,12 +170,12 @@ public:
 	/**
 	 * @return inverse of this transform
 	 */
-	TTransform3<RealType> Inverse() const
+	TTransformSRT3<RealType> Inverse() const
 	{
 		TQuaternion<RealType> InvRotation = Rotation.Inverse();
 		TVector<RealType> InvScale3D = GetSafeScaleReciprocal(Scale3D);
 		TVector<RealType> InvTranslation = InvRotation * (InvScale3D * -Translation);
-		return TTransform3<RealType>(InvRotation, InvTranslation, InvScale3D);
+		return TTransformSRT3<RealType>(InvRotation, InvTranslation, InvScale3D);
 	}
 
 
@@ -350,8 +373,8 @@ public:
 	}
 
 };
-typedef TTransform3<float> FTransform3f;
-typedef TTransform3<double> FTransform3d;
+typedef TTransformSRT3<float> FTransformSRT3f;
+typedef TTransformSRT3<double> FTransformSRT3d;
 
 
 } // end namespace UE::Geometry
diff --git a/Engine/Source/Runtime/GeometryFramework/Private/Components/DynamicMeshComponent.cpp b/Engine/Source/Runtime/GeometryFramework/Private/Components/DynamicMeshComponent.cpp
--- a/Engine/Source/Runtime/GeometryFramework/Private/Components/DynamicMeshComponent.cpp
+++ b/Engine/Source/Runtime/GeometryFramework/Private/Components/DynamicMeshComponent.cpp
@@ -162,7 +162,7 @@ const FDynamicMesh3* UDynamicMeshComponent::GetRenderMesh() const
 
 
 
-void UDynamicMeshComponent::ApplyTransform(const UE::Geometry::FTransform3d& Transform, bool bInvert)
+void UDynamicMeshComponent::ApplyTransform(const FTransform3d& Transform, bool bInvert)
 {
 	MeshObject->EditMesh([&](FDynamicMesh3& EditMesh)
 	{
