ModelingTools: Add additional early-exit points for cancelled Baker computes.

#rnx

diff --git a/Engine/Plugins/Runtime/GeometryProcessing/Source/DynamicMesh/Private/Sampling/MeshMapBaker.cpp b/Engine/Plugins/Runtime/GeometryProcessing/Source/DynamicMesh/Private/Sampling/MeshMapBaker.cpp
--- a/Engine/Plugins/Runtime/GeometryProcessing/Source/DynamicMesh/Private/Sampling/MeshMapBaker.cpp
+++ b/Engine/Plugins/Runtime/GeometryProcessing/Source/DynamicMesh/Private/Sampling/MeshMapBaker.cpp
@@ -305,6 +305,11 @@ void FMeshMapBaker::Bake()
 	ParallelFor(NumTiles, [this, &Tiles, &GutterTexelsPerTile, &OutputQueue, &WriteToOutputBuffer, &OverwriteFn, &NoopFn, &WriteQueuedOutput](int32 TileIdx)
 	{
 		TRACE_CPUPROFILER_EVENT_SCOPE(FMeshMapBaker::Bake_EvalTile);
+
+		if (CancelF())
+		{
+			return;
+		}
 		
 		// Generate unpadded and padded tiles.
 		const FImageTile Tile = Tiles.GetTile(TileIdx);	// Image area to sample
@@ -399,6 +404,7 @@ void FMeshMapBaker::Bake()
 		WriteQueuedOutput(OutputQueue);
 	}
 
+	if (!CancelF())
 	{
 		FScopedDurationTimer WriteToImageTimer(BakeAnalytics.WriteToImageDuration);
 		
@@ -462,7 +468,7 @@ void FMeshMapBaker::Bake()
 	}
 
 	// Gutter Texel processing
-	if( bGutterEnabled )
+	if( bGutterEnabled && !CancelF())
 	{
 		FScopedDurationTimer WriteToGutterTimer(BakeAnalytics.WriteToGutterDuration);
 		
@@ -471,6 +477,11 @@ void FMeshMapBaker::Bake()
 		{
 			TRACE_CPUPROFILER_EVENT_SCOPE(FMeshMapBaker::Bake_WriteGutterPixels);
 
+			if (CancelF())
+			{
+				return;
+			}
+
 			const int NumGutter = GutterTexelsPerTile[TileIdx].Num();
 			for (int64 GutterIdx = 0; GutterIdx < NumGutter; ++GutterIdx)
 			{
