remove ensures from ConstrainedDelaunay2, fix comment to clarify when the previously-ensured-against condition can happen even for a robust CDT

diff --git a/Engine/Plugins/Runtime/GeometryProcessing/Source/GeometryAlgorithms/Private/ConstrainedDelaunay2.cpp b/Engine/Plugins/Runtime/GeometryProcessing/Source/GeometryAlgorithms/Private/ConstrainedDelaunay2.cpp
--- a/Engine/Plugins/Runtime/GeometryProcessing/Source/GeometryAlgorithms/Private/ConstrainedDelaunay2.cpp
+++ b/Engine/Plugins/Runtime/GeometryProcessing/Source/GeometryAlgorithms/Private/ConstrainedDelaunay2.cpp
@@ -142,7 +142,7 @@ namespace ConstrainedDelaunay2Internal
 					return AdjEdgeIdx;
 				}
 			}
-			ensure(false);
+			checkSlow(false);
 			return -1;
 		};
 		auto GetVertSubIdx = [&Indices](int VertID, int TriID)
@@ -154,7 +154,7 @@ namespace ConstrainedDelaunay2Internal
 					return VertSubIdx;
 				}
 			}
-			ensure(false);
+			checkSlow(false);
 			return -1;
 		};
 		auto Walk = [&Adj, &OtherEdgeOnTri](int VertID, int TriID, int EdgeSubIdx)
@@ -313,8 +313,7 @@ bool TConstrainedDelaunay2<RealType>::Triangulate(TFunctionRef<bool(const TArray
 			int B = Duplicates[Edge.B];
 			if (!Delaunay.Insert({{A, B}}, OutEdges))
 			{
-				// Note the failed edge; we will try to proceed anyway, just without this edge.  Hopefully the CDT is robust and this never happens!
-				ensureMsgf(false, TEXT("CDT edge insertion failed"));
+				// Note the failed edge; we will try to proceed anyway, just without this edge.  This can happen for example if the edge is exactly on top of another edge
 				bBoundaryTrackingFailure = true;
 			}
 		}
@@ -383,8 +382,7 @@ bool TConstrainedDelaunay2<RealType>::Triangulate()
 			int B = Duplicates[Edge.B];
 			if (!Delaunay.Insert({{A, B}}, OutEdges))
 			{
-				// Note the failed edge; we will try to proceed anyway, just without this edge.  Hopefully the CDT is robust and this never happens!
-				ensureMsgf(false, TEXT("CDT edge insertion failed"));
+				// Note the failed edge; we will try to proceed anyway, just without this edge.  This can happen for example if the edge is exactly on top of another edge
 				bBoundaryTrackingFailure = true;
 			}
 			else
