Added a flag to optionally recompute normal overlay or per-vertex normals after displacement in PNTriangles.
Disabled quadratically varying normal computation for 5.0.

#rnx

diff --git a/Engine/Plugins/Experimental/GeometryScripting/Source/GeometryScriptingCore/Private/MeshSubdivideFunctions.cpp b/Engine/Plugins/Experimental/GeometryScripting/Source/GeometryScriptingCore/Private/MeshSubdivideFunctions.cpp
--- a/Engine/Plugins/Experimental/GeometryScripting/Source/GeometryScriptingCore/Private/MeshSubdivideFunctions.cpp
+++ b/Engine/Plugins/Experimental/GeometryScripting/Source/GeometryScriptingCore/Private/MeshSubdivideFunctions.cpp
@@ -73,6 +73,7 @@ UDynamicMesh* UGeometryScriptLibrary_MeshSubdivideFunctions::ApplyPNTessellation
 	{
 		FPNTriangles Tessellator(&EditMesh);
 		Tessellator.TessellationLevel = TessellationLevel;
+		Tessellator.bRecalculateNormals = Options.bRecomputeNormals;
 		
 		if (Tessellator.Validate() != EOperationValidationResult::Ok)
 		{
@@ -85,14 +86,7 @@ UDynamicMesh* UGeometryScriptLibrary_MeshSubdivideFunctions::ApplyPNTessellation
 			UE::Geometry::AppendError(Debug, EGeometryScriptErrorType::OperationFailed, LOCTEXT("ApplyPNTessellation_Failed", "ApplyPNTessellation: Tessellation failed."));
 			return;
 		}
-
-		if (Options.bRecomputeNormals && EditMesh.HasAttributes() && EditMesh.Attributes()->PrimaryNormals() != nullptr)
-		{
-			FMeshNormals MeshNormals(&EditMesh);
-			MeshNormals.RecomputeOverlayNormals(EditMesh.Attributes()->PrimaryNormals(), true, true);
-			MeshNormals.CopyToOverlay(EditMesh.Attributes()->PrimaryNormals(), false);
-		}
-
+		
 	}, EDynamicMeshChangeType::GeneralEdit, EDynamicMeshAttributeChangeFlags::Unknown, false);
 
 	return TargetMesh;
diff --git a/Engine/Plugins/Runtime/GeometryProcessing/Source/DynamicMesh/Private/Operations/PNTriangles.cpp b/Engine/Plugins/Runtime/GeometryProcessing/Source/DynamicMesh/Private/Operations/PNTriangles.cpp
--- a/Engine/Plugins/Runtime/GeometryProcessing/Source/DynamicMesh/Private/Operations/PNTriangles.cpp
+++ b/Engine/Plugins/Runtime/GeometryProcessing/Source/DynamicMesh/Private/Operations/PNTriangles.cpp
@@ -293,12 +293,12 @@ namespace FPNTrianglesLocals
 	 * @return true if the operation succeeded, false if it failed or was canceled by the user.
 	 */
 
-	bool DisplaceAndSetNormals(const FDynamicMesh3& OriginalMesh,
-					           const FControlPoints& FControlPoints,
-					           const TArray<FIndex2i>& VerticesToDisplace,
-					           const bool bComputePNNormals,
-							   FProgressCancel* ProgressCancel,
-							   FDynamicMesh3& Mesh)
+	bool DisplaceAndSetQuadraticNormals(const FDynamicMesh3& OriginalMesh,
+					           			const FControlPoints& FControlPoints,
+					           			const TArray<FIndex2i>& VerticesToDisplace,
+					           			const bool bComputePNNormals,
+							   			FProgressCancel* ProgressCancel,
+							   			FDynamicMesh3& Mesh)
 	{
 		bool bHasVertexNormals = Mesh.HasVertexNormals();
 		bool bHasAttributes = Mesh.HasAttributes();
@@ -420,6 +420,13 @@ bool FPNTriangles::Compute()
 	
 	// Compute PN triangle control points for each flat triangle of the original mesh
 	FControlPoints FControlPoints;
+
+	// TODO: We disable quadratically varying normal computation until we fully support seam normal overlay edges.
+	// We might remove them completely since using quadratically varying normals would imply that we should be using 
+	// the same approach when computing normals inside the shaders for lighting computations 
+	// (see Tessellation on Any Budget, GDC, 2011). 
+	const bool bComputePNNormals = false;
+	
 	bool bOk = ComputeControlPoints(FControlPoints, *Mesh, UseNormals, bComputePNNormals, Progress);
 	if (bOk == false) 
 	{
@@ -435,12 +442,27 @@ bool FPNTriangles::Compute()
 		return false;
 	}
 
-	// Compute displacement and normals
-	bOk = DisplaceAndSetNormals(*Mesh, FControlPoints, NewVertices, bComputePNNormals, Progress, ResultMesh);
+	// Compute displacement and optionally quadratically varying normals
+	bOk = DisplaceAndSetQuadraticNormals(*Mesh, FControlPoints, NewVertices, bComputePNNormals, Progress, ResultMesh);
 	if (bOk == false) 
 	{
 		return false;
 	}
+				
+	if (bRecalculateNormals && bComputePNNormals == false)
+	{
+		if (ResultMesh.HasAttributes())
+		{
+			FMeshNormals NewNormals(&ResultMesh);
+			FDynamicMeshNormalOverlay* NormalOverlay = ResultMesh.Attributes()->PrimaryNormals();
+			NewNormals.RecomputeOverlayNormals(NormalOverlay);
+			NewNormals.CopyToOverlay(NormalOverlay);
+		}
+		else if (ResultMesh.HasVertexNormals())
+		{
+			FMeshNormals::QuickComputeVertexNormals(ResultMesh);
+		}
+	}
 
 	Mesh->Copy(ResultMesh);
 
diff --git a/Engine/Plugins/Runtime/GeometryProcessing/Source/DynamicMesh/Public/Operations/PNTriangles.h b/Engine/Plugins/Runtime/GeometryProcessing/Source/DynamicMesh/Public/Operations/PNTriangles.h
--- a/Engine/Plugins/Runtime/GeometryProcessing/Source/DynamicMesh/Public/Operations/PNTriangles.h
+++ b/Engine/Plugins/Runtime/GeometryProcessing/Source/DynamicMesh/Public/Operations/PNTriangles.h
@@ -14,7 +14,7 @@ namespace Geometry
 class FDynamicMesh3;
 
 /**
- * FPNTriangles implements curved PN (Point-Normal) Triangles - https://alex.vlachos.com/graphics/CurvedPNTriangles.pdf.
+ * FPNTriangles implements curved PN (Point-Normal) Triangles (Curved PN Triangles, Vlachos, 2001).
  * Each PN triangle replaces one original flat triangle by a curved shape that is retriangulated into 
  * a number of small subtriangles. The geometry of a PN triangle is defined as one cubic Bezier patch using control 
  * points. The patch matches the point and normal information at the vertices of the original flat triangle.
@@ -42,10 +42,11 @@ public:
 	int32 TessellationLevel = 1;
 
 	/**
-	 * If true, use the quadratically varying normal computation. 
-	 * Otherwise, if the vertex normals are enabled, use standard per-vertex normal computation.
+	 * If true and the mesh has the normal overlay then recalculate it. 
+	 * If true and the normal overlay doesn't exists, compute per-vertex normals (if they are enabled in the mesh).
+	 * If false, only perform the displacement.
 	 */
-	bool bComputePNNormals = true;
+	bool bRecalculateNormals = false;
 
 public:
 	FPNTriangles(FDynamicMesh3* Mesh) : Mesh(Mesh)
