Added the Power Method solver for computing largest/smallest eigenvalue and the corresponding eigenvector for the standard and general eigenvalue problems.

#rnx

diff --git a/Engine/Plugins/Runtime/GeometryProcessing/Source/DynamicMesh/Private/Solvers/Internal/MatrixBase.h b/Engine/Plugins/Runtime/GeometryProcessing/Source/DynamicMesh/Private/Solvers/Internal/MatrixBase.h
--- a/Engine/Plugins/Runtime/GeometryProcessing/Source/DynamicMesh/Private/Solvers/Internal/MatrixBase.h
+++ b/Engine/Plugins/Runtime/GeometryProcessing/Source/DynamicMesh/Private/Solvers/Internal/MatrixBase.h
@@ -12,8 +12,9 @@ enum class EMatrixSolverType
 	BICGSTAB /* Iterative Bi-conjugate gradient */,
 	PCG /* Pre-conditioned conjugate gradient  - requires symmetric positive def.*/,
 #ifndef EIGEN_MPL2_ONLY
-	LDLT /**Not included due to MPL2 */
+	LDLT /**Not included due to MPL2 */,
 #endif
+	FastestPSD /* Get the best available solver for the positive semidefinite matrix.*/,
 };
 
 static FString MatrixSolverName(const EMatrixSolverType SolverType)
@@ -34,10 +35,13 @@ static FString MatrixSolverName(const EMatrixSolverType SolverType)
 		String = FString(TEXT(" Iterative Preconditioned Conjugate Gradient "));
 		break;
 #ifndef EIGEN_MPL2_ONLY
-	case EMatrixSolverType::PCG:
+	case EMatrixSolverType::LDLT:
 		String = FString(TEXT(" Direct Cholesky "));
 		break;
 #endif
+	case EMatrixSolverType::FastestPSD:
+		String = FString(TEXT(" Fastest positive semidefinite "));
+		break; 
 	default:
 		check(0);
 	}
diff --git a/Engine/Plugins/Runtime/GeometryProcessing/Source/DynamicMesh/Private/Solvers/Internal/MatrixSolver.cpp b/Engine/Plugins/Runtime/GeometryProcessing/Source/DynamicMesh/Private/Solvers/Internal/MatrixSolver.cpp
--- a/Engine/Plugins/Runtime/GeometryProcessing/Source/DynamicMesh/Private/Solvers/Internal/MatrixSolver.cpp
+++ b/Engine/Plugins/Runtime/GeometryProcessing/Source/DynamicMesh/Private/Solvers/Internal/MatrixSolver.cpp
@@ -32,6 +32,13 @@ TUniquePtr<IMatrixSolverBase> ContructMatrixSolver(const EMatrixSolverType& Matr
 		ResultPtr.Reset(new FLDLTMatrixSolver());
 		break;
 #endif
+	case EMatrixSolverType::FastestPSD:
+		#ifndef EIGEN_MPL2_ONLY
+			ResultPtr.Reset(new FLDLTMatrixSolver());
+		#else
+			ResultPtr.Reset(new FLUMatrixSolver());
+		#endif
+		break;
 	}
 
 	return ResultPtr;
diff --git a/Engine/Plugins/Runtime/GeometryProcessing/Source/DynamicMesh/Private/Solvers/Internal/PowerMethodSolver.cpp b/Engine/Plugins/Runtime/GeometryProcessing/Source/DynamicMesh/Private/Solvers/Internal/PowerMethodSolver.cpp
new file mode 100644
--- /dev/null
+++ b/Engine/Plugins/Runtime/GeometryProcessing/Source/DynamicMesh/Private/Solvers/Internal/PowerMethodSolver.cpp
@@ -0,0 +1,241 @@
+#include "PowerMethodSolver.h"
+
+using namespace UE::Geometry;
+
+FPowerMethod::ESolveStatus FPowerMethod::Solve(ScalarType& OutEigenValue, RealVectorType& OutEigenVector, bool bComputeLargest)
+{
+    if (Verify(bComputeLargest) == false) 
+    {
+        return ESolveStatus::InvalidParameters;
+    }
+
+    // If needed, prefactorize matrices A or B to speed up the per-iteration computations
+    if (this->SetupSolver(bComputeLargest) == false) 
+    {
+        return ESolveStatus::NumericalIssue;
+    }
+
+    // Initial guess for the eigenvector
+    int32 Rows = OpMatrixA.Rows();
+    OutEigenVector = RealVectorType::Random(Rows, 1); 
+
+    bool bConverged = false;
+    for (IterationNum = 0; IterationNum < Parms.MaxIterations; ++IterationNum) 
+    {	
+        // Update the eigenvector with the next best guess
+        if (this->Iteration(OutEigenVector, bComputeLargest) == false) 
+        {
+            return ESolveStatus::NumericalIssue;
+        }
+
+        // Check that we didn't produce an invalid vector
+        ScalarType SquaredMagnitude = OutEigenVector.squaredNorm();
+        if (FMath::IsFinite(SquaredMagnitude) == false || FMath::IsNearlyZero(SquaredMagnitude))
+        {
+            checkSlow(false);
+            return ESolveStatus::NumericalIssue;
+        } 
+
+        // Normalize the vector to avoid the values getting too large
+        this->Normalize(OutEigenVector);
+
+        // Compute the eigenvalue based on the current eigenvector using Rayleigh Quotient formula
+        OutEigenValue = this->RayleighQuotient(OutEigenVector);  
+        if (FMath::IsFinite(OutEigenValue) == false) 
+        {
+            checkSlow(false);
+            return ESolveStatus::NumericalIssue;
+        }
+
+        if (this->Converged(OutEigenValue, OutEigenVector)) 
+        {   
+            bConverged = true;
+            break;
+        }
+    } 
+
+    if (bConverged) 
+    {
+        return ESolveStatus::Success;
+    }
+    else 
+    {
+        return ESolveStatus::NotConverged;
+    }
+}
+
+bool FPowerMethod::Converged(ScalarType Lambda, const RealVectorType& EVector)
+{	
+    RealVectorType AEVector;
+    OpMatrixA.Product(EVector, AEVector);
+
+    if (IsGeneralProblem()) 
+    {
+        RealVectorType BEVector;
+        OpMatrixB.GetValue().Product(EVector, BEVector);
+
+        IterationTol = (AEVector- Lambda * BEVector).array().abs().maxCoeff();		
+    }
+    else 
+    {
+        IterationTol = (AEVector - Lambda * EVector).array().abs().maxCoeff();
+    }
+
+    checkSlow(FMath::IsFinite(IterationTol)); 
+    return IterationTol < Parms.Tolerance;
+}
+
+FPowerMethod::ScalarType FPowerMethod::RayleighQuotient(const RealVectorType& EVector) const
+{
+    ScalarType Divisor; 
+    if(IsGeneralProblem()) 
+    { 
+        RealVectorType BEVector;
+        OpMatrixB.GetValue().Product(EVector, BEVector);
+        Divisor = EVector.dot(BEVector);
+    }
+    else 
+    { 
+        Divisor = EVector.dot(EVector);
+    }
+
+    if (FMath::IsNearlyZero(Divisor)) 
+    {
+        checkSlow(false);
+        Divisor = 1.0;
+    } 
+
+    RealVectorType AEVector;
+    OpMatrixA.Product(EVector, AEVector);
+    
+    return EVector.dot(AEVector)/Divisor;
+}
+
+bool FPowerMethod::SetupSolver(bool bComputeLargest)
+{
+    Solver = nullptr; // release previous solver
+
+    if (bComputeLargest && IsGeneralProblem()) 
+    {
+        Solver = OpMatrixB.GetValue().Factorize();
+        return Solver != nullptr;
+    } 
+    
+	if (bComputeLargest == false)
+    {
+        Solver = OpMatrixA.Factorize();
+        return Solver != nullptr;
+    }
+    
+    return true;
+}
+
+bool FPowerMethod::Iteration(RealVectorType& EVector, bool bComputeLargest) 
+{
+    if (IsGeneralProblem()) 
+    {
+        if (bComputeLargest)
+        {
+            OpMatrixA.Product(EVector, EVector);
+        }
+        else 
+        {
+            OpMatrixB.GetValue().Product(EVector, EVector);
+        }
+
+        // Solver already contains pre-factorization of the correct matrix
+        RealVectorType EVectorSolve;
+        Solver->Solve(EVector, EVectorSolve);
+        if (Solver->bSucceeded() == false) 
+        {
+            checkSlow(false);
+            return false;
+        }
+        
+        EVector = EVectorSolve;
+    }
+    else 
+    {
+        if (bComputeLargest) 
+        {
+            // x = A*x power iteration step 
+            OpMatrixA.Product(EVector, EVector); 
+        }
+        else 
+        {
+            // x = A^(-1)*x inverse power iteration step
+            RealVectorType EVectorSolve;
+            Solver->Solve(EVector, EVectorSolve);
+            if (Solver->bSucceeded() == false) 
+            {
+                checkSlow(false);
+                return false;
+            } 
+
+            EVector = EVectorSolve;
+        }
+    }
+   
+    return true;
+}
+
+bool FPowerMethod::Verify(bool bComputeLargest) const
+{	
+    if (OpMatrixA.Product == nullptr || OpMatrixA.Rows == nullptr) 
+    {
+        return false;
+    }
+
+    if (IsGeneralProblem()) 
+    {   
+        if (OpMatrixB.GetValue().Product == nullptr || (bComputeLargest && OpMatrixB.GetValue().Factorize == nullptr))
+        {
+            return false; 
+        }
+    }
+    else 
+    {
+        if (bComputeLargest == false && OpMatrixA.Factorize == nullptr) 
+        {
+            return false;
+        }
+    }
+	
+    if (Parms.MaxIterations <= 1 || Parms.Tolerance < 0) 
+    {
+        return false;
+    } 
+
+    return true;
+}
+
+void FSparsePowerMethod::CreateSparseMatrixOperator(const FSparseMatrixD& InMatrix, 
+                                                    const FMatrixHints& InMatrixHints,
+                                                    MatrixOperator& OutMatrixOp) 
+{
+    EMatrixSolverType MatrixSolverType = InMatrixHints.bIsPSD ? EMatrixSolverType::FastestPSD : EMatrixSolverType::LU;
+    bool bIsSymmetric = InMatrixHints.bIsSymmetric || InMatrixHints.bIsPSD;
+    
+    OutMatrixOp.Product = [InMatrix](const RealVectorType& InVector, RealVectorType& OutVector) 
+    {
+        OutVector = InMatrix * InVector; 
+    };
+
+    OutMatrixOp.Factorize = [InMatrix, MatrixSolverType, bIsSymmetric]() 
+    {	
+        TUniquePtr<IMatrixSolverBase> Solver = ContructMatrixSolver(MatrixSolverType);
+        Solver->SetUp(InMatrix, bIsSymmetric);
+
+        if (Solver->bSucceeded() == false) 
+        {
+            Solver = nullptr;
+        }
+
+        return Solver;
+    };
+
+    OutMatrixOp.Rows = [InMatrix]() 
+    {
+        return InMatrix.rows();
+    };
+}
\ No newline at end of file
diff --git a/Engine/Plugins/Runtime/GeometryProcessing/Source/DynamicMesh/Private/Solvers/Internal/PowerMethodSolver.h b/Engine/Plugins/Runtime/GeometryProcessing/Source/DynamicMesh/Private/Solvers/Internal/PowerMethodSolver.h
new file mode 100644
--- /dev/null
+++ b/Engine/Plugins/Runtime/GeometryProcessing/Source/DynamicMesh/Private/Solvers/Internal/PowerMethodSolver.h
@@ -0,0 +1,239 @@
+// Copyright Epic Games, Inc. All Rights Reserved.
+
+#pragma once 
+
+#include "CoreMinimal.h"
+#include "MatrixSolver.h"
+
+namespace UE
+{
+namespace Geometry
+{
+	
+/**
+ * Power iteration solver for computing the largest or smallest eigenvalue/eigenvector pair of a generalized 
+ * eigenvalue problem A*x = lambda*B*x. 
+ * 
+ * A, B are real square matrices. B is optional.
+ * 
+ * If the matrices A and B are not symmetric, this power method might fail due to the fact that the target eigenvalue 
+ * could be complex. Additionally, with symmetric matrices, the method may still fail to converge in the case that 
+ * the largest/smallest eigenvalue is not distinct.
+ * 
+ * For sparse matrices see FSparsePowerMethod.
+ */
+class DYNAMICMESH_API FPowerMethod
+{
+public:
+	using ScalarType = typename IMatrixSolverBase::ScalarType;
+	using RealVectorType = typename IMatrixSolverBase::RealVectorType;
+
+	// User-specified parameters
+	struct Parameters
+	{
+		ScalarType Tolerance = 1e-6; // Convergence parameter, see Converged()
+		int32 MaxIterations = 1000;  // Maximum number of times the power iteration is run 
+	};
+
+	// The result of the solve
+	enum class ESolveStatus 
+	{
+		Success,          // Solve was successful and converged within the maximum number of allowed iterations.
+
+		NotConverged,     // Solve was successful but did not converge to the user-specified tolerance within the 
+					      // maximum number of the iterations. In this case we can solve again by increasing the 
+					      // maximum number of iterations or increasing the tolerance.
+ 
+		NumericalIssue,   // Encountered a numerical issues. Most likely causes are a failed linear solve 
+						  // or encounterd a non-finite number during calculations.
+
+		InvalidParameters // User specified parameters are invalid or the matrix operators are not properly set.
+	};
+
+	// Instead of providing matrices we specify matrix operations via function handles
+	struct MatrixOperator
+	{ 
+		// Matrix vector product
+		TFunction<void(const RealVectorType& VecIn, RealVectorType& VecOut)> Product;
+		
+		// Creates the IMatrixSolverBase instance and calls SetUp to prefactorize the matrix which is used to compute 
+		// the matrix inverse vector product. 
+		// This is optional and is only needed when the problem requires it:
+		// 		Matrix A needs this when computing the smallest eigenpair.
+		// 		Matrix B needs this when computing the largest eigenpair.
+		TFunction<TUniquePtr<IMatrixSolverBase>()> Factorize;
+		
+		// The number of matrix rows. Can omit for B.
+		TFunction<int32()> Rows;
+	};
+
+	FPowerMethod(const Parameters& InParms) 
+	:
+	Parms(InParms)
+	{
+	}
+
+	FPowerMethod(const MatrixOperator& InMatrixA, const Parameters& InParms) 
+	:
+	OpMatrixA(InMatrixA), Parms(InParms)
+	{
+	}
+
+	FPowerMethod(const MatrixOperator& InMatrixA, const MatrixOperator& InMatrixB, const Parameters& InParms) 
+	:
+	OpMatrixA(InMatrixA), OpMatrixB(InMatrixB), Parms(InParms)
+	{
+	}
+
+	virtual ~FPowerMethod() 
+	{
+	}
+	
+	/**
+	 * Compute the eigen pair corresponding to the largest (smallest) eigenvalue.
+	 * 
+	 * @param OutEigenValue Computed largest (smallest) eigenvalue.
+	 * @param OutEigenVector Computed eigenvector corresponding to the largest (smallest) eigen value.
+	 * @param bComputeLargest If true (false) will compute the eigen pair corresponding to the largest (smallest) eigenvalue. 
+	 */
+	virtual ESolveStatus Solve(ScalarType& OutEigenValue, RealVectorType& OutEigenVector, bool bComputeLargest = false);
+
+	/** The tolerance value upon reaching convergence or upon hitting the max iteration. */
+	inline ScalarType GetConvergedTolerance() 
+	{
+		return IterationTol;
+	}
+
+	/** How many iterations did it take for the solver to converge. */
+	inline int32 GetConvergedIterations() 
+	{
+		return IterationNum;
+	}
+
+protected:
+	/** Checks if we reached the convergence condition ||A*x - lambda*B*x||_Inf < Tolerance */
+	virtual bool Converged(ScalarType Lambda, const RealVectorType& EVector);
+	
+	/** Compute the eigenvalue corresponding to the given eigenvector using the Rayleigh Quotient: lambda = x'*A*x/x'*B*x */
+	virtual ScalarType RayleighQuotient(const RealVectorType& NormalizedEVector) const;
+	
+	/** Normalize the input vector. Default is L2 normalization. */
+	virtual void Normalize(RealVectorType& EVector) const 
+	{
+		 EVector.normalize();
+	}
+
+	/** Single power iteration that computes the next approximation for the eigenvector.*/
+	virtual bool Iteration(RealVectorType& EVector, bool bComputeLargest);
+
+	/** 
+	 * In some cases we need to solve a linear system on every iteration. This method is called outside of the loop 
+	 * to pre-factorize any matrices used in the solve.
+	 */
+	virtual bool SetupSolver(bool bComputeLargest);
+
+	/** Check that the input user parameters are valid and matrix operators are set */
+	virtual bool Verify(bool bComputeLargest) const;
+
+	/** Is matrix B not an identity matrix*/
+	inline bool IsGeneralProblem() const 
+	{
+		return OpMatrixB.IsSet();
+	}
+
+public:
+	MatrixOperator OpMatrixA; // The A matrix in the A*x = lambda*B*x formula
+	TOptional<MatrixOperator> OpMatrixB; // The optional B matrix in the A*x = lambda*B*x formula
+
+protected:
+	TUniquePtr<IMatrixSolverBase> Solver; // Will either contain pre-factorization of A or B
+
+	Parameters Parms; // User specified parameters
+	
+	ScalarType IterationTol = 0.0;  // Current iteration tolerance
+	int32 IterationNum = -1; // Current iteration number
+};
+
+/**
+ * Convenience class that handles the power iteration solver where A and B are sparse matrices. 
+ * Allows passing FSparseMatrixD matrices which are used to construct the corresponding matrix operators with correct 
+ * solvers.
+ * 
+ * Example usage:
+ * 
+ *		FSparseMatrixD L = ... // setup PSD laplacian matrix
+ *	
+ *		// Make sure we use the best solver for the job
+ *		FMatrixHints LHints;
+ *		LHints.IsPSD = true;
+ *	
+ *		FSparsePowerMethod Solver(L, LHints);
+ *
+ *		RealVectorType OutVector;
+ *		ScalarType OutValue;
+ *		ESolveStatus Status = Solver.Solve(OutValue, OutVector, true); // compute the largest eigenvalue pair
+ *
+ *		if (Status == ESolveStatus::Success) 
+ *		{
+ *			// OutValue, OutVector is a valid eigenpair
+ *		}
+ */
+class DYNAMICMESH_API FSparsePowerMethod : public FPowerMethod
+{
+public:
+
+	// Hints to help choose the best solver for the problem
+	struct FMatrixHints 
+	{
+		FMatrixHints()
+		{}
+
+		FMatrixHints(bool bSetAll)
+		:
+		bIsSymmetric(bSetAll), bIsPSD(bSetAll)
+		{}
+
+		FMatrixHints(bool bIsSymmetric, bool bIsPSD) 
+		:
+	 	bIsSymmetric(bIsSymmetric), bIsPSD(bIsPSD)
+		{}
+
+		bool bIsSymmetric = false;
+		bool bIsPSD = false;
+	};
+
+	FSparsePowerMethod(const FSparseMatrixD& InMatrixA,
+					   const FMatrixHints& InHintsA = FMatrixHints(), 
+					   const Parameters& InParms = Parameters()) 
+	:
+	FPowerMethod(InParms)
+	{
+		CreateSparseMatrixOperator(InMatrixA, InHintsA, OpMatrixA);
+	}
+
+	FSparsePowerMethod(const FSparseMatrixD& InMatrixA, 
+					   const FSparseMatrixD& InMatrixB, 
+					   const FMatrixHints& InHintsA = FMatrixHints(), 
+					   const FMatrixHints& InHintsB = FMatrixHints(), 
+					   const Parameters& InParms = Parameters()) 
+	:
+	FPowerMethod(InParms)
+	{
+		CreateSparseMatrixOperator(InMatrixA, InHintsA, OpMatrixA);
+		
+		OpMatrixB = MatrixOperator();
+		CreateSparseMatrixOperator(InMatrixB, InHintsB, OpMatrixB.GetValue());
+	}
+
+	virtual ~FSparsePowerMethod()
+	{
+	}
+
+	/** Helper method to setup Matrix Operator for the sparse real matrix */
+	static void CreateSparseMatrixOperator(const FSparseMatrixD& InMatrix, 
+										   const FMatrixHints& InHints, 
+										   MatrixOperator& OutMatrixOp);
+};
+
+} // end namespace UE::Geometry
+} // end namespace UE
\ No newline at end of file
