GeometryProcessing: Remove default values for FMeshConnectedComponents::InitializeFromTriangleComponents()

#rnx

diff --git a/Engine/Plugins/Runtime/GeometryProcessing/Source/DynamicMesh/Private/Parameterization/PatchBasedMeshUVGenerator.cpp b/Engine/Plugins/Runtime/GeometryProcessing/Source/DynamicMesh/Private/Parameterization/PatchBasedMeshUVGenerator.cpp
--- a/Engine/Plugins/Runtime/GeometryProcessing/Source/DynamicMesh/Private/Parameterization/PatchBasedMeshUVGenerator.cpp
+++ b/Engine/Plugins/Runtime/GeometryProcessing/Source/DynamicMesh/Private/Parameterization/PatchBasedMeshUVGenerator.cpp
@@ -147,7 +147,7 @@ bool FPatchBasedMeshUVGenerator::ComputeInitialMeshPatches(
 	{
 		return false;
 	}
-	InitialComponentsOut.InitializeFromTriangleComponents(Generator.FoundPolygroups, true);
+	InitialComponentsOut.InitializeFromTriangleComponents(Generator.FoundPolygroups, true, false);
 
 	return true;
 }
diff --git a/Engine/Plugins/Runtime/GeometryProcessing/Source/DynamicMesh/Public/Selections/MeshConnectedComponents.h b/Engine/Plugins/Runtime/GeometryProcessing/Source/DynamicMesh/Public/Selections/MeshConnectedComponents.h
--- a/Engine/Plugins/Runtime/GeometryProcessing/Source/DynamicMesh/Public/Selections/MeshConnectedComponents.h
+++ b/Engine/Plugins/Runtime/GeometryProcessing/Source/DynamicMesh/Public/Selections/MeshConnectedComponents.h
@@ -86,7 +86,7 @@ public:
 	 * @param bValidateIDs if true, test that each value corresponds to a valid triangle ID on the Mesh
 	 * @return true if all IDs are valid, or if check was skipped
 	 */
-	bool InitializeFromTriangleComponents(const TArray<TArray<int32>>& ComponentLists, bool bValidateIDs = false);
+	bool InitializeFromTriangleComponents(const TArray<TArray<int32>>& ComponentLists, bool bValidateIDs);
 
 	/**
 	* Initialize the internal FComponent list from the input ComponentLists, skipping any empty input lists
@@ -94,7 +94,7 @@ public:
 	* @param bValidateIDs if true, test that each value corresponds to a valid triangle ID on the Mesh
 	* @return true if all IDs are valid, or if check was skipped
 	*/
-	bool InitializeFromTriangleComponents(TArray<TArray<int32>>& ComponentLists, bool bMoveSubLists, bool bValidateIDs = false);
+	bool InitializeFromTriangleComponents(TArray<TArray<int32>>& ComponentLists, bool bMoveSubLists, bool bValidateIDs);
 
 	//
 	// Query functions. Only valid to call after a Calculation function has been called.
