ModelingTools: Add support for Object space source normal maps to BakeTexture.

diff --git a/Engine/Plugins/Experimental/GeometryScripting/Source/GeometryScriptingCore/Private/MeshBakeFunctions.cpp b/Engine/Plugins/Experimental/GeometryScripting/Source/GeometryScriptingCore/Private/MeshBakeFunctions.cpp
--- a/Engine/Plugins/Experimental/GeometryScripting/Source/GeometryScriptingCore/Private/MeshBakeFunctions.cpp
+++ b/Engine/Plugins/Experimental/GeometryScripting/Source/GeometryScriptingCore/Private/MeshBakeFunctions.cpp
@@ -448,10 +448,11 @@ namespace GeometryScriptBakeLocals
 			}
 			else
 			{
-				DetailSampler.SetNormalMap(SourceMeshToUse,
-				IMeshBakerDetailSampler::FBakeDetailTexture(
+				DetailSampler.SetNormalTextureMap(SourceMeshToUse,
+				IMeshBakerDetailSampler::FBakeDetailNormalTexture(
 						SourceNormalMap.Get(),
-						SourceOptions.SourceNormalUVLayer));
+						SourceOptions.SourceNormalUVLayer,
+						SourceOptions.SourceNormalSpace == EGeometryScriptBakeNormalSpace::Tangent ? IMeshBakerDetailSampler::EBakeDetailNormalSpace::Tangent : IMeshBakerDetailSampler::EBakeDetailNormalSpace::Object));
 
 				if (!GetMeshTangents(SourceMesh, SourceMeshTangents))
 				{
diff --git a/Engine/Plugins/Experimental/GeometryScripting/Source/GeometryScriptingCore/Public/GeometryScript/MeshBakeFunctions.h b/Engine/Plugins/Experimental/GeometryScripting/Source/GeometryScriptingCore/Public/GeometryScript/MeshBakeFunctions.h
--- a/Engine/Plugins/Experimental/GeometryScripting/Source/GeometryScriptingCore/Public/GeometryScript/MeshBakeFunctions.h
+++ b/Engine/Plugins/Experimental/GeometryScripting/Source/GeometryScriptingCore/Public/GeometryScript/MeshBakeFunctions.h
@@ -69,6 +69,15 @@ enum class EGeometryScriptBakeTypes : uint8
 	MaterialID             UMETA(DisplayName = "Material ID"),
 };
 
+UENUM(BlueprintType)
+enum class EGeometryScriptBakeNormalSpace : uint8
+{
+	/* Tangent space */
+	Tangent,
+	/* Object space */
+	Object
+};
+
 USTRUCT(BlueprintType)
 struct GEOMETRYSCRIPTINGCORE_API FGeometryScriptBakeTypes
 {
@@ -250,6 +259,9 @@ struct GEOMETRYSCRIPTINGCORE_API FGeometryScriptBakeSourceMeshOptions
 
 	UPROPERTY(BlueprintReadWrite, Category = Options, meta = (DisplayName="Source Normal UV Channel"))
 	int SourceNormalUVLayer = 0;
+
+	UPROPERTY(BlueprintReadWrite, Category = Options)
+	EGeometryScriptBakeNormalSpace SourceNormalSpace = EGeometryScriptBakeNormalSpace::Tangent;
 };
 
 USTRUCT(BlueprintType)
diff --git a/Engine/Plugins/Experimental/MeshModelingToolsetExp/Source/MeshModelingToolsExp/Private/BakeMeshAttributeMapsTool.cpp b/Engine/Plugins/Experimental/MeshModelingToolsetExp/Source/MeshModelingToolsExp/Private/BakeMeshAttributeMapsTool.cpp
--- a/Engine/Plugins/Experimental/MeshModelingToolsetExp/Source/MeshModelingToolsExp/Private/BakeMeshAttributeMapsTool.cpp
+++ b/Engine/Plugins/Experimental/MeshModelingToolsetExp/Source/MeshModelingToolsExp/Private/BakeMeshAttributeMapsTool.cpp
@@ -116,6 +116,7 @@ public:
 	// NormalMap settings
 	ImagePtr DetailMeshNormalMap;
 	int32 DetailMeshNormalUVLayer = 0;
+	IMeshBakerDetailSampler::EBakeDetailNormalSpace DetailMeshNormalSpace = IMeshBakerDetailSampler::EBakeDetailNormalSpace::Tangent; 
 
 	// Texture2DImage & MultiTexture settings
 	ImagePtr TextureImage;
@@ -164,7 +165,7 @@ public:
 			case EBakeMapType::TangentSpaceNormal:
 			{
 				TSharedPtr<FMeshNormalMapEvaluator, ESPMode::ThreadSafe> NormalEval = MakeShared<FMeshNormalMapEvaluator, ESPMode::ThreadSafe>();
-				DetailSampler.SetNormalMap(DetailMesh.Get(), IMeshBakerDetailSampler::FBakeDetailTexture(DetailMeshNormalMap.Get(), DetailMeshNormalUVLayer));
+				DetailSampler.SetNormalTextureMap(DetailMesh.Get(), IMeshBakerDetailSampler::FBakeDetailNormalTexture(DetailMeshNormalMap.Get(), DetailMeshNormalUVLayer, DetailMeshNormalSpace));
 				Baker->AddEvaluator(NormalEval);
 				break;
 			}
@@ -195,7 +196,7 @@ public:
 			{
 				TSharedPtr<FMeshPropertyMapEvaluator, ESPMode::ThreadSafe> PropertyEval = MakeShared<FMeshPropertyMapEvaluator, ESPMode::ThreadSafe>();
 				PropertyEval->Property = EMeshPropertyMapType::Normal;
-				DetailSampler.SetNormalMap(DetailMesh.Get(), IMeshBakerDetailSampler::FBakeDetailTexture(DetailMeshNormalMap.Get(), DetailMeshNormalUVLayer));
+				DetailSampler.SetNormalTextureMap(DetailMesh.Get(), IMeshBakerDetailSampler::FBakeDetailNormalTexture(DetailMeshNormalMap.Get(), DetailMeshNormalUVLayer, DetailMeshNormalSpace));
 				Baker->AddEvaluator(PropertyEval);
 				break;
 			}
@@ -311,7 +312,15 @@ void UBakeMeshAttributeMapsTool::Setup()
 	InputMeshSettings->WatchProperty(InputMeshSettings->SourceNormalMap, [this](UTexture2D*)
 	{
 		// Only invalidate detail mesh if we need to recompute tangents.
-		if (!DetailMeshTangents)
+		if (!DetailMeshTangents && InputMeshSettings->SourceNormalSpace == EBakeNormalSpace::Tangent)
+		{
+			OpState |= EBakeOpState::EvaluateDetailMesh;
+		}
+		OpState |= EBakeOpState::Evaluate;
+	});
+	InputMeshSettings->WatchProperty(InputMeshSettings->SourceNormalSpace, [this](EBakeNormalSpace Space)
+	{
+		if (!DetailMeshTangents && Space == EBakeNormalSpace::Tangent)
 		{
 			OpState |= EBakeOpState::EvaluateDetailMesh;
 		}
@@ -416,6 +425,8 @@ TUniquePtr<UE::Geometry::TGenericDataOperator<FMeshMapBaker>> UBakeMeshAttribute
 		Op->DetailMeshTangents = DetailMeshTangents;
 		Op->DetailMeshNormalMap = CachedDetailNormalMap;
 		Op->DetailMeshNormalUVLayer = CachedDetailMeshSettings.UVLayer;
+		Op->DetailMeshNormalSpace = CachedDetailMeshSettings.NormalSpace == EBakeNormalSpace::Tangent ?
+			IMeshBakerDetailSampler::EBakeDetailNormalSpace::Tangent : IMeshBakerDetailSampler::EBakeDetailNormalSpace::Object;
 	}
 
 	if ((bool)(CachedBakeSettings.BakeMapTypes & EBakeMapType::TangentSpaceNormal))
@@ -691,6 +702,7 @@ EBakeOpState UBakeMeshAttributeMapsTool::UpdateResult_DetailNormalMap()
 
 	FDetailMeshSettings DetailMeshSettings;
 	DetailMeshSettings.UVLayer = DetailUVLayer;
+	DetailMeshSettings.NormalSpace = InputMeshSettings->SourceNormalSpace;
 
 	if (!(CachedDetailMeshSettings == DetailMeshSettings))
 	{
diff --git a/Engine/Plugins/Experimental/MeshModelingToolsetExp/Source/MeshModelingToolsExp/Private/BakeMultiMeshAttributeMapsTool.cpp b/Engine/Plugins/Experimental/MeshModelingToolsetExp/Source/MeshModelingToolsExp/Private/BakeMultiMeshAttributeMapsTool.cpp
--- a/Engine/Plugins/Experimental/MeshModelingToolsetExp/Source/MeshModelingToolsExp/Private/BakeMultiMeshAttributeMapsTool.cpp
+++ b/Engine/Plugins/Experimental/MeshModelingToolsetExp/Source/MeshModelingToolsExp/Private/BakeMultiMeshAttributeMapsTool.cpp
@@ -66,6 +66,7 @@ class FMeshBakerMeshSceneSampler : public IMeshBakerDetailSampler
 {
 public:
 	using FDetailTextureMap = TMap<void*, FBakeDetailTexture>;
+	using FDetailNormalMap = TMap<void*, FBakeDetailNormalTexture>;
 
 public:
 	/**
@@ -103,6 +104,11 @@ public:
 	}
 
 	virtual void SetNormalMap(const void* Mesh, const FBakeDetailTexture& Map) override
+	{
+		DetailNormalMaps[Mesh] = FBakeDetailNormalTexture(Map.Key, Map.Value, EBakeDetailNormalSpace::Tangent);
+	}
+
+	virtual void SetNormalTextureMap(const void* Mesh, const FBakeDetailNormalTexture& Map) override
 	{
 		DetailNormalMaps[Mesh] = Map;
 	}
@@ -113,6 +119,11 @@ public:
 	}
 	
 	virtual const FBakeDetailTexture* GetNormalMap(const void* Mesh) const override
+	{
+		return nullptr;
+	}
+
+	virtual const FBakeDetailNormalTexture* GetNormalTextureMap(const void* Mesh) const override
 	{
 		return DetailNormalMaps.Find(Mesh);
 	}
@@ -294,7 +305,7 @@ public:
 	}
 
 	/** Initialize the mesh to normal textures map */
-	void SetNormalMaps(const FDetailTextureMap& Map)
+	void SetNormalMaps(const FDetailNormalMap& Map)
 	{
 		DetailNormalMaps = Map;
 	}
@@ -302,7 +313,7 @@ public:
 protected:
 	FMeshSceneAdapter* MeshScene = nullptr;
 	FDetailTextureMap DetailTextureMaps;
-	FDetailTextureMap DetailNormalMaps;
+	FDetailNormalMap DetailNormalMaps;
 };
 
 
diff --git a/Engine/Plugins/Experimental/MeshModelingToolsetExp/Source/MeshModelingToolsExp/Public/BakeMeshAttributeToolCommon.h b/Engine/Plugins/Experimental/MeshModelingToolsetExp/Source/MeshModelingToolsExp/Public/BakeMeshAttributeToolCommon.h
--- a/Engine/Plugins/Experimental/MeshModelingToolsetExp/Source/MeshModelingToolsExp/Public/BakeMeshAttributeToolCommon.h
+++ b/Engine/Plugins/Experimental/MeshModelingToolsetExp/Source/MeshModelingToolsExp/Public/BakeMeshAttributeToolCommon.h
@@ -21,6 +21,16 @@ class USkeletalMesh;
  */
 
 
+UENUM()
+enum class EBakeNormalSpace
+{
+	/** Tangent space */
+	Tangent,
+	/** Object space */
+	Object
+};
+
+
 UCLASS()
 class MESHMODELINGTOOLSEXP_API UBakeInputMeshProperties : public UInteractiveToolPropertySet
 {
@@ -82,6 +92,11 @@ public:
 		EditCondition = "bHasSourceNormalMap == true", EditConditionHides, HideEditConditionToggle))
 	FString SourceNormalMapUVLayer;
 
+	/** Normal space of the source mesh normal map. */
+	UPROPERTY(EditAnywhere, Category = BakeInput, AdvancedDisplay, meta = (TransientToolProperty,
+		EditCondition = "bHasSourceNormalMap == true", EditConditionHides, HideEditConditionToggle))
+	EBakeNormalSpace SourceNormalSpace = EBakeNormalSpace::Tangent;
+
 	/** If true, expose the SourceNormalMap and SourceNormalMapUVLayer properties */
 	UPROPERTY()
 	bool bHasSourceNormalMap = false;
@@ -300,10 +315,11 @@ public:
 struct FDetailMeshSettings
 {
 	int32 UVLayer = 0;
+	EBakeNormalSpace NormalSpace = EBakeNormalSpace::Tangent;
 	
 	bool operator==(const FDetailMeshSettings& Other) const
 	{
-		return UVLayer == Other.UVLayer;
+		return UVLayer == Other.UVLayer && NormalSpace == Other.NormalSpace;
 	}
 };
 
diff --git a/Engine/Plugins/Runtime/GeometryProcessing/Source/DynamicMesh/Private/Sampling/MeshNormalMapEvaluator.cpp b/Engine/Plugins/Runtime/GeometryProcessing/Source/DynamicMesh/Private/Sampling/MeshNormalMapEvaluator.cpp
--- a/Engine/Plugins/Runtime/GeometryProcessing/Source/DynamicMesh/Private/Sampling/MeshNormalMapEvaluator.cpp
+++ b/Engine/Plugins/Runtime/GeometryProcessing/Source/DynamicMesh/Private/Sampling/MeshNormalMapEvaluator.cpp
@@ -12,14 +12,17 @@ void FMeshNormalMapEvaluator::Setup(const FMeshBaseBaker& Baker, FEvaluationCont
 	DetailSampler = Baker.GetDetailSampler();
 	auto GetDetailNormalMaps = [this](const void* Mesh)
 	{
-		const FDetailNormalTexture* NormalMap = DetailSampler->GetNormalMap(Mesh);
+		const FNormalTexture* NormalMap = DetailSampler->GetNormalTextureMap(Mesh);
 		if (NormalMap)
 		{
-			// Require valid normal map, UV layer and tangents to enable normal map transfer.
-			const bool bEnableNormalMapTransfer = DetailSampler->HasUVs(Mesh, NormalMap->Get<1>()) && DetailSampler->HasTangents(Mesh);
+			// Require valid normal map and UV layer to enable normal map transfer.
+			// Tangents also required if the map is in tangent space.
+			const bool bDetailNormalTangentSpace = NormalMap->Get<2>() == IMeshBakerDetailSampler::EBakeDetailNormalSpace::Tangent; 
+			const bool bEnableNormalMapTransfer = DetailSampler->HasUVs(Mesh, NormalMap->Get<1>()) &&
+				(!bDetailNormalTangentSpace || DetailSampler->HasTangents(Mesh));
 			if (bEnableNormalMapTransfer)
 			{
-				DetailNormalTextures.Add(Mesh, *NormalMap);
+				DetailNormalMaps.Add(Mesh, *NormalMap);
 				bHasDetailNormalTextures = true;
 			}
 		}
@@ -86,10 +89,11 @@ FVector3f FMeshNormalMapEvaluator::SampleFunction(const FCorrespondenceSample& S
 	{
 		const TImageBuilder<FVector4f>* DetailNormalMap = nullptr;
 		int DetailNormalUVLayer = 0;
-		const FDetailNormalTexture* DetailNormalTexture = DetailNormalTextures.Find(DetailMesh);
+		IMeshBakerDetailSampler::EBakeDetailNormalSpace DetailNormalSpace = IMeshBakerDetailSampler::EBakeDetailNormalSpace::Tangent;
+		const FNormalTexture* DetailNormalTexture = DetailNormalMaps.Find(DetailMesh);
 		if (DetailNormalTexture)
 		{
-			Tie(DetailNormalMap, DetailNormalUVLayer) = *DetailNormalTexture;
+			Tie(DetailNormalMap, DetailNormalUVLayer, DetailNormalSpace) = *DetailNormalTexture;
 		}
 
 		if (DetailNormalMap)
@@ -112,10 +116,15 @@ FVector3f FMeshNormalMapEvaluator::SampleFunction(const FCorrespondenceSample& S
 
 			// Map color space [0,1] to normal space [-1,1]
 			const FVector3f DetailNormalColor(DetailNormalColor4.X, DetailNormalColor4.Y, DetailNormalColor4.Z);
-			const FVector3f DetailNormalTangentSpace = (DetailNormalColor * 2.0f) - FVector3f::One();
-
-			// Convert detail normal tangent space to object space
-			FVector3f DetailNormalObjectSpace = DetailNormalTangentSpace.X * FVector3f(DetailTangentX) + DetailNormalTangentSpace.Y * FVector3f(DetailTangentY) + DetailNormalTangentSpace.Z * DetailNormal;
+			FVector3f DetailNormalObjectSpace = (DetailNormalColor * 2.0f) - FVector3f::One();
+			// Ideally this branch could be made compile time. Unfortunately since each mesh could have its
+			// own source normal map each with their own normal space, this branch must be resolved at runtime.
+			if (DetailNormalSpace == IMeshBakerDetailSampler::EBakeDetailNormalSpace::Tangent)
+			{
+				// Convert detail normal tangent space to object space
+				const FVector3f DetailNormalTangentSpace = DetailNormalObjectSpace;
+				DetailNormalObjectSpace = DetailNormalTangentSpace.X * FVector3f(DetailTangentX) + DetailNormalTangentSpace.Y * FVector3f(DetailTangentY) + DetailNormalTangentSpace.Z * DetailNormal;
+			}
 			Normalize(DetailNormalObjectSpace);
 			DetailNormal = DetailNormalObjectSpace;
 		}
diff --git a/Engine/Plugins/Runtime/GeometryProcessing/Source/DynamicMesh/Private/Sampling/MeshPropertyMapEvaluator.cpp b/Engine/Plugins/Runtime/GeometryProcessing/Source/DynamicMesh/Private/Sampling/MeshPropertyMapEvaluator.cpp
--- a/Engine/Plugins/Runtime/GeometryProcessing/Source/DynamicMesh/Private/Sampling/MeshPropertyMapEvaluator.cpp
+++ b/Engine/Plugins/Runtime/GeometryProcessing/Source/DynamicMesh/Private/Sampling/MeshPropertyMapEvaluator.cpp
@@ -12,14 +12,17 @@ void FMeshPropertyMapEvaluator::Setup(const FMeshBaseBaker& Baker, FEvaluationCo
 	DetailSampler = Baker.GetDetailSampler();
 	auto GetDetailNormalMaps = [this](const void* Mesh)
 	{
-		const FDetailNormalTexture* NormalMap = DetailSampler->GetNormalMap(Mesh);
+		const FNormalTexture* NormalMap = DetailSampler->GetNormalTextureMap(Mesh);
 		if (NormalMap)
 		{
-			// Require valid normal map, UV layer and tangents to enable normal map transfer.
-			const bool bEnableNormalMapTransfer = DetailSampler->HasUVs(Mesh, NormalMap->Get<1>()) && DetailSampler->HasTangents(Mesh);
+			// Require valid normal map and UV layer to enable normal map transfer.
+			// Tangents also required if the map is in tangent space.
+			const bool bDetailNormalTangentSpace = NormalMap->Get<2>() == IMeshBakerDetailSampler::EBakeDetailNormalSpace::Tangent; 
+			const bool bEnableNormalMapTransfer = DetailSampler->HasUVs(Mesh, NormalMap->Get<1>()) &&
+				(!bDetailNormalTangentSpace || DetailSampler->HasTangents(Mesh));
 			if (bEnableNormalMapTransfer)
 			{
-				DetailNormalTextures.Add(Mesh, *NormalMap);
+				DetailNormalMaps.Add(Mesh, *NormalMap);
 				bHasDetailNormalTextures = true;
 			}
 		}
@@ -120,10 +123,11 @@ FVector3f FMeshPropertyMapEvaluator::SampleFunction(const FCorrespondenceSample&
 			{
 				const TImageBuilder<FVector4f>* DetailNormalMap = nullptr;
 				int DetailNormalUVLayer = 0;
-				const FDetailNormalTexture* DetailNormalTexture = DetailNormalTextures.Find(DetailMesh);
+				IMeshBakerDetailSampler::EBakeDetailNormalSpace DetailNormalSpace = IMeshBakerDetailSampler::EBakeDetailNormalSpace::Tangent;
+				const FNormalTexture* DetailNormalTexture = DetailNormalMaps.Find(DetailMesh);
 				if (DetailNormalTexture)
 				{
-					Tie(DetailNormalMap, DetailNormalUVLayer) = *DetailNormalTexture;
+					Tie(DetailNormalMap, DetailNormalUVLayer, DetailNormalSpace) = *DetailNormalTexture;
 				}
 
 				if (DetailNormalMap)
@@ -141,10 +145,15 @@ FVector3f FMeshPropertyMapEvaluator::SampleFunction(const FCorrespondenceSample&
 
 					// Map color space [0,1] to normal space [-1,1]
 					const FVector3f DetailNormalColor(DetailNormalColor4.X, DetailNormalColor4.Y, DetailNormalColor4.Z);
-					const FVector3f DetailNormalTangentSpace = (DetailNormalColor * 2.0f) - FVector3f::One();
-
-					// Convert detail normal tangent space to object space
-					FVector3f DetailNormalObjectSpace = DetailNormalTangentSpace.X * FVector3f(DetailTangentX) + DetailNormalTangentSpace.Y * FVector3f(DetailTangentY) + DetailNormalTangentSpace.Z * DetailNormal;
+					FVector3f DetailNormalObjectSpace = (DetailNormalColor * 2.0f) - FVector3f::One();
+					// Ideally this branch could be made compile time. Unfortunately since each mesh could have its
+					// own source normal map each with their own normal space, this branch must be resolved at runtime.
+					if (DetailNormalSpace == IMeshBakerDetailSampler::EBakeDetailNormalSpace::Tangent)
+					{
+						// Convert detail normal tangent space to object space
+						const FVector3f DetailNormalTangentSpace = DetailNormalObjectSpace;
+						DetailNormalObjectSpace = DetailNormalTangentSpace.X * FVector3f(DetailTangentX) + DetailNormalTangentSpace.Y * FVector3f(DetailTangentY) + DetailNormalTangentSpace.Z * DetailNormal;
+					}
 					Normalize(DetailNormalObjectSpace);
 					DetailNormal = DetailNormalObjectSpace;
 				}
diff --git a/Engine/Plugins/Runtime/GeometryProcessing/Source/DynamicMesh/Public/Sampling/MeshBakerCommon.h b/Engine/Plugins/Runtime/GeometryProcessing/Source/DynamicMesh/Public/Sampling/MeshBakerCommon.h
--- a/Engine/Plugins/Runtime/GeometryProcessing/Source/DynamicMesh/Public/Sampling/MeshBakerCommon.h
+++ b/Engine/Plugins/Runtime/GeometryProcessing/Source/DynamicMesh/Public/Sampling/MeshBakerCommon.h
@@ -96,6 +96,15 @@ class IMeshBakerDetailSampler
 public:
 	/** (Image, UVLayer) pair for detail textures */
 	using FBakeDetailTexture = TTuple<const TImageBuilder<FVector4f>*, int>;
+
+	enum class EBakeDetailNormalSpace
+	{
+		Tangent,
+		Object
+	};
+
+	/** (Image, NormalSpace, UVLayer) tuple for detail normal textures */
+	using FBakeDetailNormalTexture = TTuple<const TImageBuilder<FVector4f>*, int, EBakeDetailNormalSpace>;
 	
 	virtual ~IMeshBakerDetailSampler() = default;
 
@@ -114,14 +123,21 @@ public:
 	virtual void SetTextureMap(const void* Mesh, const FBakeDetailTexture& Map) = 0;
 
 	/** Associate a normal map and UV layer index for a given mesh in the detail set */
+	UE_DEPRECATED(5.1, "Use SetNormalTextureMap instead. This implementation assumes tangent normals.")
 	virtual void SetNormalMap(const void* Mesh, const FBakeDetailTexture& Map) = 0;
 
+	/** Associate a normal map and UV layer index for a given mesh in the detail set */
+	virtual void SetNormalTextureMap(const void* Mesh, const FBakeDetailNormalTexture& Map) = 0;
+
 	/** Retrieve a texture map and UV layer index from a given mesh in the detail set */
 	virtual const FBakeDetailTexture* GetTextureMap(const void* Mesh) const = 0;
 
 	/** Retrieve a normal map and UV layer index from a given mesh in the detail set */
+	UE_DEPRECATED(5.1, "Use GetNormalTextureMap instead. This implementation assumes tangent normals.")
 	virtual const FBakeDetailTexture* GetNormalMap(const void* Mesh) const = 0;
 
+	virtual const FBakeDetailNormalTexture* GetNormalTextureMap(const void* Mesh) const = 0;
+
 	/** @return true if identity correspondence is supported */
 	virtual bool SupportsIdentityCorrespondence() const = 0;
 	
@@ -354,7 +370,12 @@ public:
 
 	virtual void SetNormalMap(const void* Mesh, const FBakeDetailTexture& Map) override
 	{
-		DetailNormalMap = Map;
+		DetailNormalTextureMap = FBakeDetailNormalTexture(Map.Key, Map.Value, EBakeDetailNormalSpace::Tangent);
+	}
+
+	virtual void SetNormalTextureMap(const void* Mesh, const FBakeDetailNormalTexture& Map) override
+	{
+		DetailNormalTextureMap = Map;
 	}
 
 	virtual const FBakeDetailTexture* GetTextureMap(const void* Mesh) const override
@@ -364,7 +385,12 @@ public:
 	
 	virtual const FBakeDetailTexture* GetNormalMap(const void* Mesh) const override
 	{
-		return &DetailNormalMap;
+		return nullptr;
+	}
+
+	virtual const FBakeDetailNormalTexture* GetNormalTextureMap(const void* Mesh) const override
+	{
+		return &DetailNormalTextureMap;
 	}
 
 	virtual bool SupportsIdentityCorrespondence() const override
@@ -564,7 +590,10 @@ protected:
 	const FDynamicMeshAABBTree3* DetailSpatial = nullptr;
 	const FMeshTangentsd* DetailTangents = nullptr;
 	FBakeDetailTexture DetailTextureMap = FBakeDetailTexture(nullptr, 0);
+
+	UE_DEPRECATED(5.1, "Use DetailNormalTextureMap instead.")
 	FBakeDetailTexture DetailNormalMap = FBakeDetailTexture(nullptr, 0);
+	FBakeDetailNormalTexture DetailNormalTextureMap = FBakeDetailNormalTexture(nullptr, 0, EBakeDetailNormalSpace::Tangent);
 };		
 	
 
diff --git a/Engine/Plugins/Runtime/GeometryProcessing/Source/DynamicMesh/Public/Sampling/MeshNormalMapEvaluator.h b/Engine/Plugins/Runtime/GeometryProcessing/Source/DynamicMesh/Public/Sampling/MeshNormalMapEvaluator.h
--- a/Engine/Plugins/Runtime/GeometryProcessing/Source/DynamicMesh/Public/Sampling/MeshNormalMapEvaluator.h
+++ b/Engine/Plugins/Runtime/GeometryProcessing/Source/DynamicMesh/Public/Sampling/MeshNormalMapEvaluator.h
@@ -35,9 +35,15 @@ public:
 protected:
 	// Cached data
 	const IMeshBakerDetailSampler* DetailSampler = nullptr;
-	using FDetailNormalTexture = IMeshBakerDetailSampler::FBakeDetailTexture;
-	using FDetailNormalTextureMap = TMap<const void*, FDetailNormalTexture>;
-	FDetailNormalTextureMap DetailNormalTextures;
+	using FDetailNormalTexture UE_DEPRECATED(5.1, "Use FNormalTexture instead.") = IMeshBakerDetailSampler::FBakeDetailTexture;
+	using FDetailNormalTextureMap UE_DEPRECATED(5.1, "Use FNormalTextureMap instead.") = TMap<const void*, IMeshBakerDetailSampler::FBakeDetailTexture>;
+	UE_DEPRECATED(5.1, "Use DetailNormalMaps instead.")
+	TMap<const void*, IMeshBakerDetailSampler::FBakeDetailTexture> DetailNormalTextures;
+
+	using FNormalTexture = IMeshBakerDetailSampler::FBakeDetailNormalTexture;
+	using FNormalTextureMap = TMap<const void*, FNormalTexture>;
+	FNormalTextureMap DetailNormalMaps;
+	
 	bool bHasDetailNormalTextures = false;
 	const TMeshTangents<double>* BaseMeshTangents = nullptr;
 
diff --git a/Engine/Plugins/Runtime/GeometryProcessing/Source/DynamicMesh/Public/Sampling/MeshPropertyMapEvaluator.h b/Engine/Plugins/Runtime/GeometryProcessing/Source/DynamicMesh/Public/Sampling/MeshPropertyMapEvaluator.h
--- a/Engine/Plugins/Runtime/GeometryProcessing/Source/DynamicMesh/Public/Sampling/MeshPropertyMapEvaluator.h
+++ b/Engine/Plugins/Runtime/GeometryProcessing/Source/DynamicMesh/Public/Sampling/MeshPropertyMapEvaluator.h
@@ -48,9 +48,16 @@ public:
 protected:
 	// Cached data
 	const IMeshBakerDetailSampler* DetailSampler = nullptr;
-	using FDetailNormalTexture = IMeshBakerDetailSampler::FBakeDetailTexture;
-	using FDetailNormalTextureMap = TMap<const void*, FDetailNormalTexture>;
-	FDetailNormalTextureMap DetailNormalTextures;
+	
+	using FDetailNormalTexture UE_DEPRECATED(5.1, "Use FNormalTexture instead.") = IMeshBakerDetailSampler::FBakeDetailTexture;
+	using FDetailNormalTextureMap UE_DEPRECATED(5.1, "Use FNormalTextureMap instead.") = TMap<const void*, IMeshBakerDetailSampler::FBakeDetailTexture>;
+	UE_DEPRECATED(5.1, "Use DetailNormalMaps instead.")
+	TMap<const void*, IMeshBakerDetailSampler::FBakeDetailTexture> DetailNormalTextures;
+
+	using FNormalTexture = IMeshBakerDetailSampler::FBakeDetailNormalTexture;
+	using FNormalTextureMap = TMap<const void*, FNormalTexture>;
+	FNormalTextureMap DetailNormalMaps;
+	
 	bool bHasDetailNormalTextures = false;
 	FAxisAlignedBox3d Bounds;
 	FVector3f DefaultValue = FVector3f::Zero();
