Fixed a bug where we were setting the wrong indices for the pinned vertices when computing the conformal UV parameterization.

#rnx

diff --git a/Engine/Plugins/Runtime/GeometryProcessing/Source/DynamicMesh/Private/Parameterization/DynamicMeshUVEditor.cpp b/Engine/Plugins/Runtime/GeometryProcessing/Source/DynamicMesh/Private/Parameterization/DynamicMeshUVEditor.cpp
--- a/Engine/Plugins/Runtime/GeometryProcessing/Source/DynamicMesh/Private/Parameterization/DynamicMeshUVEditor.cpp
+++ b/Engine/Plugins/Runtime/GeometryProcessing/Source/DynamicMesh/Private/Parameterization/DynamicMeshUVEditor.cpp
@@ -586,7 +586,7 @@ bool FDynamicMeshUVEditor::SetTriangleUVsFromFreeBoundaryConformal(const TArray<
 			if (DistSqr > MaxDistSqr)
 			{
 				MaxDistSqr = DistSqr;
-				MaxDistPair = FIndex2i(i, j);
+				MaxDistPair = FIndex2i(ConstrainLoop[i], ConstrainLoop[j]);
 			}
 		}
 	}
