Fix DynamicMeshEditor's split meshes function setting every group ID to invalid group ID (-1) in its output meshes

This fixes plane cut setting invalid group data when 'keep both halves' is checked

diff --git a/Engine/Plugins/Runtime/GeometryProcessing/Source/DynamicMesh/Private/DynamicMeshEditor.cpp b/Engine/Plugins/Runtime/GeometryProcessing/Source/DynamicMesh/Private/DynamicMeshEditor.cpp
--- a/Engine/Plugins/Runtime/GeometryProcessing/Source/DynamicMesh/Private/DynamicMeshEditor.cpp
+++ b/Engine/Plugins/Runtime/GeometryProcessing/Source/DynamicMesh/Private/DynamicMeshEditor.cpp
@@ -2068,10 +2068,21 @@ bool FDynamicMeshEditor::SplitMesh(const FDynamicMesh3* SourceMesh, TArray<FDyna
 
 		FIndex3i Tri = SourceMesh->GetTriangle(SourceTID);
 
-		// FindOrCreateDuplicateGroup
-		// TODO: despite the FindOrCreateDuplicateGroup comment, this code does not create?  check about intent!
-		int NewGID = SourceMesh->HasTriangleGroups() ?
-			IndexMaps.GetNewGroup(SourceMesh->GetTriangleGroup(SourceTID)) : FDynamicMesh3::InvalidID;
+		// Find or create corresponding triangle group
+		int NewGID = FDynamicMesh3::InvalidID;
+		if (SourceMesh->HasTriangleGroups())
+		{
+			int SourceGroupID = SourceMesh->GetTriangleGroup(SourceTID);
+			if (SourceGroupID >= 0)
+			{
+				NewGID = IndexMaps.GetNewGroup(SourceGroupID);
+				if (NewGID == IndexMaps.InvalidID())
+				{
+					NewGID = Mesh.AllocateTriangleGroup();
+					IndexMaps.SetGroup(SourceGroupID, NewGID);
+				}
+			}
+		}
 
 		FIndex3i NewTri;
 		for (int j = 0; j < 3; ++j)
