GeometryProcessing: handle bad input edge list that sometimes comes from Bevel operation
#rnx

diff --git a/Engine/Plugins/Runtime/GeometryProcessing/Source/DynamicMesh/Private/Operations/PolyEditingEdgeUtil.cpp b/Engine/Plugins/Runtime/GeometryProcessing/Source/DynamicMesh/Private/Operations/PolyEditingEdgeUtil.cpp
--- a/Engine/Plugins/Runtime/GeometryProcessing/Source/DynamicMesh/Private/Operations/PolyEditingEdgeUtil.cpp
+++ b/Engine/Plugins/Runtime/GeometryProcessing/Source/DynamicMesh/Private/Operations/PolyEditingEdgeUtil.cpp
@@ -18,22 +18,31 @@ void UE::Geometry::ComputeInsetLineSegmentsFromEdges(
 	InsetLinesOut.SetNum(NumEdges);
 	for (int32 k = 0; k < NumEdges; ++k)
 	{
-		const FDynamicMesh3::FEdge EdgeVT = Mesh.GetEdge(EdgeList[k]);
-		FVector3d A = Mesh.GetVertex(EdgeVT.Vert[0]);
-		FVector3d B = Mesh.GetVertex(EdgeVT.Vert[1]);
-		FVector3d EdgeDir = Normalized(A - B);
-		FVector3d Midpoint = (A + B) * 0.5;
-		int32 EdgeTri = EdgeVT.Tri[0];
-		FVector3d Normal, Centroid; double Area;
-		Mesh.GetTriInfo(EdgeTri, Normal, Area, Centroid);
-
-		FVector3d InsetDir = Normal.Cross(EdgeDir);
-		if ((Centroid - Midpoint).Dot(InsetDir) < 0)
+		if (Mesh.IsEdge(EdgeList[k]))
 		{
-			InsetDir = -InsetDir;
+			const FDynamicMesh3::FEdge EdgeVT = Mesh.GetEdge(EdgeList[k]);
+			FVector3d A = Mesh.GetVertex(EdgeVT.Vert[0]);
+			FVector3d B = Mesh.GetVertex(EdgeVT.Vert[1]);
+			FVector3d EdgeDir = Normalized(A - B);
+			FVector3d Midpoint = (A + B) * 0.5;
+			int32 EdgeTri = EdgeVT.Tri[0];
+			FVector3d Normal, Centroid; double Area;
+			Mesh.GetTriInfo(EdgeTri, Normal, Area, Centroid);
+
+			FVector3d InsetDir = Normal.Cross(EdgeDir);
+			if ((Centroid - Midpoint).Dot(InsetDir) < 0)
+			{
+				InsetDir = -InsetDir;
+			}
+
+			InsetLinesOut[k] = FLine3d(Midpoint + InsetDistance * InsetDir, EdgeDir);
+		}
+		else
+		{
+			// This may produce nonsense in the calling code...but unclear what else to do here.
+			// Have observed this case coming from Bevel but so far have not determined the source.
+			InsetLinesOut[k] = FLine3d();		
 		}
-
-		InsetLinesOut[k] = FLine3d(Midpoint + InsetDistance * InsetDir, EdgeDir);
 	}
 }
 
