DynamicMesh: fix FGroupTopology::GetSelectionFrame for single-edge case, so that output frame correctly aligns X axis with edge line. Previous code was assuming that incoming InitialLocalFrame would have handled aligning Z axis to be perpendicular to the edge line.
#rnx

diff --git a/Engine/Plugins/Runtime/GeometryProcessing/Source/DynamicMesh/Private/GroupTopology.cpp b/Engine/Plugins/Runtime/GeometryProcessing/Source/DynamicMesh/Private/GroupTopology.cpp
--- a/Engine/Plugins/Runtime/GeometryProcessing/Source/DynamicMesh/Private/GroupTopology.cpp
+++ b/Engine/Plugins/Runtime/GeometryProcessing/Source/DynamicMesh/Private/GroupTopology.cpp
@@ -663,13 +663,26 @@ FFrame3d FGroupTopology::GetSelectionFrame(const FGroupTopologySelection& Select
 	FFrame3d StartFrame = (InitialLocalFrame) ? (*InitialLocalFrame) : FFrame3d();
 	if (NumEdges == 1)
 	{
-		int32 EdgeID = Selection.GetASelectedEdgeID();
+		int32 GroupEdgeID = Selection.GetASelectedEdgeID();
+		int32 MeshEdgeID = GetGroupEdgeEdges(GroupEdgeID)[0];
+
+		// align Z axis of frame to face normal of one of the connected faces. 
+		FIndex2i EdgeTris = Mesh->GetEdgeT(MeshEdgeID);
+		int32 UseFace = FMath::Min(EdgeTris.A, EdgeTris.B);
+		FVector3d FaceNormal = Mesh->GetTriNormal(UseFace);
+		if (FaceNormal.Length() > 0.1)
+		{
+			StartFrame.AlignAxis(2, FaceNormal);
+		}
+
+		// align X axis along the edge, around the aligned Z axis
 		FVector3d Tangent;
-		if (GetGroupEdgeTangent(EdgeID, Tangent))
+		if (GetGroupEdgeTangent(GroupEdgeID, Tangent))
 		{
 			StartFrame.ConstrainedAlignAxis(0, Tangent, StartFrame.Z());
 		}
-		StartFrame.Origin = GetEdgeMidpoint(EdgeID);
+
+		StartFrame.Origin = GetEdgeMidpoint(GroupEdgeID);
 		return StartFrame;
 	}
 	if (NumCorners == 1)
