MeshModelingTools: make FDynamicMeshEditor::RescaleAttributeUVs not crash on meshes with unset UV's when bWorldSpace is true.

diff --git a/Engine/Plugins/Runtime/GeometryProcessing/Source/DynamicMesh/Private/DynamicMeshEditor.cpp b/Engine/Plugins/Runtime/GeometryProcessing/Source/DynamicMesh/Private/DynamicMeshEditor.cpp
--- a/Engine/Plugins/Runtime/GeometryProcessing/Source/DynamicMesh/Private/DynamicMeshEditor.cpp
+++ b/Engine/Plugins/Runtime/GeometryProcessing/Source/DynamicMesh/Private/DynamicMeshEditor.cpp
@@ -1171,6 +1171,10 @@ void FDynamicMeshEditor::RescaleAttributeUVs(float UVScale, bool bWorldSpace, in
 		double TotalEdgeLen = 0;
 		for (int TID : Mesh->TriangleIndicesItr())
 		{
+			if (!UVs->IsSetTriangle(TID))
+			{
+				continue;
+			}
 			UVs->GetTriElements(TID, TriUVs[0], TriUVs[1], TriUVs[2]);
 			Mesh->GetTriVertices(TID, TriVs[0], TriVs[1], TriVs[2]);
 			if (ToWorld.IsSet())
