MeshModelingTools: Fix PolyEd edge selection crash on boundary edges.

diff --git a/Engine/Plugins/Runtime/GeometryProcessing/Source/DynamicMesh/Private/GroupTopology.cpp b/Engine/Plugins/Runtime/GeometryProcessing/Source/DynamicMesh/Private/GroupTopology.cpp
--- a/Engine/Plugins/Runtime/GeometryProcessing/Source/DynamicMesh/Private/GroupTopology.cpp
+++ b/Engine/Plugins/Runtime/GeometryProcessing/Source/DynamicMesh/Private/GroupTopology.cpp
@@ -668,7 +668,8 @@ FFrame3d FGroupTopology::GetSelectionFrame(const FGroupTopologySelection& Select
 
 		// align Z axis of frame to face normal of one of the connected faces. 
 		FIndex2i EdgeTris = Mesh->GetEdgeT(MeshEdgeID);
-		int32 UseFace = FMath::Min(EdgeTris.A, EdgeTris.B);
+		int32 UseFace = (EdgeTris.B != IndexConstants::InvalidID) ? FMath::Min(EdgeTris.A, EdgeTris.B)
+			: EdgeTris.A;
 		FVector3d FaceNormal = Mesh->GetTriNormal(UseFace);
 		if (FaceNormal.Length() > 0.1)
 		{
