Add a "prevent tiny triangles" option to Remesher. Currently the area threshold matches the threshold used when creating a cloth simulation mesh.

diff --git a/Engine/Plugins/Experimental/MeshModelingToolsetExp/Source/MeshModelingToolsExp/Private/RemeshMeshTool.cpp b/Engine/Plugins/Experimental/MeshModelingToolsetExp/Source/MeshModelingToolsExp/Private/RemeshMeshTool.cpp
--- a/Engine/Plugins/Experimental/MeshModelingToolsetExp/Source/MeshModelingToolsExp/Private/RemeshMeshTool.cpp
+++ b/Engine/Plugins/Experimental/MeshModelingToolsetExp/Source/MeshModelingToolsExp/Private/RemeshMeshTool.cpp
@@ -65,6 +65,7 @@ URemeshMeshToolProperties::URemeshMeshToolProperties()
 	bCollapses = true;
 	bReproject = true;
 	bPreventNormalFlips = true;
+	bPreventTinyTriangles = true;
 	bUseTargetEdgeLength = false;
 }
 
@@ -216,6 +217,7 @@ TUniquePtr<FDynamicMeshOperator> URemeshMeshTool::MakeNewOperator()
 	Op->GroupBoundaryConstraint = (EEdgeRefineFlags)BasicProperties->GroupBoundaryConstraint;
 	Op->MaterialBoundaryConstraint = (EEdgeRefineFlags)BasicProperties->MaterialBoundaryConstraint;
 	Op->bPreventNormalFlips = BasicProperties->bPreventNormalFlips;
+	Op->bPreventTinyTriangles = BasicProperties->bPreventTinyTriangles;
 	Op->bReproject = BasicProperties->bReproject;
 	Op->bSplits = BasicProperties->bSplits;
 	Op->RemeshIterations = BasicProperties->RemeshIterations;
diff --git a/Engine/Plugins/Experimental/MeshModelingToolsetExp/Source/MeshModelingToolsExp/Public/Properties/RemeshProperties.h b/Engine/Plugins/Experimental/MeshModelingToolsetExp/Source/MeshModelingToolsExp/Public/Properties/RemeshProperties.h
--- a/Engine/Plugins/Experimental/MeshModelingToolsetExp/Source/MeshModelingToolsExp/Public/Properties/RemeshProperties.h
+++ b/Engine/Plugins/Experimental/MeshModelingToolsetExp/Source/MeshModelingToolsExp/Public/Properties/RemeshProperties.h
@@ -63,6 +63,9 @@ public:
 	UPROPERTY(EditAnywhere, Category = "Remeshing|Edge Constraints")
 	bool bPreventNormalFlips = false;
 
+	/** Prevent introduction of tiny triangles or slivers */
+	UPROPERTY(EditAnywhere, Category = "Remeshing|Edge Constraints")
+	bool bPreventTinyTriangles = false;
 };
 
 UCLASS()
diff --git a/Engine/Plugins/Runtime/GeometryProcessing/Source/DynamicMesh/Private/MeshRefinerBase.cpp b/Engine/Plugins/Runtime/GeometryProcessing/Source/DynamicMesh/Private/MeshRefinerBase.cpp
--- a/Engine/Plugins/Runtime/GeometryProcessing/Source/DynamicMesh/Private/MeshRefinerBase.cpp
+++ b/Engine/Plugins/Runtime/GeometryProcessing/Source/DynamicMesh/Private/MeshRefinerBase.cpp
@@ -58,7 +58,41 @@ bool FMeshRefinerBase::CheckIfCollapseCreatesFlipOrInvalid(int vid, int vother,
 }
 
 
+bool FMeshRefinerBase::CheckIfCollapseCreatesTinyTriangle(int VertexID, int OtherVertexID, const FVector3d& NewVertexPosition, int IncidentTriangleC, int IncidentTriangleD) const
+{
+	for (int TriangleID : Mesh->VtxTrianglesItr(VertexID))
+	{
+		if (TriangleID == IncidentTriangleC || TriangleID == IncidentTriangleD)
+		{
+			continue;
+		}
+		FIndex3i CurrentTriangle = Mesh->GetTriangle(TriangleID);
+		if (CurrentTriangle[0] == OtherVertexID || CurrentTriangle[1] == OtherVertexID || CurrentTriangle[2] == OtherVertexID)
+		{
+			return true;		// invalid nbrhood for collapse
+		}
 
+		FVector3d TriangleVertices[3];
+		Mesh->GetTriVertices(TriangleID, TriangleVertices[0], TriangleVertices[1], TriangleVertices[2]);
+
+		// If the current triangle is already tiny, *do* allow the creation of a new tiny triangle (to allow remeshing to continue and hopefully improve things)
+		FVector3d CurrentNormal = (TriangleVertices[1] - TriangleVertices[0]).Cross(TriangleVertices[2] - TriangleVertices[0]);
+		if (CurrentNormal.SquaredLength() < TinyTriangleThreshold)
+		{
+			continue;
+		}
+
+		// Now find the size of the triangle that would result
+		TriangleVertices[CurrentTriangle.IndexOf(VertexID)] = NewVertexPosition;
+		FVector3d NewNormal = (TriangleVertices[1] - TriangleVertices[0]).Cross(TriangleVertices[2] - TriangleVertices[0]);
+		if (NewNormal.SquaredLength() < TinyTriangleThreshold)
+		{
+			return true;
+		}
+	}
+
+	return false;
+}
 
 /**
  * Check if edge flip might reverse normal direction.
@@ -93,7 +127,41 @@ bool FMeshRefinerBase::CheckIfFlipInvertsNormals(int a, int b, int c, int d, int
 }
 
 
+bool FMeshRefinerBase::CheckIfFlipCreatesTinyTriangle(int OriginalEdgeVertexA, int OriginalEdgeVertexB, int OppositeEdgeVertexC, int OppositeEdgeVertexD, int OriginalTriangleIndex) const
+{
+	FVector3d vC = Mesh->GetVertex(OppositeEdgeVertexC);
+	FVector3d vD = Mesh->GetVertex(OppositeEdgeVertexD);
+	FVector3d vA = Mesh->GetVertex(OriginalEdgeVertexA);
+	FVector3d vB = Mesh->GetVertex(OriginalEdgeVertexB);
+
+	// If the current triangles are already tiny, allow new triangles to be tiny as well
+	double CurrNormalSquaredLen = (vA - vC).Cross(vB - vC).SquaredLength();
+	if (CurrNormalSquaredLen < TinyTriangleThreshold)
+	{
+		return false;
+	}
+
+	CurrNormalSquaredLen = (vA - vD).Cross(vB - vD).SquaredLength();
+	if (CurrNormalSquaredLen < TinyTriangleThreshold)
+	{
+		return false;
+	}
 
+	// Now check triangle area of potential new triangle
+	double NewNormalSquaredLen = (vD - vC).Cross(vD - vB).SquaredLength();
+	if (NewNormalSquaredLen < TinyTriangleThreshold)
+	{
+		return true;
+	}
+
+	NewNormalSquaredLen = (vD - vC).Cross(vD - vA).SquaredLength();
+	if (NewNormalSquaredLen < TinyTriangleThreshold)
+	{
+		return true;
+	}
+
+	return false;
+}
 
 
 
diff --git a/Engine/Plugins/Runtime/GeometryProcessing/Source/DynamicMesh/Private/Remesher.cpp b/Engine/Plugins/Runtime/GeometryProcessing/Source/DynamicMesh/Private/Remesher.cpp
--- a/Engine/Plugins/Runtime/GeometryProcessing/Source/DynamicMesh/Private/Remesher.cpp
+++ b/Engine/Plugins/Runtime/GeometryProcessing/Source/DynamicMesh/Private/Remesher.cpp
@@ -211,6 +211,15 @@ FRemesher::EProcessResult FRemesher::ProcessEdge(int edgeID)
 			}
 		}
 
+		if (bPreventTinyTriangles)
+		{
+			if (CheckIfCollapseCreatesTinyTriangle(a, b, vNewPos, t0, t1) || CheckIfCollapseCreatesTinyTriangle(b, a, vNewPos, t0, t1))
+			{
+				goto abort_collapse;
+			}
+		}
+
+
 		// lots of cases where we cannot collapse, but we should just let
 		// mesh sort that out, right?
 		SaveEdgeBeforeModify(edgeID);
@@ -285,6 +294,11 @@ abort_collapse:
 			bTryFlip = false;
 		}
 
+		if (bTryFlip && bPreventTinyTriangles && CheckIfFlipCreatesTinyTriangle(a, b, c, d, t0))
+		{
+			bTryFlip = false;
+		}
+
 		if (bTryFlip) 
 		{
 			SaveEdgeBeforeModify(edgeID);
diff --git a/Engine/Plugins/Runtime/GeometryProcessing/Source/DynamicMesh/Public/MeshRefinerBase.h b/Engine/Plugins/Runtime/GeometryProcessing/Source/DynamicMesh/Public/MeshRefinerBase.h
--- a/Engine/Plugins/Runtime/GeometryProcessing/Source/DynamicMesh/Public/MeshRefinerBase.h
+++ b/Engine/Plugins/Runtime/GeometryProcessing/Source/DynamicMesh/Public/MeshRefinerBase.h
@@ -176,6 +176,24 @@ protected:
 	 */
 	bool CheckIfCollapseCreatesFlipOrInvalid(int vid, int vother, const FVector3d& newv, int tc, int td) const;
 
+	/** Avoid creation of triangles smaller than this value. 
+	 * We compare the triangle cross product's norm squared to this value since it's cheaper to compute than the triangle area. So the actual
+	 * area threshold is sqrt(TinyTriangleThreshold)/2.
+	 * Default is based on the threshold for creating triangles in a simulation mesh 
+	 */
+	double TinyTriangleThreshold = SMALL_NUMBER;
+
+	/**
+	 * Check if edge collapse will create a triangle with small area (either all vertices are close together, or in a sliver configuration)
+	 * This only checks one-ring of vid, so you have to call it twice, with vid and vother reversed, to check both one-rings
+	 * @param vid first vertex of edge
+	 * @param vother other vertex of edge
+	 * @param newv new vertex position after collapse
+	 * @param tc triangle on one side of edge
+	 * @param td triangle on other side of edge
+	 */
+	bool CheckIfCollapseCreatesTinyTriangle(int vid, int vother, const FVector3d& newv, int tc, int td) const;
+
 	/**
 	 * Check if edge flip might reverse normal direction.
 	 * Not entirely clear on how to best implement this test. Currently checking if any normal-pairs are reversed.
@@ -187,6 +205,16 @@ protected:
 	 */
 	bool CheckIfFlipInvertsNormals(int a, int b, int c, int d, int t0) const;
 
+	/**
+	 * Check if edge flip might create a triangle with small area (either all vertices are close together, or in a sliver configuration)
+	 * @param a first vertex of edge
+	 * @param b second vertex of edge
+	 * @param c opposing vertex 1
+	 * @param d opposing vertex 2
+	 * @param t0 index of triangle containing [a,b,c]
+	 */
+	bool CheckIfFlipCreatesTinyTriangle(int OriginalEdgeVertexA, int OriginalEdgeVertexB, int OppositeEdgeVertexC, int OppositeEdgeVertexD, int OriginalTriangleIndex) const;
+
 	/**
 	 * Figure out if we can collapse edge eid=[a,b] under current constraint set.
 	 * First we resolve vertex constraints using CanCollapseVertex(). However this
diff --git a/Engine/Plugins/Runtime/GeometryProcessing/Source/DynamicMesh/Public/Remesher.h b/Engine/Plugins/Runtime/GeometryProcessing/Source/DynamicMesh/Public/Remesher.h
--- a/Engine/Plugins/Runtime/GeometryProcessing/Source/DynamicMesh/Public/Remesher.h
+++ b/Engine/Plugins/Runtime/GeometryProcessing/Source/DynamicMesh/Public/Remesher.h
@@ -59,6 +59,11 @@ public:
 	 */
 	bool bPreventNormalFlips = false;
 
+	/** Controls whether we disallow creation of triangles with small areas inside edge operations.
+	 * This is moderately expensive and in some cases can result in lower-quality meshes. Disabled by default.
+	 */
+	bool bPreventTinyTriangles = false;
+
 	/** Available Edge Flip metrics */
 	enum class EFlipMetric
 	{
diff --git a/Engine/Plugins/Runtime/MeshModelingToolset/Source/ModelingOperators/Private/CleaningOps/RemeshMeshOp.cpp b/Engine/Plugins/Runtime/MeshModelingToolset/Source/ModelingOperators/Private/CleaningOps/RemeshMeshOp.cpp
--- a/Engine/Plugins/Runtime/MeshModelingToolset/Source/ModelingOperators/Private/CleaningOps/RemeshMeshOp.cpp
+++ b/Engine/Plugins/Runtime/MeshModelingToolset/Source/ModelingOperators/Private/CleaningOps/RemeshMeshOp.cpp
@@ -114,6 +114,7 @@ void FRemeshMeshOp::CalculateResult(FProgressCancel* Progress)
 	bool bIsUniformSmooth = (Remesher->SmoothType == FRemesher::ESmoothTypes::Uniform);
 
 	Remesher->bPreventNormalFlips = bPreventNormalFlips;
+	Remesher->bPreventTinyTriangles = bPreventTinyTriangles;
 
 	Remesher->DEBUG_CHECK_LEVEL = 0;
 
diff --git a/Engine/Plugins/Runtime/MeshModelingToolset/Source/ModelingOperators/Public/CleaningOps/RemeshMeshOp.h b/Engine/Plugins/Runtime/MeshModelingToolset/Source/ModelingOperators/Public/CleaningOps/RemeshMeshOp.h
--- a/Engine/Plugins/Runtime/MeshModelingToolset/Source/ModelingOperators/Public/CleaningOps/RemeshMeshOp.h
+++ b/Engine/Plugins/Runtime/MeshModelingToolset/Source/ModelingOperators/Public/CleaningOps/RemeshMeshOp.h
@@ -64,7 +64,8 @@ public:
 	int TriangleCountHint = 0;
 	float SmoothingStrength, TargetEdgeLength;
 	ERemeshSmoothingType SmoothingType;
-	bool bDiscardAttributes, bPreserveSharpEdges, bFlips, bSplits, bCollapses, bReproject, bPreventNormalFlips;
+	bool bDiscardAttributes, bPreserveSharpEdges, bFlips, bSplits, bCollapses, bReproject, bPreventNormalFlips, bPreventTinyTriangles;
+
 	// When true, result will have attributes object regardless of whether attributes 
 	// were discarded or present initially.
 	bool bResultMustHaveAttributesEnabled = false;
