Add a flag to optionally turn off the boundary reordering. This is useful when computing full Laplacian matricies where boundary reordering is not needed.

diff --git a/Engine/Plugins/Runtime/GeometryProcessing/Source/DynamicMesh/Public/Solvers/MeshLinearization.h b/Engine/Plugins/Runtime/GeometryProcessing/Source/DynamicMesh/Public/Solvers/MeshLinearization.h
--- a/Engine/Plugins/Runtime/GeometryProcessing/Source/DynamicMesh/Public/Solvers/MeshLinearization.h
+++ b/Engine/Plugins/Runtime/GeometryProcessing/Source/DynamicMesh/Public/Solvers/MeshLinearization.h
@@ -15,27 +15,36 @@ namespace Geometry
 class FDynamicMesh3;
 
 /**
-* Used linearize the VtxIds in a mesh as a single array and allow mapping from array offset to mesh VtxId.
-* Generally, the array offset will correspond to a matrix row when forming a Laplacian.
-* The last NumBoundaryVerts are the boundary verts.  This may be zero.
-*/
+ * Used to linearize the VtxIds in a mesh as a single array and allow mapping from array offset to mesh VtxId.
+ * Generally, the array offset will correspond to a matrix row when forming a Laplacian.
+ * The last NumBoundaryVerts are the boundary verts. This may be zero.
+ */
 class FVertexLinearization : public FElementLinearization
 {
 public:
 	FVertexLinearization() = default;
 
+	/**
+	 * @param Mesh 			  The mesh to lineralize.
+	 * @param bRemapBoundary  If true, will move the boundary vertices to the end of the arrays and record the number of 
+	 * 						  boundary vertices.
+	 */
 	template<typename MeshType>
-	FVertexLinearization(const MeshType& Mesh)
+	FVertexLinearization(const MeshType& Mesh, bool bRemapBoundary = true)
 	{
-		Reset(Mesh);
+		Reset(Mesh, bRemapBoundary);
 	}
 
 	template<typename MeshType>
-	void Reset(const MeshType& Mesh)
+	void Reset(const MeshType& Mesh, bool bRemapBoundary = true)
 	{
 		Empty();
 		FElementLinearization::Populate(Mesh.MaxVertexID(), Mesh.VertexCount(), Mesh.VertexIndicesItr());
-		RemapBoundaryVerts(Mesh);
+
+		if (bRemapBoundary) 
+		{
+			RemapBoundaryVerts(Mesh); 
+		}
 	}
 
 	int32 NumVerts() const { return FElementLinearization::NumIds(); }
