Capture variables by reference for the matrix operators.

diff --git a/Engine/Plugins/Runtime/GeometryProcessing/Source/DynamicMesh/Private/Solvers/Internal/PowerMethodSolver.cpp b/Engine/Plugins/Runtime/GeometryProcessing/Source/DynamicMesh/Private/Solvers/Internal/PowerMethodSolver.cpp
--- a/Engine/Plugins/Runtime/GeometryProcessing/Source/DynamicMesh/Private/Solvers/Internal/PowerMethodSolver.cpp
+++ b/Engine/Plugins/Runtime/GeometryProcessing/Source/DynamicMesh/Private/Solvers/Internal/PowerMethodSolver.cpp
@@ -218,12 +218,12 @@ void FSparsePowerMethod::CreateSparseMatrixOperator(const FSparseMatrixD& InMatr
     EMatrixSolverType MatrixSolverType = InMatrixHints.bIsPSD ? EMatrixSolverType::FastestPSD : EMatrixSolverType::LU;
     bool bIsSymmetric = InMatrixHints.bIsSymmetric || InMatrixHints.bIsPSD;
     
-    OutMatrixOp.Product = [InMatrix](const RealVectorType& InVector, RealVectorType& OutVector) 
+    OutMatrixOp.Product = [&InMatrix](const RealVectorType& InVector, RealVectorType& OutVector) 
     {
         OutVector = InMatrix * InVector; 
     };
 
-    OutMatrixOp.Factorize = [InMatrix, MatrixSolverType, bIsSymmetric]() 
+    OutMatrixOp.Factorize = [&InMatrix, &MatrixSolverType, bIsSymmetric]() 
     {	
         TUniquePtr<IMatrixSolverBase> Solver = ContructMatrixSolver(MatrixSolverType);
         Solver->SetUp(InMatrix, bIsSymmetric);
@@ -236,7 +236,7 @@ void FSparsePowerMethod::CreateSparseMatrixOperator(const FSparseMatrixD& InMatr
         return Solver;
     };
 
-    OutMatrixOp.Rows = [InMatrix]() 
+    OutMatrixOp.Rows = [&InMatrix]() 
     {
         return InMatrix.rows();
     };
