Store triangle area value in CotanTriangleData class.
Decrease the floor for the small triangle area to handle small yet still visible meshes.
Added the new TriangleArea option to the ECotangentWeightMode.

diff --git a/Engine/Plugins/Runtime/GeometryProcessing/Source/DynamicMesh/Private/Solvers/PrecomputedMeshWeightData.cpp b/Engine/Plugins/Runtime/GeometryProcessing/Source/DynamicMesh/Private/Solvers/PrecomputedMeshWeightData.cpp
--- a/Engine/Plugins/Runtime/GeometryProcessing/Source/DynamicMesh/Private/Solvers/PrecomputedMeshWeightData.cpp
+++ b/Engine/Plugins/Runtime/GeometryProcessing/Source/DynamicMesh/Private/Solvers/PrecomputedMeshWeightData.cpp
@@ -71,10 +71,10 @@ void UE::MeshDeformation::CotanTriangleData::Initialize(const FDynamicMesh3& Dyn
 		Cotangent[1] = -EdgeAB.Dot(EdgeBC) / TwiceArea;
 		Cotangent[2] = EdgeAC.Dot(EdgeBC) / TwiceArea;
 
+		Area = 0.5 * TwiceArea;
+
 		if (bIsObtuse())
 		{
-			const double Area = 0.5 * TwiceArea;
-
 			// Voronoi inappropriate case.  Instead use Area(T)/2 at obtuse corner
 			// and Area(T) / 4 at the other corners.
 
@@ -120,6 +120,8 @@ void UE::MeshDeformation::CotanTriangleData::Initialize(const FDynamicMesh3& Dyn
 		VoronoiArea[0] = SmallTriangleArea / 3.;
 		VoronoiArea[1] = SmallTriangleArea / 3.;
 		VoronoiArea[2] = SmallTriangleArea / 3.;
+
+		Area = SmallTriangleArea;
 	}
 }
 
diff --git a/Engine/Plugins/Runtime/GeometryProcessing/Source/DynamicMesh/Public/Solvers/LaplacianMatrixAssembly.h b/Engine/Plugins/Runtime/GeometryProcessing/Source/DynamicMesh/Public/Solvers/LaplacianMatrixAssembly.h
--- a/Engine/Plugins/Runtime/GeometryProcessing/Source/DynamicMesh/Public/Solvers/LaplacianMatrixAssembly.h
+++ b/Engine/Plugins/Runtime/GeometryProcessing/Source/DynamicMesh/Public/Solvers/LaplacianMatrixAssembly.h
@@ -197,8 +197,10 @@ namespace UE
 		{
 			/** Standard cotangent weights */
 			Default = 0,
-			/** magnitude of matrix entries clamped to [-1e5,1e5], scaled by area weight   */
-			ClampedMagnitude = 1
+			/** Magnitude of matrix entries clamped to [-1e5,1e5], scaled by area weight */
+			ClampedMagnitude = 1,
+			/** Divide cotangent weights by the area of the triangle */
+			TriangleArea = 2
 		};
 
 
@@ -890,10 +892,21 @@ void UE::MeshDeformation::ConstructFullCotangentLaplacian(const FDynamicMesh3& M
 			// Get the cotangents for this edge.
 			const int32 Tri0Idx = ToTriIdx[Edge.Tri[0]];
 			const CotanTriangleData& Tri0Data = CotangentTriangleDataArray[Tri0Idx];
-			const double CotanAlpha = Tri0Data.GetOpposingCotangent(EdgeId);
+			double CotanAlpha = Tri0Data.GetOpposingCotangent(EdgeId);
 
 			// The second triangle will be invalid if this is an edge!
-			const double CotanBeta = (Edge.Tri[1] != FDynamicMesh3::InvalidID) ? CotangentTriangleDataArray[ToTriIdx[Edge.Tri[1]]].GetOpposingCotangent(EdgeId) : 0.0;
+			double CotanBeta = (Edge.Tri[1] != FDynamicMesh3::InvalidID) ? CotangentTriangleDataArray[ToTriIdx[Edge.Tri[1]]].GetOpposingCotangent(EdgeId) : 0.0;
+
+			if (WeightMode == ECotangentWeightMode::TriangleArea) 
+			{
+				CotanAlpha /= Tri0Data.Area;
+
+				if (Edge.Tri[1] != FDynamicMesh3::InvalidID)
+				{
+					const CotanTriangleData& Tri1Data = CotangentTriangleDataArray[ToTriIdx[Edge.Tri[1]]];
+					CotanBeta /= Tri1Data.Area; 
+				}
+			}
 
 			// do not need to multiply by 0.5 here...
 			//double WeightIJ = 0.5 * (CotanAlpha + CotanBeta);
diff --git a/Engine/Plugins/Runtime/GeometryProcessing/Source/DynamicMesh/Public/Solvers/PrecomputedMeshWeightData.h b/Engine/Plugins/Runtime/GeometryProcessing/Source/DynamicMesh/Public/Solvers/PrecomputedMeshWeightData.h
--- a/Engine/Plugins/Runtime/GeometryProcessing/Source/DynamicMesh/Public/Solvers/PrecomputedMeshWeightData.h
+++ b/Engine/Plugins/Runtime/GeometryProcessing/Source/DynamicMesh/Public/Solvers/PrecomputedMeshWeightData.h
@@ -119,7 +119,7 @@ namespace UE
 			// NB: the cotan laplacian has terms ~ 1/TriArea
 			//     and the deformation matrix has terms ~ 1/TriArea**2
 
-			static constexpr double SmallTriangleArea = 1.e-4;
+			static constexpr double SmallTriangleArea = SMALL_NUMBER; //1.e-8;
 
 			// testing
 			int32 TriId = -1;
@@ -134,6 +134,10 @@ namespace UE
 
 			double VoronoiArea[3] = { 0. };
 
+			// Area of the triangle 
+			
+			double Area = -1;
+
 			// OppositeEdge[i] = Corresponding DynamicsMesh3::EdgeId for the edge that is opposite
 			// the i'th vertex in this triangle
 
