ModelingTools: Properly set identity correspondence strategy in baker when input meshes are the same.

diff --git a/Engine/Plugins/Experimental/GeometryScripting/Source/GeometryScriptingCore/Private/MeshBakeFunctions.cpp b/Engine/Plugins/Experimental/GeometryScripting/Source/GeometryScriptingCore/Private/MeshBakeFunctions.cpp
--- a/Engine/Plugins/Experimental/GeometryScripting/Source/GeometryScriptingCore/Private/MeshBakeFunctions.cpp
+++ b/Engine/Plugins/Experimental/GeometryScripting/Source/GeometryScriptingCore/Private/MeshBakeFunctions.cpp
@@ -495,6 +495,10 @@ namespace GeometryScriptBakeLocals
 		Baker.SetDimensions(BakeDimensions);
 		Baker.SetProjectionDistance(BakeOptions.ProjectionDistance);
 		Baker.SetSamplesPerPixel(GetSamplesPerPixel(BakeOptions.SamplesPerPixel));
+		if (bIsBakeToSelf)
+		{
+			Baker.SetCorrespondenceStrategy(FMeshBaseBaker::ECorrespondenceStrategy::Identity);
+		}
 
 		auto IsValidBakeType = [](EGeometryScriptBakeTypes, TArray<FGeometryScriptDebugMessage>*)
 		{
@@ -645,6 +649,10 @@ namespace GeometryScriptBakeLocals
 		Baker.SetTargetMeshUVLayer(TargetOptions.TargetUVLayer);
 		Baker.SetDetailSampler(&DetailSampler);
 		Baker.SetProjectionDistance(BakeOptions.ProjectionDistance);
+		if (bIsBakeToSelf)
+		{
+			Baker.SetCorrespondenceStrategy(FMeshBaseBaker::ECorrespondenceStrategy::Identity);
+		}
 
 		FEvaluatorState EvalState;
 		EvalState.TargetMesh = TargetMesh->GetMeshPtr();
diff --git a/Engine/Plugins/Experimental/MeshModelingToolsetExp/Source/MeshModelingToolsExp/Private/BakeMeshAttributeMapsTool.cpp b/Engine/Plugins/Experimental/MeshModelingToolsetExp/Source/MeshModelingToolsExp/Private/BakeMeshAttributeMapsTool.cpp
--- a/Engine/Plugins/Experimental/MeshModelingToolsetExp/Source/MeshModelingToolsExp/Private/BakeMeshAttributeMapsTool.cpp
+++ b/Engine/Plugins/Experimental/MeshModelingToolsetExp/Source/MeshModelingToolsExp/Private/BakeMeshAttributeMapsTool.cpp
@@ -103,6 +103,7 @@ public:
 	UBakeMeshAttributeMapsTool::FBakeSettings BakeSettings;
 	TSharedPtr<UE::Geometry::TMeshTangents<double>, ESPMode::ThreadSafe> BaseMeshTangents;
 	TSharedPtr<TArray<int32>, ESPMode::ThreadSafe> BaseMeshUVCharts;
+	bool bIsBakeToSelf = false;
 
 	// Map Type settings
 	EBakeMapType Maps;
@@ -136,6 +137,10 @@ public:
 		Baker->SetSamplesPerPixel(BakeSettings.SamplesPerPixel);
 		Baker->SetTargetMeshUVCharts(BaseMeshUVCharts.Get());
 		Baker->SetTargetMeshTangents(BaseMeshTangents);
+		if (bIsBakeToSelf)
+		{
+			Baker->SetCorrespondenceStrategy(FMeshBaseBaker::ECorrespondenceStrategy::Identity);
+		}
 		
 		FMeshBakerDynamicMeshSampler DetailSampler(DetailMesh.Get(), DetailSpatial.Get(), DetailMeshTangents.Get());
 		Baker->SetDetailSampler(&DetailSampler);
@@ -413,6 +418,7 @@ TUniquePtr<UE::Geometry::TGenericDataOperator<FMeshMapBaker>> UBakeMeshAttribute
 	Op->BaseMesh = &TargetMesh;
 	Op->BakeSettings = CachedBakeSettings;
 	Op->BaseMeshUVCharts = TargetMeshUVCharts;
+	Op->bIsBakeToSelf = bIsBakeToSelf;
 
 	constexpr EBakeMapType RequiresTangents = EBakeMapType::TangentSpaceNormal | EBakeMapType::BentNormal;
 	if ((bool)(CachedBakeSettings.BakeMapTypes & RequiresTangents))
diff --git a/Engine/Plugins/Experimental/MeshModelingToolsetExp/Source/MeshModelingToolsExp/Private/BakeMeshAttributeVertexTool.cpp b/Engine/Plugins/Experimental/MeshModelingToolsetExp/Source/MeshModelingToolsExp/Private/BakeMeshAttributeVertexTool.cpp
--- a/Engine/Plugins/Experimental/MeshModelingToolsetExp/Source/MeshModelingToolsExp/Private/BakeMeshAttributeVertexTool.cpp
+++ b/Engine/Plugins/Experimental/MeshModelingToolsetExp/Source/MeshModelingToolsExp/Private/BakeMeshAttributeVertexTool.cpp
@@ -59,6 +59,7 @@ public:
 	FDynamicMesh3* BaseMesh;
 	TSharedPtr<TMeshTangents<double>, ESPMode::ThreadSafe> BaseMeshTangents;
 	TUniquePtr<FMeshVertexBaker> Baker;
+	bool bIsBakeToSelf = false;
 
 	UBakeMeshAttributeVertexTool::FBakeSettings BakeSettings;
 	FOcclusionMapSettings OcclusionSettings;
@@ -82,6 +83,11 @@ public:
 		Baker->SetTargetMesh(BaseMesh);
 		Baker->SetTargetMeshTangents(BaseMeshTangents);
 		Baker->SetProjectionDistance(BakeSettings.ProjectionDistance);
+		if (bIsBakeToSelf)
+		{
+			Baker->SetCorrespondenceStrategy(FMeshBaseBaker::ECorrespondenceStrategy::Identity);
+		}
+		
 		Baker->BakeMode = BakeSettings.OutputMode == EBakeVertexOutput::RGBA ? FMeshVertexBaker::EBakeMode::RGBA : FMeshVertexBaker::EBakeMode::PerChannel;
 		
 		FMeshBakerDynamicMeshSampler DetailSampler(DetailMesh.Get(), DetailSpatial.Get());
@@ -443,6 +449,7 @@ TUniquePtr<UE::Geometry::TGenericDataOperator<FMeshVertexBaker>> UBakeMeshAttrib
 	Op->CurvatureSettings = CachedCurvatureMapSettings;
 	Op->TextureSettings = CachedTexture2DSettings;
 	Op->MultiTextureSettings = CachedMultiTexture2DSettings;
+	Op->bIsBakeToSelf = bIsBakeToSelf;
 
 	// Texture2DImage & MultiTexture settings
 	Op->TextureImage = CachedTextureImage;
diff --git a/Engine/Plugins/Runtime/GeometryProcessing/Source/DynamicMesh/Private/Sampling/MeshMapBaker.cpp b/Engine/Plugins/Runtime/GeometryProcessing/Source/DynamicMesh/Private/Sampling/MeshMapBaker.cpp
--- a/Engine/Plugins/Runtime/GeometryProcessing/Source/DynamicMesh/Private/Sampling/MeshMapBaker.cpp
+++ b/Engine/Plugins/Runtime/GeometryProcessing/Source/DynamicMesh/Private/Sampling/MeshMapBaker.cpp
@@ -161,16 +161,18 @@ void FMeshMapBaker::Bake()
 	ECorrespondenceStrategy UseStrategy = this->CorrespondenceStrategy;
 	bool bIsIdentity = true;
 	int NumDetailMeshes = 0;
-	auto CheckIdentity = [Mesh, &bIsIdentity, &NumDetailMeshes](const void* DetailMesh)
+	auto CheckIdentity = [this, Mesh, &bIsIdentity, &NumDetailMeshes](const void* DetailMesh)
 	{
-		bIsIdentity = bIsIdentity && (DetailMesh == Mesh);
+		// When the mesh pointers differ, loosely compare the meshes as a sanity check.
+		// TODO: Expose additional comparison metrics on the detail sampler when the mesh pointers differ.
+		bIsIdentity = bIsIdentity && (DetailMesh == Mesh || Mesh->TriangleCount() == DetailSampler->GetTriangleCount(DetailMesh));
 		++NumDetailMeshes;
 	};
 	DetailSampler->ProcessMeshes(CheckIdentity);
-	if (UseStrategy == ECorrespondenceStrategy::Identity && !ensure(bIsIdentity && (NumDetailMeshes == 1)))
+	if (UseStrategy == ECorrespondenceStrategy::Identity && !ensure(bIsIdentity && NumDetailMeshes == 1))
 	{
-		// Identity strategy requires mesh to be the same. Could potentially have two copies, in which
-		// case this ensure is too conservative, but for now we will assume this
+		// Identity strategy requires there to be only one mesh that is the same
+		// as the target mesh. 
 		UseStrategy = ECorrespondenceStrategy::NearestPoint;
 	}
 
diff --git a/Engine/Plugins/Runtime/GeometryProcessing/Source/DynamicMesh/Private/Sampling/MeshVertexBaker.cpp b/Engine/Plugins/Runtime/GeometryProcessing/Source/DynamicMesh/Private/Sampling/MeshVertexBaker.cpp
--- a/Engine/Plugins/Runtime/GeometryProcessing/Source/DynamicMesh/Private/Sampling/MeshVertexBaker.cpp
+++ b/Engine/Plugins/Runtime/GeometryProcessing/Source/DynamicMesh/Private/Sampling/MeshVertexBaker.cpp
@@ -101,14 +101,16 @@ void FMeshVertexBaker::BakeImpl(void* Data)
 	int NumDetailMeshes = 0;
 	auto CheckIdentity = [Baker, &bIsIdentity, &NumDetailMeshes](const void* DetailMesh)
 	{
-		bIsIdentity = bIsIdentity && (DetailMesh == Baker->TargetMesh);
+		// When the mesh pointers differ, loosely compare the meshes as a sanity check.
+		// TODO: Expose additional comparison metrics on the detail sampler when the mesh pointers differ.
+		bIsIdentity = bIsIdentity && (Baker->TargetMesh == DetailMesh || Baker->TargetMesh->TriangleCount() == Baker->DetailSampler->GetTriangleCount(DetailMesh));
 		++NumDetailMeshes;
 	};
 	Baker->DetailSampler->ProcessMeshes(CheckIdentity);
-	if (UseStrategy == ECorrespondenceStrategy::Identity && !ensure(bIsIdentity && (NumDetailMeshes == 1)))
+	if (UseStrategy == ECorrespondenceStrategy::Identity && !ensure(bIsIdentity && NumDetailMeshes == 1))
 	{
-		// Identity strategy requires mesh to be the same. Could potentially have two copies, in which
-		// case this ensure is too conservative, but for now we will assume this
+		// Identity strategy requires there to be only one mesh that is the same
+		// as the target mesh.
 		UseStrategy = ECorrespondenceStrategy::NearestPoint;
 	}
 	
