Added the Spectral conformal parameterization option for unwrapping UVs.
Refactored existing LSCM/DNCP UV unwrapping since a lot of code is shared with the Spectral method.

diff --git a/Engine/Plugins/Runtime/GeometryProcessing/Source/DynamicMesh/Private/Parameterization/DynamicMeshUVEditor.cpp b/Engine/Plugins/Runtime/GeometryProcessing/Source/DynamicMesh/Private/Parameterization/DynamicMeshUVEditor.cpp
--- a/Engine/Plugins/Runtime/GeometryProcessing/Source/DynamicMesh/Private/Parameterization/DynamicMeshUVEditor.cpp
+++ b/Engine/Plugins/Runtime/GeometryProcessing/Source/DynamicMesh/Private/Parameterization/DynamicMeshUVEditor.cpp
@@ -523,7 +523,18 @@ bool FDynamicMeshUVEditor::SetTriangleUVsFromFreeBoundaryConformal(const TArray<
 {
 	return SetTriangleUVsFromFreeBoundaryConformal(Triangles, false, Result);
 }
-bool FDynamicMeshUVEditor::SetTriangleUVsFromFreeBoundaryConformal(const TArray<int32>& Triangles, bool bUseExistingUVTopology, FUVEditResult* Result)
+
+bool FDynamicMeshUVEditor::SetTriangleUVsFromFreeBoundaryConformal(const TArray<int32>& Triangles, bool bUseExistingUVTopology, FUVEditResult* Result) 
+{
+	return SetTriangleUVsFromConformal(Triangles, bUseExistingUVTopology, false, false, Result);
+}
+
+bool FDynamicMeshUVEditor::SetTriangleUVsFromFreeBoundarySpectralConformal(const TArray<int32>& Triangles, bool bUseExistingUVTopology, bool bPreserveIrregularity, FUVEditResult* Result) 
+{
+	return SetTriangleUVsFromConformal(Triangles, bUseExistingUVTopology, true, bPreserveIrregularity, Result);
+}
+
+bool FDynamicMeshUVEditor::SetTriangleUVsFromConformal(const TArray<int32>& Triangles, bool bUseExistingUVTopology, bool bUseSpectral, bool bPreserveIrregularity, FUVEditResult* Result)
 {
 	if (ensure(UVOverlay) == false) return false;
 	if (Triangles.Num() == 0) return false;
@@ -566,38 +577,59 @@ bool FDynamicMeshUVEditor::SetTriangleUVsFromFreeBoundaryConformal(const TArray<
 	}
 
 	// is there a quick check we can do to ensure that we have a single connected component?
+	
+	TUniquePtr<UE::Solvers::IConstrainedMeshUVSolver> Solver = nullptr;
 
-	// make the UV solver
-	TUniquePtr<UE::Solvers::IConstrainedMeshUVSolver> Solver = UE::MeshDeformation::ConstructNaturalConformalParamSolver(Submesh);
-
-	// find pair of vertices to constrain. The standard procedure is to find the two furthest-apart vertices on the
-	// largest boundary loop. 
 	FMeshBoundaryLoops Loops(&Submesh, true);
-	if (Loops.GetLoopCount() == 0) return false;
+	if (Loops.GetLoopCount() == 0) 
+	{ 
+		return false; 
+	}
 	const TArray<int32>& ConstrainLoop = Loops[Loops.GetLongestLoopIndex()].Vertices;
 	int32 LoopNum = ConstrainLoop.Num();
-	FIndex2i MaxDistPair = FIndex2i::Invalid();
-	double MaxDistSqr = 0;
-	for (int32 i = 0; i < LoopNum; ++i)
+
+	if (bUseSpectral) 
 	{
-		for (int32 j = i + 1; j < LoopNum; ++j)
+		Solver = UE::MeshDeformation::ConstructSpectralConformalParamSolver(Submesh, bPreserveIrregularity);
+		
+		for (int32 Idx = 0; Idx < LoopNum; ++Idx)
 		{
-			double DistSqr = DistanceSquared(Submesh.GetVertex(ConstrainLoop[i]), Submesh.GetVertex(ConstrainLoop[j]));
-			if (DistSqr > MaxDistSqr)
-			{
-				MaxDistSqr = DistSqr;
-				MaxDistPair = FIndex2i(ConstrainLoop[i], ConstrainLoop[j]);
-			}
+			// It doesnt matter what the uv values or weights are, we are only interested in the indicies of the
+			// boundary vertices.
+			Solver->AddConstraint(ConstrainLoop[Idx], 0.0, FVector2d(0.0, 0.0), false);
 		}
 	}
-	if (ensure(MaxDistPair != FIndex2i::Invalid()) == false)
+	else 
 	{
-		return false;
-	}
+		Solver = UE::MeshDeformation::ConstructNaturalConformalParamSolver(Submesh);
 
-	// pin those vertices
-	Solver->AddConstraint(MaxDistPair.A, 1.0, FVector2d(0.0, 0.5), false);
-	Solver->AddConstraint(MaxDistPair.B, 1.0, FVector2d(1.0, 0.5), false);
+		// Find a pair of vertices to constrain. The standard procedure is to find the two furthest-apart vertices 
+		// on the largest boundary loop. 
+		FIndex2i MaxDistPair = FIndex2i::Invalid();
+		double MaxDistSqr = 0;
+		for (int32 Idx = 0; Idx < LoopNum; ++Idx)
+		{
+			for (int32 NextIdx = Idx + 1; NextIdx < LoopNum; ++NextIdx)
+			{
+				const double DistSqr = DistanceSquared(Submesh.GetVertex(ConstrainLoop[Idx]), 
+													   Submesh.GetVertex(ConstrainLoop[NextIdx]));
+				if (DistSqr > MaxDistSqr)
+				{
+					MaxDistSqr = DistSqr;
+					MaxDistPair = FIndex2i(ConstrainLoop[Idx], ConstrainLoop[NextIdx]);
+				}
+			}
+		}
+
+		if (ensure(MaxDistPair != FIndex2i::Invalid()) == false)
+		{
+			return false;
+		}
+
+		// pin those vertices
+		Solver->AddConstraint(MaxDistPair.A, 1.0, FVector2d(0.0, 0.5), false);
+		Solver->AddConstraint(MaxDistPair.B, 1.0, FVector2d(1.0, 0.5), false);
+	}
 
 	// solve for UVs
 	TArray<FVector2d> UVBuffer;
diff --git a/Engine/Plugins/Runtime/GeometryProcessing/Source/DynamicMesh/Private/Solvers/Internal/MeshUVSolver.cpp b/Engine/Plugins/Runtime/GeometryProcessing/Source/DynamicMesh/Private/Solvers/Internal/MeshUVSolver.cpp
--- a/Engine/Plugins/Runtime/GeometryProcessing/Source/DynamicMesh/Private/Solvers/Internal/MeshUVSolver.cpp
+++ b/Engine/Plugins/Runtime/GeometryProcessing/Source/DynamicMesh/Private/Solvers/Internal/MeshUVSolver.cpp
@@ -4,12 +4,15 @@
 
 #include "Solvers/MatrixInterfaces.h"
 #include "Solvers/LaplacianMatrixAssembly.h"
+#include "PowerMethodSolver.h"
 #include "MatrixSolver.h"
 #include "MeshBoundaryLoops.h"
 #include "MeshQueries.h"
 
 using namespace UE::Geometry;
 
+using ScalarT = FSparseMatrixD::Scalar;
+using MatrixTripletT = Eigen::Triplet<ScalarT>;
 
 //
 // Extension of TSparseMatrixAssembler suitable for eigen sparse matrix,
@@ -21,32 +24,24 @@ using namespace UE::Geometry;
 class FEigenDNCPSparseMatrixAssembler : public UE::Solvers::TSparseMatrixAssembler<double>
 {
 public:
-	typedef FSparseMatrixD::Scalar    ScalarT;
-	typedef Eigen::Triplet<ScalarT>  MatrixTripletT;
-
-	TUniquePtr<FSparseMatrixD> Matrix;
-	std::vector<MatrixTripletT> EntryTriplets;
-
-	int32 N;
-
-	FEigenDNCPSparseMatrixAssembler(int32 RowsI, int32 ColsJ)
+	
+	FEigenDNCPSparseMatrixAssembler(const int32 InNumVert)
 	{
-		check(RowsI == ColsJ);
-		N = RowsI;
+		NumVert = InNumVert;
 
-		Matrix = MakeUnique<FSparseMatrixD>(2 * N, 2 * N);
+		Matrix = MakeUnique<FSparseMatrixD>(2 * NumVert, 2 * NumVert);
 
 		ReserveEntriesFunc = [this](int32 NumElements)
 		{
 			EntryTriplets.reserve(NumElements);
 		};
 
-		AddEntryFunc = [this](int32 i, int32 j, double Value)
+		AddEntryFunc = [this](int32 RowIdx, int32 ColIdx, double Value)
 		{
 			// set upper-left block value
-			EntryTriplets.push_back(MatrixTripletT(i, j, Value));
+			EntryTriplets.push_back(MatrixTripletT(RowIdx, ColIdx, Value));
 			// set lower-right block value
-			EntryTriplets.push_back(MatrixTripletT(N + i, N + j, Value));
+			EntryTriplets.push_back(MatrixTripletT(NumVert + RowIdx, NumVert + ColIdx, Value));
 		};
 	}
 
@@ -56,300 +51,516 @@ public:
 		Matrix->makeCompressed();
 		Result.swap(*Matrix);
 	}
-};
 
+public:
+
+	TUniquePtr<FSparseMatrixD> Matrix;
+	std::vector<MatrixTripletT> EntryTriplets;
 
+	int32 NumVert;
+};
 
 
 //
-// Build DNCP matrix
+// FConstrainedMeshUVSolver
 //
-//
-static void ConstructNaturalConformalLaplacianSystem(
-	const FDynamicMesh3& DynamicMesh,
-	const FVertexLinearization& VertexMap,
-	const TArray<int32>& PinnedVertices,
-	FSparseMatrixD& Laplacian)
+
+FConstrainedMeshUVSolver::FConstrainedMeshUVSolver(const FDynamicMesh3& DynamicMesh)
+:
+IConstrainedMeshUVSolver(), Mesh(DynamicMesh)
 {
-	typedef FSparseMatrixD::Scalar  ScalarT;
-	typedef Eigen::Triplet<ScalarT> MatrixTripletT;
+	// compute linearization so we can store constraints at linearized indices
+	VtxLinearization.Reset(DynamicMesh, false);
+}
 
-	const TArray<int32>& ToMeshV = VertexMap.ToId();
-	const TArray<int32>& ToIndex = VertexMap.ToIndex();
-	const int32 NumVerts = VertexMap.NumVerts();
 
-	// Construct 2Nx2N system that includes both X and Y values for UVs. We will use block form [X, 0; 0, Y]
-	FEigenDNCPSparseMatrixAssembler DNCPLaplacianAssembler(NumVerts, NumVerts);
-	UE::MeshDeformation::ConstructFullCotangentLaplacian<double>(DynamicMesh, VertexMap, DNCPLaplacianAssembler, 
-		UE::MeshDeformation::ECotangentWeightMode::ClampedMagnitude,
-		UE::MeshDeformation::ECotangentAreaMode::NoArea );
-	FSparseMatrixD CotangentMatrix;
-	DNCPLaplacianAssembler.ExtractResult(CotangentMatrix);
-	// we want diagonal to be positive, so that sign of quadratic form is the same as
-	// the area matrix (which is a positive area)
-	CotangentMatrix = CotangentMatrix * -1.0;
+void FConstrainedMeshUVSolver::AddConstraint(const int32 VtxId, const double Weight, const FVector2d& Pos, const bool bPostFix)
+{
+	if ( ensure(VtxLinearization.IsValidId(VtxId)) == false) return;
+	int32 Index = VtxLinearization.GetIndex(VtxId);
 
-	// construct Area matrix. This matrix calculates the 2D area of the mesh.
-	int32 N = NumVerts;
-	FSparseMatrixD AreaMatrix(2 * N, 2 * N);
-	FMeshBoundaryLoops Loops(&DynamicMesh, true);
-	for (FEdgeLoop& Loop : Loops.Loops)
-	{
-		Algo::Reverse(Loop.Vertices);		// reverse loop to handle UE mesh orientation
-		int32 NumV = Loop.GetVertexCount();
-		for (int32 k = 0; k < NumV; ++k)
-		{
-			int32 a = Loop.Vertices[k];
-			int32 jx = ToIndex[a];
-			int32 jy = jx + N;
-			int32 b = Loop.Vertices[(k + 1) % NumV];
-			int32 kx = ToIndex[b];
-			int32 ky = kx + N;
-
-			AreaMatrix.coeffRef(jx, ky) = AreaMatrix.coeffRef(jx, ky) + 1;
-			AreaMatrix.coeffRef(ky, jx) = AreaMatrix.coeffRef(ky, jx) + 1;
-			AreaMatrix.coeffRef(kx, jy) = AreaMatrix.coeffRef(kx, jy) - 1;
-			AreaMatrix.coeffRef(jy, kx) = AreaMatrix.coeffRef(jy, kx) - 1;
-		}
-	}
+	FUVConstraint NewConstraint;
+	NewConstraint.ElementID = VtxId;
+	NewConstraint.ConstraintIndex = Index;
+	NewConstraint.Position = Pos;
+	NewConstraint.Weight = Weight;
+	NewConstraint.bPostFix = bPostFix;
+
+	ConstraintMap.Add(Index, NewConstraint);
 
+	bConstraintPositionsDirty = true;
+	bConstraintWeightsDirty = true;
+}
+
+
+bool FConstrainedMeshUVSolver::UpdateConstraintPosition(const int32 VtxId, const FVector2d& NewPosition, const bool bPostFix)
+{
+	if (ensure(VtxLinearization.IsValidId(VtxId)) == false) return false;
+	int32 Index = VtxLinearization.GetIndex(VtxId);
 
-	// test code that computes the 2D area of the input mesh, as well as the quadratic forms for the Cotan and
-	// Area matrices. On a planar mesh these should all produce the same value
-	//Eigen::Matrix<double, Eigen::Dynamic, 1> UVVector(2 * N);
-	//for (int32 vid : DynamicMesh.VertexIndicesItr())
-	//{
-	//	int32 Index = VertexMap.GetIndex(vid);
-	//	FVector3d Pos = DynamicMesh.GetVertex(vid);
-	//	UVVector.coeffRef(Index) = Pos.X;
-	//	UVVector.coeffRef(Index + N) = Pos.Y;
-	//}
-	//FVector2d VolArea = TMeshQueries<FDynamicMesh3>::GetVolumeArea(DynamicMesh);
-	//Eigen::Matrix<double, Eigen::Dynamic, 1> AreaMeasure = 0.5 * UVVector.transpose() * AreaMatrix * UVVector;
-	//Eigen::Matrix<double, Eigen::Dynamic, 1> CotanMeasure = 0.5 * UVVector.transpose() * CotangentMatrix * UVVector;
-
-	// assemble Conformal energy matrix
-	FSparseMatrixD Lc = CotangentMatrix - AreaMatrix;
-
-	// test code that computes the conformal energy quadratic form. This should be zero on a planar mesh.
-	//Eigen::Matrix<double, Eigen::Dynamic, 1> ConformalMeasure = 0.5 * UVVector.transpose() * Lc * UVVector;
-
-	// set fixed vertex rows to M(i,i) = 1;
-	FSparseMatrixD LcCons = Lc;
-	for (int32 PinnedVertID : PinnedVertices)
+	FUVConstraint* Found = ConstraintMap.Find(Index);
+	if (ensure(Found != nullptr))
 	{
-		int32 k = ToIndex[PinnedVertID];
-		// clear existing rows
-		LcCons.row(k) *= 0;
-		LcCons.row(k + N) *= 0;
-		// set to identity rows
-		LcCons.coeffRef(k, k) = 1.0;
-		LcCons.coeffRef(k + N, k + N) = 1.0;
+		Found->Position = NewPosition;
+		Found->bPostFix = bPostFix;
+		bConstraintPositionsDirty = true;
+		return true;
 	}
+	return false;
+}
 
-	// TODO: we can move the constrained columns to the RHS to keep the matrix symmetric.
-	// This would allow for more efficient solving, if we had a symmetric solver...however it
-	// also means we need a way to pass back the vector that needs to be added to the RHS
 
-	// make the compressed result
-	Laplacian = MoveTemp(LcCons);
-	Laplacian.makeCompressed();
+bool FConstrainedMeshUVSolver::UpdateConstraintWeight(const int32 VtxId, const double NewWeight)
+{
+	if (ensure(VtxLinearization.IsValidId(VtxId)) == false) return false;
+	int32 Index = VtxLinearization.GetIndex(VtxId);
+
+	FUVConstraint* Found = ConstraintMap.Find(Index);
+	if ( ensure(Found != nullptr) )
+	{
+		Found->Weight = NewWeight;
+		bConstraintWeightsDirty = true;
+		return true;
+	}
+	return false;
 }
 
 
+bool FConstrainedMeshUVSolver::IsConstrained(const int32 VtxId) const
+{
+	if (VtxLinearization.IsValidId(VtxId) == false) return false;
+	int32 Index = VtxLinearization.GetIndex(VtxId);
+	return ConstraintMap.Contains(Index);
+}
 
 
+void FConstrainedMeshUVSolver::ClearConstraints()
+{ 
+	ConstraintMap.Empty();
+	bConstraintPositionsDirty = true;
+	bConstraintWeightsDirty = true;
+}
 
 
 //
-// Solve a the natural/free-boundary conformal parameterization problem defined by the given system matrix
-// (likely generated by ConstructNaturalConformalLaplacianSystem() above) using the specified linear solver.
-// 
-// Requires that FixedIndices/Positions pairs define at least two constraint points for the solution to be well-defined.
-// Assumes that these rows are also constrained in the System Matrix
-static bool SolveDiscreteNaturalConformalSystem(
-	const FSparseMatrixD& CombinedUVSystemMatrix, const EMatrixSolverType MatrixSolverType,
-	const TArray<int32>& FixedIndices, const TArray<FVector2d>& FixedPositions, TArray<FVector2d>& Solution)
+//  FConformalMeshUVSolver
+//
+
+void FConformalMeshUVSolver::ConstructVectorAreaMatrix(FSparseMatrixD& OutAreaMatrix) 
 {
-	// create a suitable matrix solver
-	TUniquePtr<IMatrixSolverBase> MatrixSolver = ContructMatrixSolver(MatrixSolverType);
-	MatrixSolver->SetUp(CombinedUVSystemMatrix, false);
+	const int32 NumVert = VtxLinearization.NumVerts();
+	 
+	FMeshBoundaryLoops Loops(&Mesh, true);	
 
-	// set the constraint positions in the RHS
-	int32 N = CombinedUVSystemMatrix.rows() / 2;
-	Eigen::Matrix<double, Eigen::Dynamic, 1> RHSVector(2 * N);
-	RHSVector.setZero();
-	for (int32 k = 0; k < FixedIndices.Num(); ++k)
+	// Compute how many matrix entries to expect 
+	int32 NumTripletEntries = 0; 
+	for (const FEdgeLoop& Loop : Loops.Loops)
 	{
-		int32 Index = FixedIndices[k];
-		RHSVector.coeffRef(Index) = FixedPositions[k].X;
-		RHSVector.coeffRef(Index + N) = FixedPositions[k].Y;
-	}
+		NumTripletEntries += 4*Loop.GetVertexCount(); // 4 entries per loop edge
+	}	
 
-	// solve the linear system
-	Eigen::Matrix<double, Eigen::Dynamic, 1> SolutionVector(2 * N);
-	MatrixSolver->Solve(RHSVector, SolutionVector);
-	bool bSuccess = MatrixSolver->bSucceeded();
+	std::vector<MatrixTripletT> Triplets;
+	Triplets.reserve(NumTripletEntries);
 
-	// extract the solution if matrix solve succeeded, or set result to zero
-	Solution.SetNum(N);
-	if (bSuccess)
+	for (FEdgeLoop& Loop : Loops.Loops)
 	{
-		for (int32 k = 0; k < N; ++k)
+		const int32 NumLoopVert = Loop.GetVertexCount();
+		for (int32 Idx = 0; Idx < NumLoopVert; ++Idx)
 		{
-			Solution[k] = FVector2d(SolutionVector.coeffRef(k), SolutionVector.coeffRef(k + N));
+			// Directed edge is [EdgeVertex1, EdgeVertex2] with UVs [(U1,V1), (U2,V2)]
+			// Note we are reversing the direction of the edge to handle UE mesh orientation.
+			// Otherwise our Area matrix will be the negative of the correct Area matrix.
+			const int32 EdgeVertex1 = Loop.Vertices[(Idx + 1) % NumLoopVert];
+			const int32 EdgeVertex2 = Loop.Vertices[Idx];
+			
+			const int32 U1 = VtxLinearization.GetIndex(EdgeVertex1);
+			const int32 V1 = U1 + NumVert;
+			
+			const int32 U2 = VtxLinearization.GetIndex(EdgeVertex2);
+			const int32 V2 = U2 + NumVert;
+
+			Triplets.push_back(MatrixTripletT(U1, V2, 1));
+			Triplets.push_back(MatrixTripletT(U2, V1, -1));
+			
+			// Make it symmetric
+			Triplets.push_back(MatrixTripletT(V2, U1, 1));
+			Triplets.push_back(MatrixTripletT(V1, U2, -1));
 		}
 	}
-	else
+
+	OutAreaMatrix.resize(2 * NumVert, 2 * NumVert);
+	OutAreaMatrix.setFromTriplets(Triplets.begin(), Triplets.end());
+}
+
+
+void FConformalMeshUVSolver::ConstructWeightedVectorAreaMatrix(FSparseMatrixD& OutAreaMatrix,
+															   double& OutMaxTriArea,
+											  				   const double InSmallTriangleArea) 
+{
+	const int32 NumVert = VtxLinearization.NumVerts();
+	
+	std::vector<MatrixTripletT> Triplets;
+
+	const int32 NumTripletEntries = 4*3*Mesh.TriangleCount(); // 4 entries for each of the 3 edges of each triangle
+	Triplets.reserve(NumTripletEntries);
+
+	// Compute the areas of all triangles and find max.
+	// At the same time cache the results to be used later
+	TArray<double> TriAreasArray;
+	TriAreasArray.SetNum(Mesh.MaxTriangleID());
+	OutMaxTriArea = -1;
+	for (const int32 TriId : Mesh.TriangleIndicesItr()) 
 	{
-		for (int32 k = 0; k < N; ++k)
-		{
-			Solution[k] = FVector2d::Zero();
-		}
+		const double TriArea = Mesh.GetTriArea(TriId);
+		OutMaxTriArea = FMathd::Max(OutMaxTriArea, TriArea);
+		TriAreasArray[TriId] = TriArea;
 	}
 
-	return bSuccess;
-}
+	for (const int32 TriId : Mesh.TriangleIndicesItr()) 
+	{
+		const FIndex3i TriVert = Mesh.GetTriangle(TriId);
+		
+		double TriArea = TriAreasArray[TriId];
 
+		if (TriArea < InSmallTriangleArea) 
+		{
+			TriArea = InSmallTriangleArea;
+		}
 
+		// Normalize all area values with respect to the largest area among all triangles.
+		// This helps with the numerics being affected by the mesh scale and doesn't affect 
+		// the result of the solve.
+		TriArea /= OutMaxTriArea;
+		
+		const double Value = 1.0/TriArea;
 
+		// Iterator over each edge of the triangle
+		for (int32 VertIndex = 0; VertIndex < 3; ++VertIndex) 
+		{
+			// Directed edge is [EdgeVertex1, EdgeVertex2] with UVs [(U1,V1), (U2,V2)]
+			// Note we are reversing the direction of the edge to handle UE mesh orientation.
+			// Otherwise our Area matrix will be the negative of the correct Area matrix.
+			const int32 EdgeVertex1 = TriVert[(VertIndex + 1) % 3];
+			const int32 EdgeVertex2 = TriVert[VertIndex];
+
+			const int32 U1 = VtxLinearization.GetIndex(EdgeVertex1);
+			const int32 V1 = U1 + NumVert;
+			
+			const int32 U2 = VtxLinearization.GetIndex(EdgeVertex2);
+			const int32 V2 = U2 + NumVert;
+	
+			Triplets.push_back(MatrixTripletT(U1, V2,  Value));
+			Triplets.push_back(MatrixTripletT(V1, U2, -Value));
+
+			// Make it symmetric
+			Triplets.push_back(MatrixTripletT(V2, U1,  Value));
+			Triplets.push_back(MatrixTripletT(U2, V1, -Value));
+		}
+	}
 
+	OutAreaMatrix.resize(2 * NumVert, 2 * NumVert);
+	OutAreaMatrix.setFromTriplets(Triplets.begin(), Triplets.end());
+}
 
 
+void FConformalMeshUVSolver::ConstructConformalEnergyMatrix(FSparseMatrixD& OutConfromalEnergy,
+															const bool bInAreaWeighted, 
+															const double InSmallTriangleArea,
+															FSparseMatrixD* OutCotangentMatrix,
+															FSparseMatrixD* OutAreaMatrix)
+{
+	using namespace UE::MeshDeformation;
 
+	const int32 NumVerts = VtxLinearization.NumVerts();
+	
+	// Construct 2Nx2N system that includes both X and Y values for UVs. We will use block form [X, 0; 0, Y]
+	FEigenDNCPSparseMatrixAssembler DNCPLaplacianAssembler(NumVerts);
+
+	ConstructFullCotangentLaplacian<double>(Mesh,
+											VtxLinearization, 
+											DNCPLaplacianAssembler, 
+											bInAreaWeighted ? ECotangentWeightMode::TriangleArea : ECotangentWeightMode::ClampedMagnitude,
+											ECotangentAreaMode::NoArea);
+	
+	FSparseMatrixD CotangentMatrixLocal;
+	FSparseMatrixD* CotangentMatrix = OutCotangentMatrix == nullptr ? &CotangentMatrixLocal : OutCotangentMatrix;
+	DNCPLaplacianAssembler.ExtractResult(*CotangentMatrix);
+	
+	FSparseMatrixD AreaMatrixLocal;
+	FSparseMatrixD* AreaMatrix = OutAreaMatrix == nullptr ? &AreaMatrixLocal : OutAreaMatrix;
+
+	
+	// we want diagonal of the CotangentMatrix to be positive, so that sign of quadratic form is the same as
+	// the area matrix (which is a positive area)
+	if (bInAreaWeighted) 
+	{
+		double MaxTriArea;
+		ConstructWeightedVectorAreaMatrix(*AreaMatrix, MaxTriArea, InSmallTriangleArea);
+		
+		OutConfromalEnergy = -MaxTriArea*(*CotangentMatrix) - *AreaMatrix;
+	}
+	else 
+	{
+		ConstructVectorAreaMatrix(*AreaMatrix); 
+		
+		OutConfromalEnergy = -1*(*CotangentMatrix) - *AreaMatrix;
+	}
+}
 
 
+//
+// FLeastSquaresConformalMeshUVSolver
+//
 
-FConstrainedMeshUVSolver::~FConstrainedMeshUVSolver()
+void FLeastSquaresConformalMeshUVSolver::ConstructSystemMatrices(FSparseMatrixD& OutSystemMatrix)
 {
+	const int32 NumVerts = VtxLinearization.NumVerts();
+	
+	ConstructConformalEnergyMatrix(OutSystemMatrix);
+
+	// Set fixed vertex rows to M(i,i) = 1;
+	// TODO: we can move the constrained columns to the RHS to keep the matrix symmetric.
+	// This would allow for more efficient solving, if we had a symmetric solver...however it
+	// also means we need a way to pass back the vector that needs to be added to the RHS
+	for (const TPair<int32, FUVConstraint>& ElemPair : ConstraintMap)
+	{
+		const int32 RowIdx = ElemPair.Value.ConstraintIndex;
+		
+		// clear existing rows
+		OutSystemMatrix.row(RowIdx) *= 0;
+		OutSystemMatrix.row(RowIdx + NumVerts) *= 0;
+
+		// set rows to identity, these coefficients should already exist so no insertion happens here
+		OutSystemMatrix.coeffRef(RowIdx, RowIdx) = 1.0;
+		OutSystemMatrix.coeffRef(RowIdx + NumVerts, RowIdx + NumVerts) = 1.0;
+	}	
+	
+	OutSystemMatrix.makeCompressed();
 }
 
-FConstrainedMeshUVSolver::FConstrainedMeshUVSolver(const FDynamicMesh3& DynamicMesh, EUVSolveType UVSolveTypeIn)
+bool FLeastSquaresConformalMeshUVSolver::SolveParameterization(const FSparseMatrixD& InSystemMatrix, 
+															   const EMatrixSolverType InSolverType,
+															   FColumnVectorD& OutSolution)
 {
-	UVSolveType = UVSolveTypeIn;
+	const int32 NumVerts = VtxLinearization.NumVerts();
+	
+	// create a suitable matrix solver
+	TUniquePtr<IMatrixSolverBase> MatrixSolver = ContructMatrixSolver(InSolverType);
+	MatrixSolver->SetUp(InSystemMatrix, false);
+	if (MatrixSolver->bSucceeded() == false)  // check that the factorization succeeded 
+	{
+		return false;
+	}
 
-	// compute linearization so we can store constraints at linearized indices
-	VtxLinearization.Reset(DynamicMesh);
-}
+	// set the constraint positions in the RHS
+	FColumnVectorD RHSVector(2 * NumVerts);
+	RHSVector.setZero();
+	for (const TPair<int32, FUVConstraint>& ElemPair : ConstraintMap)
+	{
+		const int32 RowIdx = ElemPair.Value.ConstraintIndex;
 
+		RHSVector(RowIdx) = ElemPair.Value.Position.X;
+		RHSVector(RowIdx + NumVerts) = ElemPair.Value.Position.Y;
+	}
 
+	// solve the linear system
+	OutSolution.resize(2 * NumVerts, 1);
+	MatrixSolver->Solve(RHSVector, OutSolution);
+	const bool bSuccess = MatrixSolver->bSucceeded();
 
+	return bSuccess;
+}
 
 
-bool FConstrainedMeshUVSolver::SolveUVs(const FDynamicMesh3* DynamicMesh, TArray<FVector2d>& UVBuffer)
+bool FLeastSquaresConformalMeshUVSolver::SolveUVs(const FDynamicMesh3* InMeshPtr /*ignored*/, TArray<FVector2d>& OutUVBuffer)
 {
-	check(UVSolveType == EUVSolveType::NaturalConformal);
-
-	// create list of pinned vertices and positions
-	TArray<int32> FixedVertexIDs;
-	TArray<int32> FixedIndices;
-	TArray<FVector2d> FixedUVs;
-	for (auto& ElemPair : ConstraintMap)
-	{
-		FixedIndices.Add(ElemPair.Value.ConstraintIndex);
-		FixedVertexIDs.Add(ElemPair.Value.ElementID);
-		FixedUVs.Add(ElemPair.Value.Position);
-	}
+	// InMeshPtr and Mesh class member variable should be the same meshes.
+	checkSlow(InMeshPtr->IsSameAs(Mesh, FDynamicMesh3::FSameAsOptions()));
 
 	// build DNCP system
-	TUniquePtr<FSparseMatrixD> UVSystemMatrix = MakeUnique<FSparseMatrixD>();
-	ConstructNaturalConformalLaplacianSystem(*DynamicMesh, VtxLinearization, FixedVertexIDs, *UVSystemMatrix);
+	FSparseMatrixD SystemMatrix;
+	ConstructSystemMatrices(SystemMatrix);
 
-	// transfer to solver and solve
-	TArray<FVector2d> Solution;
-	bool bSolveOK = SolveDiscreteNaturalConformalSystem(*UVSystemMatrix, EMatrixSolverType::LU,
-		FixedIndices, FixedUVs, Solution);
+	// Transfer to solver and solve
+	FColumnVectorD Solution;
+	bool bSolveOK = SolveParameterization(SystemMatrix, EMatrixSolverType::LU, Solution);
 
-	ensure(bSolveOK);
-	if (bSolveOK == false)
+	if (ensure(bSolveOK) == false)
 	{
 		// If solve failed we will try QR solver which is more robust.
-		// This should perhaps be optional as the QR solve is much more expensive...
-		Solution.Reset();
-		bSolveOK = SolveDiscreteNaturalConformalSystem(*UVSystemMatrix, EMatrixSolverType::QR,
-			FixedIndices, FixedUVs, Solution);
-		ensure(bSolveOK);
+		// This should perhaps be optional as the QR solve is much more expensive.
+		bSolveOK = SolveParameterization(SystemMatrix, EMatrixSolverType::QR, Solution);
 	}
-
-	// copy back to input buffer
-	UVBuffer.SetNum(DynamicMesh->MaxVertexID());
-	for (int32 Index = 0; Index < VtxLinearization.NumIndices(); ++Index)
+	
+	OutUVBuffer.SetNum(Mesh.MaxVertexID());
+	if (ensure(bSolveOK)) 
 	{
-		int32 Id = VtxLinearization.GetId(Index);
-		UVBuffer[Id] = Solution[Index];
+		// Copy back to input buffer
+		const int32 NumVerts = VtxLinearization.NumVerts();
+		for (int32 Idx = 0; Idx < VtxLinearization.NumVerts(); ++Idx)
+		{
+			const int32 VertexID = VtxLinearization.GetId(Idx);
+			OutUVBuffer[VertexID] = FVector2d(Solution(Idx), Solution(Idx + NumVerts)); //Solution[Idx];
+		}
+	}
+	else 
+	{
+		// TODO: This should be deleted, the caller of the method should handle the case if SolveUVs returns false.
+		// Keeping this for now for the backwards compatibility
+		for (int32 Idx = 0; Idx < VtxLinearization.NumVerts(); ++Idx)
+		{
+			const int32 VertexID = VtxLinearization.GetId(Idx);
+			OutUVBuffer[VertexID] = FVector2d::Zero();
+		}
 	}
 
-	return true;
+	return bSolveOK;
 }
 
 
+//
+// FSpectralConformalMeshUVSolver
+//
 
-void FConstrainedMeshUVSolver::AddConstraint(const int32 VtxId, const double Weight, const FVector2d& Pos, const bool bPostFix)
+void FSpectralConformalMeshUVSolver::ConstructSystemMatrices(FSparseMatrixD& OutConformalEnergy,
+															 FSparseMatrixD& OutMatrixB,
+															 FSparseMatrixD& OutMatrixE)
 {
-	if ( ensure(VtxLinearization.IsValidId(VtxId)) == false) return;
-	int32 Index = VtxLinearization.GetIndex(VtxId);
+	using namespace UE::MeshDeformation;
 
-	FUVConstraint NewConstraint;
-	NewConstraint.ElementID = VtxId;
-	NewConstraint.ConstraintIndex = Index;
-	NewConstraint.Position = Pos;
-	NewConstraint.Weight = Weight;
-	NewConstraint.bPostFix = bPostFix;
+	const int32 NumVerts = VtxLinearization.NumVerts();
 
-	ConstraintMap.Add(Index, NewConstraint);
+	// Ensures that we are using the same area floor value as UE::MeshDeformation::CotanTriangleData::Initialize
+	// method which computes cotangent entries for the Laplacian matrix.
+	const double SmallTriangleArea = UE::MeshDeformation::CotanTriangleData::SmallTriangleArea;
+	ConstructConformalEnergyMatrix(OutConformalEnergy, bPreserveIrregularity, SmallTriangleArea);
+	OutConformalEnergy.makeCompressed();
 
-	bConstraintPositionsDirty = true;
-	bConstraintWeightsDirty = true;
-}
+	// construct matrix B
+	const int32 NumBndrVerts = ConstraintMap.Num();
+	std::vector<MatrixTripletT> TripletsB;	
+	TripletsB.reserve(2*NumBndrVerts);
 
+	for (const TPair<int32, FUVConstraint>& ElemPair : ConstraintMap)
+	{
+		const int32 RowIdx = ElemPair.Value.ConstraintIndex;
+		checkSlow(RowIdx < NumVerts);
 
-bool FConstrainedMeshUVSolver::UpdateConstraintPosition(const int32 VtxId, const FVector2d& NewPosition, const bool bPostFix)
-{
-	if (ensure(VtxLinearization.IsValidId(VtxId)) == false) return false;
-	int32 Index = VtxLinearization.GetIndex(VtxId);
+		TripletsB.push_back(MatrixTripletT(RowIdx, RowIdx, 1.0));
+		TripletsB.push_back(MatrixTripletT(RowIdx + NumVerts, RowIdx + NumVerts, 1.0));
+	}
 
-	FUVConstraint* Found = ConstraintMap.Find(Index);
-	if (ensure(Found != nullptr))
+	OutMatrixB = FSparseMatrixD(2*NumVerts, 2*NumVerts); 
+	OutMatrixB.setFromTriplets(TripletsB.begin(), TripletsB.end());
+	OutMatrixB.makeCompressed();
+		
+	// construct matrix E
+	std::vector<MatrixTripletT> TripletsE;
+	TripletsE.reserve(2*NumBndrVerts);
+
+	const ScalarT SqrtBndrNum = 1.0/FMath::Sqrt(static_cast<ScalarT>(NumBndrVerts));		
+	for (const TPair<int32, FUVConstraint>& ElemPair : ConstraintMap)
 	{
-		Found->Position = NewPosition;
-		Found->bPostFix = bPostFix;
-		bConstraintPositionsDirty = true;
-		return true;
+		const int32 RowIdx = ElemPair.Value.ConstraintIndex;
+
+		TripletsE.push_back(MatrixTripletT(RowIdx, 0, SqrtBndrNum));
+		TripletsE.push_back(MatrixTripletT(RowIdx + NumVerts, 1, SqrtBndrNum)) ;
 	}
-	return false;
+	
+	OutMatrixE = FSparseMatrixD(2*NumVerts, 2); 
+	OutMatrixE.setFromTriplets(TripletsE.begin(), TripletsE.end());
+	OutMatrixE.makeCompressed();
 }
 
 
-bool FConstrainedMeshUVSolver::UpdateConstraintWeight(const int32 VtxId, const double NewWeight)
-{
-	if (ensure(VtxLinearization.IsValidId(VtxId)) == false) return false;
-	int32 Index = VtxLinearization.GetIndex(VtxId);
-
-	FUVConstraint* Found = ConstraintMap.Find(Index);
-	if ( ensure(Found != nullptr) )
+bool FSpectralConformalMeshUVSolver::SolveParameterization(const FSparseMatrixD& InSystemMatrix,
+														   const FSparseMatrixD& InMatrixB,
+														   const FSparseMatrixD& InMatrixE,
+														   FColumnVectorD& OutSolution)
+{	
+	static constexpr FSparseMatrixD::Scalar Eps = 1e-8;
+	FSparseMatrixD IdentityMat(InSystemMatrix.rows(), InSystemMatrix.cols());
+	IdentityMat.setIdentity();
+	
+	// Make sure we are working with a PSD matrix
+	FSparseMatrixD SystemMatrixPSD = InSystemMatrix + Eps * IdentityMat;
+	
+	// Setup a custom matrix operator to avoid the expensive computation of the large dense matrix E*E^T and avoid 
+	// turning (B-E*E^T)*x into a dense matrix-vector product.
+	FPowerMethod::MatrixOperator OpMatrixB;
+	OpMatrixB.Product = [&InMatrixB, &InMatrixE](const FPowerMethod::RealVectorType& InVector, 
+												 FPowerMethod::RealVectorType& OutVector)
 	{
-		Found->Weight = NewWeight;
-		bConstraintWeightsDirty = true;
+		FPowerMethod::RealVectorType EEVector = (InMatrixE.transpose() * InVector).eval();
+		OutVector = InMatrixB * InVector - InMatrixE * EEVector;
+	};
+
+	// User-specified solver parameters
+	FPowerMethod::Parameters Parms;
+	Parms.Tolerance = 1e-10;
+	
+	FPowerMethod::ScalarType EigenValue;
+	FPowerMethod::RealVectorType EigenVector;
+	TUniquePtr<FPowerMethod> PMSolverPtr = MakeUnique<FSparsePowerMethod>(SystemMatrixPSD, 
+																		  FSparsePowerMethod::FMatrixHints(true), 
+																		  Parms);
+	PMSolverPtr->OpMatrixB = OpMatrixB;
+	
+	FPowerMethod::ESolveStatus SolveStatus = PMSolverPtr->Solve(EigenValue, OutSolution);
+	
+	if (ensure(SolveStatus == FPowerMethod::ESolveStatus::Success || SolveStatus == FPowerMethod::ESolveStatus::NotConverged))
+	{
+		// We allow the NotConverged case to go through since most likely the result should still be ok.
+		// But let the user know that tha result is not optimal.
+		if (SolveStatus == FPowerMethod::ESolveStatus::NotConverged)
+		{
+			checkSlow(false);
+			UE_LOG(LogTemp, Error, TEXT("Spectral Conformal UV Solver failed to converge. Resulting parametarization \
+										 might be invalid. You alternatively can try the default Conformal solve instead."));
+		} 
+		
 		return true;
 	}
-	return false;
+	else
+	{
+		UE_LOG(LogTemp, Error, TEXT("Spectral Conformal UV Solver failed."));
+		
+		return false;
+	}
 }
 
 
-bool FConstrainedMeshUVSolver::IsConstrained(const int32 VtxId) const
+bool FSpectralConformalMeshUVSolver::SolveUVs(const FDynamicMesh3* InMeshPtr, TArray<FVector2d>& OutUVBuffer)
 {
-	if (VtxLinearization.IsValidId(VtxId) == false) return false;
-	int32 Index = VtxLinearization.GetIndex(VtxId);
-	return ConstraintMap.Contains(Index);
-}
-
-
-void FConstrainedMeshUVSolver::ClearConstraints()
-{ 
-	ConstraintMap.Empty();
-	bConstraintPositionsDirty = true;
-	bConstraintWeightsDirty = true;
-}
+	// InMeshPtr and Mesh class member variable should be the same meshes.
+	checkSlow(InMeshPtr->IsSameAs(Mesh, FDynamicMesh3::FSameAsOptions()));
 
+	FSparseMatrixD SystemMatrix, MatrixB, MatrixE;
+	ConstructSystemMatrices(SystemMatrix, MatrixB, MatrixE);
 
+	FColumnVectorD Solution;
+	const bool bSolveOK = SolveParameterization(SystemMatrix, MatrixB, MatrixE, Solution);
+	
+	OutUVBuffer.SetNum(Mesh.MaxVertexID());
+	if (ensure(bSolveOK)) 
+	{
+		const int32 NumVerts = VtxLinearization.NumVerts();
+		for (int32 Idx = 0; Idx < VtxLinearization.NumVerts(); ++Idx)
+		{
+			const int32 VertexID = VtxLinearization.GetId(Idx);
+			OutUVBuffer[VertexID] = FVector2d(Solution(Idx), Solution(Idx + NumVerts));
+		}
+	}
+	else 
+	{
+		// TODO: This should be deleted, the caller of the method should handle the case if SolveUVs returns false.
+		// Keeping this for now for the backwards compatibility
+		for (int32 Idx = 0; Idx < VtxLinearization.NumVerts(); ++Idx)
+		{
+			const int32 VertexID = VtxLinearization.GetId(Idx);
+			OutUVBuffer[VertexID] = FVector2d::Zero();
+		}
+	}	
 
+	return bSolveOK;;
+}
diff --git a/Engine/Plugins/Runtime/GeometryProcessing/Source/DynamicMesh/Private/Solvers/Internal/MeshUVSolver.h b/Engine/Plugins/Runtime/GeometryProcessing/Source/DynamicMesh/Private/Solvers/Internal/MeshUVSolver.h
--- a/Engine/Plugins/Runtime/GeometryProcessing/Source/DynamicMesh/Private/Solvers/Internal/MeshUVSolver.h
+++ b/Engine/Plugins/Runtime/GeometryProcessing/Source/DynamicMesh/Private/Solvers/Internal/MeshUVSolver.h
@@ -5,7 +5,6 @@
 #include "CoreMinimal.h"
 #include "DynamicMesh/DynamicMesh3.h"
 #include "Solvers/MatrixInterfaces.h"
-#include "Solvers/MeshLaplacian.h"
 #include "Solvers/MeshLinearization.h"
 #include "Solvers/ConstrainedMeshSolver.h"
 
@@ -16,23 +15,19 @@ namespace UE
 namespace Geometry
 {
 
-enum class EUVSolveType
-{
-	NaturalConformal
-};
-
 /**
- * FConstrainedMeshUVSolver is an implmentation of IConstrainedMeshUVSolver that solves
- * for UVs using various methods
- *
+ * FConstrainedMeshUVSolver is an abstract class that implements methods for setting vertex constraints and their weights.
  */
 class DYNAMICMESH_API FConstrainedMeshUVSolver : public UE::Solvers::IConstrainedMeshUVSolver
 {
 public:
 	typedef UE::Solvers::FUVConstraint  FUVConstraint;
+	
+	FConstrainedMeshUVSolver(const FDynamicMesh3& DynamicMesh);
 
-	FConstrainedMeshUVSolver(const FDynamicMesh3& DynamicMesh, EUVSolveType UVSolveType);
-	virtual ~FConstrainedMeshUVSolver();
+	virtual ~FConstrainedMeshUVSolver() 
+	{
+	};
 
 	// Add constraint associated with given vertex id.  Boundary vertices will be ignored
 	void AddConstraint(const int32 VtxId, const double Weight, const FVector2d& Pos, const bool bPostFix) override;
@@ -49,14 +44,8 @@ public:
 	// Test if for constraint associated with given vertex id. Will return false for any boundary vert.
 	bool IsConstrained(const int32 VtxId) const override;
 
-	virtual bool SolveUVs(const FDynamicMesh3* DynamicMesh, TArray<FVector2d>& UVBuffer) override;
-
-
 protected:
 
-	// what type of UV solve to do
-	EUVSolveType UVSolveType;
-
 	// The Key (int32) here is the vertex index not vertex ID
 	// making it the same as the matrix row
 	TMap<int32, FUVConstraint> ConstraintMap;
@@ -64,10 +53,198 @@ protected:
 	// currently unused
 	bool bConstraintPositionsDirty = true;
 	bool bConstraintWeightsDirty = true;
+	
+	// Used to map between VtxId and vertex Index in linear vector.
+	FVertexLinearization VtxLinearization;
+
+	// The Mesh for which we are computing the UV mapping
+	const FDynamicMesh3& Mesh;
+};
+
 
-	// Used to map between VtxId and vertex Index in linear vector..
-	FVertexLinearization  VtxLinearization;
+/**
+ * FConformalMeshUVSolver is an abstract class that implements shared methods for the constrained conformal UV solvers.
+ */
+class DYNAMICMESH_API FConformalMeshUVSolver : public FConstrainedMeshUVSolver
+{
+public:
+	
+	FConformalMeshUVSolver(const FDynamicMesh3& InMesh)
+	:
+	FConstrainedMeshUVSolver(InMesh)
+	{
+	};
+
+	virtual ~FConformalMeshUVSolver() 
+	{
+	};
+
+	/**
+	 * Construct the 2Vx2V (V is the number of vertices) conformal energy matrix C = L - A, which is equal to the 
+	 * difference between the (possibly area-weighted) cotangent Laplacian (L) and the vector area matrices (A). 
+	 * 
+	 * @param OutConfromalEnergy 	Conformal energy matrix.
+	 * @param bInAreaWeighted  		Should the entries of the cotangent Laplacian and the vector area matrices be  
+	 * 								weighted by the areas of the triangles, they are computed on.
+	 * @param InSmallTriangleArea   The "floor" for the triangle area. Triangles with areas below this value are  
+	 * 						   	    considered to be triangles with SmallTriangleArea area. Optional and only used
+	 * 								if bAreaWeighted == true.
+	 * @param OutCotangentMatrix 	(optional) If not nullptr, then store the cotangent Laplacian matrix.
+	 * @param OutAreaMatrix 		(optional) If not nullptr, then store the vector area matrix.
+	 */
+	void ConstructConformalEnergyMatrix(FSparseMatrixD& OutConfromalEnergy,
+										const bool bInAreaWeighted = false, 
+										const double InSmallTriangleArea = SMALL_NUMBER,
+										FSparseMatrixD* OutCotangentMatrix = nullptr,
+										FSparseMatrixD* OutAreaMatrix = nullptr);
+protected:
+
+	using FColumnVectorD = Eigen::Matrix<FSparseMatrixD::Scalar, Eigen::Dynamic, 1>;
+	
+	/**
+	 * Construct the 2Vx2V (V is the number of vertices) symmetric vector area matrix A. Given the vector X = [u;v] of 
+	 * the stacked UV coordinates, we can compute the signed 2D area of the parameterization using the quadratic form 
+	 * (1/4)*X'*A*X.
+	 * 
+	 * The area is computed using only the coordinates of the boundary vertices as the contributions from internal edges 
+	 * cancel out. Given a border edge ij:
+	 * A_{Ui,Vj} =  1, where Ui is the u index (into X) of vert_i and Vj is the v index of vert_j.
+	 * A_{Vi,Uj} = -1, where Vi is the v index of vert_i and Uj is the u index of vert_j.
+	 * 
+	 * To ensure the symmetry of A, we additionally set A_{Vj,Ui} = 1 and A_{Uj,Vi} = -1 since the area matrix A 
+	 * will be used as a quadratic coefficients matrix.
+	 * 
+	 * @param OutAreaMatrix The signed vector area matrix.
+	 */
+	void ConstructVectorAreaMatrix(FSparseMatrixD& OutAreaMatrix);
+
+	
+	/**
+	 * Construct the 2Vx2V (V is the number of vertices) symmetric weighted area matrix A. This is similar to 
+	 * ConstructVectorAreaMatrix() method but the matrix coefficients are now scaled by the areas of the triangle on
+	 * which they are computed. This formulation is used to enforce insensitivity to sampling irregularity of the meshes.
+	 * 
+	 * Because of area scaling, we no longer can only use border edges, but all the directed edges will need to be 
+	 * considered. Given the directed edge ij of a triangle t:
+	 * A_{Ui,Vj} = 1/Area(t), where Ui is the u index (into X) of vert_i and Vj is the v index of vert_j.
+	 * A_{Vi,Uj} = -1/Area(t), where Vi is the v index of vert_i and Uj is the u index of vert_j.
+	 * 
+	 * To ensure the symmetry of A, we additionally set A_{Vj,Ui} = 1/Area(t) and A_{Uj,Vi} = -1/Area(t) since the area 
+	 * matrix A will be used as a quadratic coefficients matrix.
+	 *  
+	 * @param OutAreaMatrix  	   The signed weighted vector area matrix.
+	 * @param OutMaxArea		   The maximum area among all mesh triangles
+	 * @param InSmallTriangleArea  The "floor" for the triangle area. Triangles with areas below this value are  
+	 * 						   	   considered to be triangles with SmallTriangleArea area.
+	 */
+	void ConstructWeightedVectorAreaMatrix(FSparseMatrixD& OutAreaMatrix, 
+										   double& OutMaxArea,
+										   const double InSmallTriangleArea = SMALL_NUMBER); 
+
+	
 };
 
+/**
+ * FLeastSquaresConformalMeshUVSolver computes the conformal UV map using the Least squares conformal maps for automatic
+ * texture atlas generation [Levy 2002] method. It's equivalent to Intrinsic parameterizations of surface meshes [Desburn 2002].
+ * 
+ * However, the energy is constructed using the formulation presented in the Spectral Conformal parameterization [Mullen 2008].
+ */
+class DYNAMICMESH_API FLeastSquaresConformalMeshUVSolver : public FConformalMeshUVSolver
+{
+public:
+		
+	FLeastSquaresConformalMeshUVSolver(const FDynamicMesh3& InMesh)
+	:
+	FConformalMeshUVSolver(InMesh)
+	{
+	};
+
+	virtual ~FLeastSquaresConformalMeshUVSolver() 
+	{
+	};
+
+	virtual bool SolveUVs(const FDynamicMesh3* InMesh, TArray<FVector2d>& OutUVBuffer) override;
+
+protected:
+	
+	/** 
+	 * Create the "left-hand side" system matrix used in our linear solve. The matrix is equal to the conformal energy 
+	 * matrix (created with ConstructConformalEnergyMatrix()) but with the rows corresponding to the constraints 
+	 * (defined by ConstraintMap) to only have one non-zero value along the diagonal (set to one).
+	 */
+	void ConstructSystemMatrices(FSparseMatrixD& OutSystemMatrix);
+	
+	/**
+	 * Solve the natural/free-boundary conformal parameterization problem defined by the given system matrix
+	 * (likely generated by ConstructSystemMatrices() above) using the specified linear solver.
+	 * 
+	 * Requires that FixedIndices/Positions pairs define at least two constraint points for the solution to be 
+	 * well-defined. Assumes that these rows are also constrained in the System Matrix.
+	 * 
+	 * @param OutSolution Array of UV coordinates if parameterization succeeded. Empty if the solve failed.
+	 * 
+	 * @return true if the solve succeded, false otherwise.
+	 */
+	bool SolveParameterization(const FSparseMatrixD& InSystemMatrix, 
+							   const EMatrixSolverType InMatrixSolverType, 
+							   FColumnVectorD& OutSolution);
+};
+
+
+/**
+ * FConformalMeshUVSolver computes the conformal UV map using the Spectral Conformal Parameterization [Mullen 2008] 
+ * method including enforcing insensitivity to sampling irregularity extension.
+ */
+class DYNAMICMESH_API FSpectralConformalMeshUVSolver : public FConformalMeshUVSolver
+{
+public:
+
+	FSpectralConformalMeshUVSolver(const FDynamicMesh3& InMesh, const bool bInPreserveIrregularity = false)
+	:
+	FConformalMeshUVSolver(InMesh), bPreserveIrregularity(bInPreserveIrregularity)
+	{
+	};
+
+	virtual ~FSpectralConformalMeshUVSolver() 
+	{
+	};
+
+	virtual bool SolveUVs(const FDynamicMesh3* InMesh, TArray<FVector2d>& OutUVBuffer) override;
+
+protected: 
+
+	/** 
+	 * Construct the matrices needed for solving the general eigenvalue problem.
+	 * 
+	 * @param OutConformalEnergy Conformal energy matrix (created using ConstructConformalEnergyMatrix() method).
+	 * @param OutMatrixB 		 2V×2V diagonal matrix with 1 at each diagonal element corresponding to the boundary  
+	 *							 vertices (not including any of the internal boundaries) and 0 everywhere else.
+	 * @param OutMatrixE 		 2V×2V matrix such that E(i,1) (resp., E(i,2)) is 1/sqrt(V_b) for each u-coordinate 
+	 * 					 		 (resp.,v-coordinate) of the boundary vertices (not including any of the internal 
+	 * 							 boundaries) and 0 otherwise. V_b is the number of the boundary vertices.
+	 */
+	void ConstructSystemMatrices(FSparseMatrixD& OutConformalEnergy, 
+								 FSparseMatrixD& OutMatrixB, 
+								 FSparseMatrixD& OutMatrixE);
+
+	/** 
+	 * Solve the general eigenvalue problem to compute the UV mapping using the matrices computed in 
+	 * ConstructSystemMatrices().
+	 */
+	bool SolveParameterization(const FSparseMatrixD& InConformalEnergy,
+							   const FSparseMatrixD& InMatrixB,
+							   const FSparseMatrixD& InMatrixE,
+							   FColumnVectorD& OutSolution);
+	
+protected: 
+
+	// The mesh can contain triangles of variable sizes. By default this difference in size is not considered during 
+	// the UV solve. Enabling this allows to preserve the triangle size irregularity in the resulting parameterization.
+	bool bPreserveIrregularity = false;
+};
+
+
+
 } // end namespace UE::Geometry
 } // end namespace UE
\ No newline at end of file
diff --git a/Engine/Plugins/Runtime/GeometryProcessing/Source/DynamicMesh/Private/Solvers/MeshParameterizationSolvers.cpp b/Engine/Plugins/Runtime/GeometryProcessing/Source/DynamicMesh/Private/Solvers/MeshParameterizationSolvers.cpp
--- a/Engine/Plugins/Runtime/GeometryProcessing/Source/DynamicMesh/Private/Solvers/MeshParameterizationSolvers.cpp
+++ b/Engine/Plugins/Runtime/GeometryProcessing/Source/DynamicMesh/Private/Solvers/MeshParameterizationSolvers.cpp
@@ -6,7 +6,12 @@
 
 TUniquePtr<UE::Solvers::IConstrainedMeshUVSolver> UE::MeshDeformation::ConstructNaturalConformalParamSolver(const FDynamicMesh3& DynamicMesh)
 {
-	TUniquePtr<UE::Solvers::IConstrainedMeshUVSolver> Solver(new FConstrainedMeshUVSolver(DynamicMesh, EUVSolveType::NaturalConformal));
+	TUniquePtr<UE::Solvers::IConstrainedMeshUVSolver> Solver(new FLeastSquaresConformalMeshUVSolver(DynamicMesh));
 	return Solver;
 }
 
+TUniquePtr<UE::Solvers::IConstrainedMeshUVSolver> UE::MeshDeformation::ConstructSpectralConformalParamSolver(const FDynamicMesh3& DynamicMesh, bool bPreserveIrregularity)
+{
+	TUniquePtr<UE::Solvers::IConstrainedMeshUVSolver> Solver(new FSpectralConformalMeshUVSolver(DynamicMesh, bPreserveIrregularity));
+	return Solver;
+}
\ No newline at end of file
diff --git a/Engine/Plugins/Runtime/GeometryProcessing/Source/DynamicMesh/Public/Parameterization/DynamicMeshUVEditor.h b/Engine/Plugins/Runtime/GeometryProcessing/Source/DynamicMesh/Public/Parameterization/DynamicMeshUVEditor.h
--- a/Engine/Plugins/Runtime/GeometryProcessing/Source/DynamicMesh/Public/Parameterization/DynamicMeshUVEditor.h
+++ b/Engine/Plugins/Runtime/GeometryProcessing/Source/DynamicMesh/Public/Parameterization/DynamicMeshUVEditor.h
@@ -171,6 +171,15 @@ public:
 	 */
 	bool SetTriangleUVsFromFreeBoundaryConformal(const TArray<int32>& Triangles, bool bUseExistingUVTopology, FUVEditResult* Result = nullptr);
 
+	/**
+	 * Create new UV island for given Triangles, and set UVs for that island using Spectral Conformal Map.
+	 * @param Triangles list of triangles
+	 * @param bUseExistingUVTopology if true, re-solve for existing UV set, rather than constructing per-vertex UVs from triangle set. Allows for solving w/ partial seams, interior cuts, etc. 
+	 * @param bPreserveIrregularity if true, reduces distortion for meshes with triangles of vastly different sizes.
+	 * @warning computes a single parameterization, so input triangle set must be connected, however this is not verified internally
+	 */
+	bool SetTriangleUVsFromFreeBoundarySpectralConformal(const TArray<int32>& Triangles, bool bUseExistingUVTopology, bool bPreserveIrregularity, FUVEditResult* Result = nullptr);
+
 	/**
 	 * Cut existing UV topolgy with a set of edges. This allows for creating partial seams/darts, interior cuts, etc.
 	 * 
@@ -260,6 +269,13 @@ public:
 	* @param BoundingBox if not null, return the overall bounding box of the UVs for the specified triangles along with the exact area occupied.
 	*/
 	static double DetermineAreaFromUVs(const FDynamicMeshUVOverlay& UVOverlay, const TArray<int32>& Triangles, FAxisAlignedBox2f* BoundingBox = nullptr);
+
+private: 
+
+	/**
+	 * Helper method to solve either the Least Squares Conformal Map or the Spectral Conformal Map problems.
+	 */
+	bool SetTriangleUVsFromConformal(const TArray<int32>& Triangles, bool bUseExistingUVTopology, bool bUseSpectral, bool bPreserveIrregularity, FUVEditResult* Result);
 };
 
 
diff --git a/Engine/Plugins/Runtime/GeometryProcessing/Source/DynamicMesh/Public/Solvers/MeshParameterizationSolvers.h b/Engine/Plugins/Runtime/GeometryProcessing/Source/DynamicMesh/Public/Solvers/MeshParameterizationSolvers.h
--- a/Engine/Plugins/Runtime/GeometryProcessing/Source/DynamicMesh/Public/Solvers/MeshParameterizationSolvers.h
+++ b/Engine/Plugins/Runtime/GeometryProcessing/Source/DynamicMesh/Public/Solvers/MeshParameterizationSolvers.h
@@ -14,10 +14,16 @@ namespace UE
 		using namespace UE::Geometry;
 
 		/**
-		 * Create solver for Free-Boundary UV Parameterization for this mesh.
-		 * @warning Assumption is that mesh is a single connected component
+		 * Create solver for the Least Squares Conformal UV Parameterization for this mesh.
+		 * @warning Assumption is that mesh is a single connected component.
 		 */
 		TUniquePtr<UE::Solvers::IConstrainedMeshUVSolver> DYNAMICMESH_API ConstructNaturalConformalParamSolver(const FDynamicMesh3& DynamicMesh);
+	
+		/**
+		 * Create solver for the Spectral Conformal UV Parameterization for this mesh.
+		 * @warning Assumption is that mesh is a single connected component.
+		 */
+		TUniquePtr<UE::Solvers::IConstrainedMeshUVSolver> DYNAMICMESH_API ConstructSpectralConformalParamSolver(const FDynamicMesh3& DynamicMesh, bool bPreserveIrregularity);
 	}
 }
 
diff --git a/Engine/Plugins/Runtime/MeshModelingToolset/Source/ModelingOperators/Private/ParameterizationOps/RecomputeUVsOp.cpp b/Engine/Plugins/Runtime/MeshModelingToolset/Source/ModelingOperators/Private/ParameterizationOps/RecomputeUVsOp.cpp
--- a/Engine/Plugins/Runtime/MeshModelingToolset/Source/ModelingOperators/Private/ParameterizationOps/RecomputeUVsOp.cpp
+++ b/Engine/Plugins/Runtime/MeshModelingToolset/Source/ModelingOperators/Private/ParameterizationOps/RecomputeUVsOp.cpp
@@ -158,23 +158,36 @@ bool FRecomputeUVsOp::CalculateResult_Basic(FProgressCancel* Progress)
 		bComponentSolved[k] = false;
 		switch (UnwrapType)
 		{
-		case ERecomputeUVsUnwrapType::ExpMap:
-		{
-			FDynamicMeshUVEditor::FExpMapOptions Options;
-			Options.NormalSmoothingRounds = this->NormalSmoothingRounds;
-			Options.NormalSmoothingAlpha = this->NormalSmoothingAlpha;
-			bComponentSolved[k] = UVEditor.SetTriangleUVsFromExpMap(ComponentTris, Options);
-		}
-		break;
+			case ERecomputeUVsUnwrapType::ExpMap:
+			{
+				FDynamicMeshUVEditor::FExpMapOptions Options;
+				Options.NormalSmoothingRounds = this->NormalSmoothingRounds;
+				Options.NormalSmoothingAlpha = this->NormalSmoothingAlpha;
+				bComponentSolved[k] = UVEditor.SetTriangleUVsFromExpMap(ComponentTris, Options);
+				break;
+			}
 
-		case ERecomputeUVsUnwrapType::ConformalFreeBoundary:
-			bComponentSolved[k] = UVEditor.SetTriangleUVsFromFreeBoundaryConformal(ComponentTris, bUseExisingUVTopology);
-			if (bComponentSolved[k])
+			case ERecomputeUVsUnwrapType::ConformalFreeBoundary:
+			{
+				bComponentSolved[k] = UVEditor.SetTriangleUVsFromFreeBoundaryConformal(ComponentTris, bUseExisingUVTopology);
+				if (bComponentSolved[k])
+				{
+					UVEditor.ScaleUVAreaTo3DArea(ComponentTris, true);
+				}
+				break;
+			}
+			
+			case ERecomputeUVsUnwrapType::SpectralConformal:
 			{
-				UVEditor.ScaleUVAreaTo3DArea(ComponentTris, true);
+				bComponentSolved[k] = UVEditor.SetTriangleUVsFromFreeBoundarySpectralConformal(ComponentTris, bUseExisingUVTopology, bPreserveIrregularity);
+				if (bComponentSolved[k])
+				{
+					UVEditor.ScaleUVAreaTo3DArea(ComponentTris, true);
+				}
+				break;
 			}
-			break;
 		}
+		
 
 		if (bComponentSolved[k])
 		{
@@ -401,6 +414,10 @@ TUniquePtr<FDynamicMeshOperator> URecomputeUVsOpFactory::MakeNewOperator()
 		RecomputeUVsOp->UnwrapType = ERecomputeUVsUnwrapType::ConformalFreeBoundary;
 		RecomputeUVsOp->bMergingOptimization = false;
 		break;
+	case ERecomputeUVsPropertiesUnwrapType::SpectralConformal:
+		RecomputeUVsOp->UnwrapType = ERecomputeUVsUnwrapType::SpectralConformal;
+		RecomputeUVsOp->bMergingOptimization = false;
+		break;
 	case ERecomputeUVsPropertiesUnwrapType::IslandMerging:
 		RecomputeUVsOp->UnwrapType = ERecomputeUVsUnwrapType::ExpMap;
 		RecomputeUVsOp->bMergingOptimization = true;
@@ -408,6 +425,8 @@ TUniquePtr<FDynamicMeshOperator> URecomputeUVsOpFactory::MakeNewOperator()
 
 	RecomputeUVsOp->bAutoRotate = (Settings->AutoRotation == ERecomputeUVsToolOrientationMode::MinBounds);
 
+	RecomputeUVsOp->bPreserveIrregularity = Settings->bPreserveIrregularity;
+
 	RecomputeUVsOp->NormalSmoothingRounds = Settings->SmoothingSteps;
 	RecomputeUVsOp->NormalSmoothingAlpha = Settings->SmoothingAlpha;
 
diff --git a/Engine/Plugins/Runtime/MeshModelingToolset/Source/ModelingOperators/Public/ParameterizationOps/RecomputeUVsOp.h b/Engine/Plugins/Runtime/MeshModelingToolset/Source/ModelingOperators/Public/ParameterizationOps/RecomputeUVsOp.h
--- a/Engine/Plugins/Runtime/MeshModelingToolset/Source/ModelingOperators/Public/ParameterizationOps/RecomputeUVsOp.h
+++ b/Engine/Plugins/Runtime/MeshModelingToolset/Source/ModelingOperators/Public/ParameterizationOps/RecomputeUVsOp.h
@@ -21,7 +21,8 @@ namespace Geometry
 enum class ERecomputeUVsUnwrapType
 {
 	ExpMap = 0,
-	ConformalFreeBoundary = 1
+	ConformalFreeBoundary = 1,
+	SpectralConformal = 2
 };
 
 
@@ -66,6 +67,11 @@ public:
 	ERecomputeUVsUnwrapType UnwrapType = ERecomputeUVsUnwrapType::ExpMap;
 	ERecomputeUVsIslandMode IslandMode = ERecomputeUVsIslandMode::PolyGroups;
 
+	//
+	// Conformal Map Options 
+	//
+	bool bPreserveIrregularity = false;
+	
 	// 
 	// ExpMap Options
 	//
diff --git a/Engine/Plugins/Runtime/MeshModelingToolset/Source/ModelingOperators/Public/Properties/RecomputeUVsProperties.h b/Engine/Plugins/Runtime/MeshModelingToolset/Source/ModelingOperators/Public/Properties/RecomputeUVsProperties.h
--- a/Engine/Plugins/Runtime/MeshModelingToolset/Source/ModelingOperators/Public/Properties/RecomputeUVsProperties.h
+++ b/Engine/Plugins/Runtime/MeshModelingToolset/Source/ModelingOperators/Public/Properties/RecomputeUVsProperties.h
@@ -18,8 +18,10 @@ enum class ERecomputeUVsPropertiesUnwrapType
 	ExpMap = 0 UMETA(DisplayName = "ExpMap"),
 	/** Conformal UV flattening is increasingly expensive on large islands but reduces distortion */
 	Conformal = 1,
+	/** Compared to the default Conformal method does not pin two vertices along the boundary which reduces the distortion but is more expensive to compute. */
+	SpectralConformal = 2,
 	/** UV islands will be merged into larger islands if it does not increase stretching and distortion beyond defined limits */
-	IslandMerging = 2
+	IslandMerging = 3
 };
 
 
@@ -74,12 +76,16 @@ public:
 
 	/** Type of UV flattening algorithm to use */
 	UPROPERTY(EditAnywhere, Category = "UV Unwrap")
-	ERecomputeUVsPropertiesUnwrapType UnwrapType = ERecomputeUVsPropertiesUnwrapType::Conformal;
+	ERecomputeUVsPropertiesUnwrapType UnwrapType = ERecomputeUVsPropertiesUnwrapType::SpectralConformal;
 
 	/** Type of automatic rotation applied to each UV island */
 	UPROPERTY(EditAnywhere, Category = "UV Unwrap")
 	ERecomputeUVsToolOrientationMode AutoRotation = ERecomputeUVsToolOrientationMode::MinBounds;
 
+	/**  If enabled, reduces distortion for meshes with triangles of vastly different sizes, This is only enabled if the Unwrap Type is set to Spectral Conformal. */
+	UPROPERTY(EditAnywhere, Category = "UV Unwrap", meta = (EditCondition = "UnwrapType == ERecomputeUVsPropertiesUnwrapType::SpectralConformal"))
+	bool bPreserveIrregularity = true;
+
 	/** Number of smoothing steps to apply; this slightly increases distortion but produces more stable results. This is only enabled if the Unwrap Type is set to ExpMap or Island Merging. */
 	UPROPERTY(EditAnywhere, Category = "UV Unwrap", meta = (UIMin = "0", UIMax = "25", ClampMin = "0", ClampMax = "1000",
 		EditCondition = "UnwrapType == ERecomputeUVsPropertiesUnwrapType::ExpMap || UnwrapType == ERecomputeUVsPropertiesUnwrapType::IslandMerging"))
