change FArrangement2d and TConstrainedDelaunay2 to use new FDelaunay2 code for triangulations

diff --git a/Engine/Plugins/Runtime/GeometryProcessing/Source/GeometryAlgorithms/Private/Arrangement2d.cpp b/Engine/Plugins/Runtime/GeometryProcessing/Source/GeometryAlgorithms/Private/Arrangement2d.cpp
--- a/Engine/Plugins/Runtime/GeometryProcessing/Source/GeometryAlgorithms/Private/Arrangement2d.cpp
+++ b/Engine/Plugins/Runtime/GeometryProcessing/Source/GeometryAlgorithms/Private/Arrangement2d.cpp
@@ -3,10 +3,8 @@
 // Port of geometry3Sharp Arrangement2d
 
 #include "Arrangement2d.h"
-#include "ThirdParty/GTEngine/Mathematics/GteConstrainedDelaunay2.h"
-#include "ThirdParty/GTEngine/Mathematics/GteBSNumber.h"
-#include "ThirdParty/GTEngine/Mathematics/GteUIntegerFP32.h"
-#include <vector>
+
+#include "CompGeom/Delaunay2.h"
 
 using namespace UE::Geometry;
 
@@ -85,18 +83,26 @@ namespace
 #endif
 }
 
-bool FArrangement2d::AttemptTriangulate(TArray<FIntVector> &Triangles, TArray<int32> &SkippedEdges, int32 BoundaryEdgeGroupID)
+bool FArrangement2d::AttemptTriangulate(TArray<FIntVector>& Triangles, TArray<int32>& SkippedEdges, int32 BoundaryEdgeGroupID)
+{
+	// TODO: Currently this just constructs an FIndex3i array and then copies it to an FIntVector array; if we need to construct FIntVector versions faster, we could directly construct them here instead
+	// However this would require writing FIntVector versions of the GetTriangles functions in FDelaunay2; ideally we would instead mainly call the FIndex3i variant
+	TArray<FIndex3i> TrianglesInd3i;
+	bool bResult = AttemptTriangulate(TrianglesInd3i, SkippedEdges, BoundaryEdgeGroupID);
+	Triangles.SetNumUninitialized(TrianglesInd3i.Num());
+	for (int32 Idx = 0; Idx < TrianglesInd3i.Num(); Idx++)
+	{
+		Triangles[Idx] = (FIntVector)TrianglesInd3i[Idx];
+	}
+	return bResult;
+}
+
+bool FArrangement2d::AttemptTriangulate(TArray<FIndex3i> &Triangles, TArray<int32> &SkippedEdges, int32 BoundaryEdgeGroupID)
 {
 	Triangles.Empty();
 
-	// Use this declaration for the non-robust version of the algorithm; this is very non-robust, and not recommended!
-	//gte::ConstrainedDelaunay2<double, double> Delaunay;
-
-	// Here are the types you could use for robust math (in theory makes the algorithm overall robust; in practice I am not sure)
-	gte::ConstrainedDelaunay2<double, gte::BSNumber<gte::UIntegerFP32<263>>> Delaunay; // Value of 263 is from comment in GTEngine/Mathematics/GteDelaunay2.h
-	//gte::ConstrainedDelaunay2<double, gte::BSNumber<gte::UIntegerAP32>> Delaunay; // Full arbitrary precision (slowest method, creates a std::vector per number)
-	
-	std::vector<gte::Vector2<double>> InputVertices;
+	// A flat array of the vertices, copied out of the graph
+	TArray<FVector2d> InputVertices;
 
 	// If there are unused vertices, define a vertex index remapping so the delaunay code won't have to understand that
 	bool bNeedsRemap = Graph.MaxVertexID() != Graph.VertexCount();
@@ -113,7 +119,7 @@ bool FArrangement2d::AttemptTriangulate(TArray<FIntVector> &Triangles, TArray<in
 				OutputIndices[i - TotalOffset] = i;
 				InputIndices[i] = i - TotalOffset;
 				FVector2d Vertex = Graph.GetVertex(i);
-				InputVertices.push_back(gte::Vector2<double> {{Vertex.X, Vertex.Y}});
+				InputVertices.Add(Vertex);
 			}
 			else
 			{
@@ -127,18 +133,26 @@ bool FArrangement2d::AttemptTriangulate(TArray<FIntVector> &Triangles, TArray<in
 		for (int i = 0; i < Graph.MaxVertexID(); i++)
 		{
 			FVector2d Vertex = Graph.GetVertex(i);
-			InputVertices.push_back(gte::Vector2<double> {{Vertex.X, Vertex.Y}});
+			InputVertices.Add(Vertex);
 		}
 	}
-	if (!Delaunay(Graph.VertexCount(), &InputVertices[0], 0.001f))
+
+	FDelaunay2 Delaunay;
+	Delaunay.bAutomaticallyFixEdgesToDuplicateVertices = false; // Arrangement will remove duplicates already
+
+	if (!Delaunay.Triangulate(InputVertices))
 	{
 		return false;
 	}
-	std::vector<int> OutEdges;
-	bool InsertConstraintFailure = false;
-	TSet<TPair<int, int>> BoundarySet; // tracks all the boundary edges as they are added, so we can later eat what's outside them
+
+	Delaunay.bValidateEdges = false;
+	Delaunay.bKeepFastEdgeAdjacencyData = true;
+
+	bool bInsertConstraintFailure = false;
 	bool bBoundaryTrackingFailure = false;
 
+	TArray<FIndex2i> AllEdges, BoundaryEdges;
+
 	for (int EdgeIdx : Graph.EdgeIndices())
 	{
 		FDynamicGraph::FEdge Edge = Graph.GetEdge(EdgeIdx);
@@ -147,108 +161,54 @@ bool FArrangement2d::AttemptTriangulate(TArray<FIntVector> &Triangles, TArray<in
 			Edge.A = InputIndices[Edge.A];
 			Edge.B = InputIndices[Edge.B];
 		}
-		if (!Delaunay.Insert({{Edge.A, Edge.B}}, OutEdges))
+		FIndex2i& AddedEdge = AllEdges.Emplace_GetRef(Edge.A, Edge.B);
+		if (Edge.Group == BoundaryEdgeGroupID)
 		{
-			// Note the failed edge; we will try to proceed anyway, just without this edge.  Hopefully the CDT is robust and this never happens!
-			ensureMsgf(false, TEXT("CDT edge insertion failed"));
-			InsertConstraintFailure = true;
-			SkippedEdges.Add(EdgeIdx);
-			// specially note if we failed to add a boundary edge; in this case removing the outside-boundary triangles would likely delete all triangles from the output
-			// in this case I choose (below) to just add all triangles to the output, ignoring boundaries, as this may create less noticeable artifacts in the typical cases
-			if (Edge.Group == BoundaryEdgeGroupID)
-			{
-				bBoundaryTrackingFailure = true;
-			}
-		}
-		else
-		{
-			if (Edge.Group == BoundaryEdgeGroupID)
-			{
-				for (size_t i = 0; i + 1 < OutEdges.size(); i++)
-				{
-					int MinV = FMath::Min(OutEdges[i], OutEdges[i + 1]);
-					int MaxV = FMath::Max(OutEdges[i], OutEdges[i + 1]);
-					BoundarySet.Add(TPair<int, int>(MinV, MaxV));
-				}
-			}
+			BoundaryEdges.Add(AddedEdge);
 		}
 	}
-	const std::vector<int>& Indices = Delaunay.GetIndices();
-	if (!bBoundaryTrackingFailure && BoundarySet.Num() > 0) {
-		auto IsBoundaryEdge = [&BoundarySet](int V0, int V1)
-		{
-			int MinV = FMath::Min(V0, V1);
-			int MaxV = FMath::Max(V0, V1);
-			return BoundarySet.Contains(TPair<int, int>(MinV, MaxV));
-		};
-
-		// eat hull to boundary
-		const std::vector<int>& Adj = Delaunay.GetAdjacencies();
-		int TriNum = int(Adj.size() / 3);
-		TArray<bool> Eaten;
-		Eaten.SetNumZeroed(TriNum);
-
-		TArray<int> ToEatQ;
-		// seed the feed queue with any triangles that are on hull but w/ an edge that is not on intended boundary
-		for (int TriIdx = 0; TriIdx < TriNum; TriIdx++)
-		{
-			int BaseIdx = TriIdx * 3;
-			bool bEatIt = false;
-			for (int SubIdx = 0, NextIdx = 2; !bEatIt && SubIdx < 3; NextIdx = SubIdx++)
-			{
-				// eat the triangle if it has a hull edge that is not a boundary edge
-				bEatIt = Adj[BaseIdx + NextIdx] == -1 && !IsBoundaryEdge(Indices[BaseIdx + SubIdx], Indices[BaseIdx + NextIdx]);
-			}
-			if (bEatIt)
-			{
-				ToEatQ.Add(TriIdx);
-				Eaten[TriIdx] = true;
-			}
-		}
 
-		// eat any triangle that we can get to by crossing a non-boundary edge from an already-been-eaten triangle
-		while (ToEatQ.Num())
+	Delaunay.ConstrainEdges(InputVertices, AllEdges);
+
+	// Verify all edges after all constraints are in -- to ensure that inserted edges were also not removed by subsequent edge insertion
+	for (int EdgeIdx : Graph.EdgeIndices())
+	{
+		FDynamicGraph::FEdge Edge = Graph.GetEdge(EdgeIdx);
+		if (bNeedsRemap)
 		{
-			int EatTri = ToEatQ.Pop();
-			int BaseIdx = EatTri * 3;
-			for (int SubIdx = 0, NextIdx = 2; SubIdx < 3; NextIdx = SubIdx++)
-			{
-				int AdjTri = Adj[BaseIdx + NextIdx];
-				if (AdjTri != -1 && !Eaten[AdjTri] && !IsBoundaryEdge(Indices[BaseIdx + SubIdx], Indices[BaseIdx + NextIdx]))
-				{
-					ToEatQ.Add(AdjTri);
-					Eaten[AdjTri] = true;
-				}
-			}
+			Edge.A = InputIndices[Edge.A];
+			Edge.B = InputIndices[Edge.B];
 		}
-
-		// copy uneaten triangles
-		for (int i = 0; i < TriNum; i++)
+		if (!Delaunay.HasEdge(FIndex2i(Edge.A, Edge.B), false))
 		{
-			if (!Eaten[i])
+			bInsertConstraintFailure = true;
+			SkippedEdges.Add(EdgeIdx);
+			if (Edge.Group == BoundaryEdgeGroupID)
 			{
-				Triangles.Add(FIntVector(Indices[i * 3], Indices[i * 3 + 1], Indices[i * 3 + 2]));
+				bBoundaryTrackingFailure = true;
 			}
 		}
+	}
+	
+	if (!bBoundaryTrackingFailure && BoundaryEdges.Num() > 0)
+	{
+		Triangles = Delaunay.GetFilledTriangles(BoundaryEdges, FDelaunay2::EFillMode::Solid);
 		ensure(Triangles.Num()); // technically it could be valid to not have any triangles after eating the outside-boundary ones, but it would be unusual and could also be a bug
 	}
 	else
 	{
-		// copy all triangles over
-		for (size_t i = 0, n = Indices.size() / 3; i < n; i++)
-		{
-			Triangles.Add(FIntVector(Indices[i * 3], Indices[i * 3 + 1], Indices[i * 3 + 2]));
-		}
+		Triangles = Delaunay.GetTriangles();
 	}
+
 	if (bNeedsRemap)
 	{
-		for (FIntVector& Face : Triangles)
+		for (FIndex3i& Face : Triangles)
 		{
-			Face.X = OutputIndices[Face.X];
-			Face.Y = OutputIndices[Face.Y];
-			Face.Z = OutputIndices[Face.Z];
+			Face.A = OutputIndices[Face.A];
+			Face.B = OutputIndices[Face.B];
+			Face.C = OutputIndices[Face.C];
 		}
 	}
-	return !InsertConstraintFailure;
+	return !bInsertConstraintFailure;
 }
 
diff --git a/Engine/Plugins/Runtime/GeometryProcessing/Source/GeometryAlgorithms/Private/ConstrainedDelaunay2.cpp b/Engine/Plugins/Runtime/GeometryProcessing/Source/GeometryAlgorithms/Private/ConstrainedDelaunay2.cpp
--- a/Engine/Plugins/Runtime/GeometryProcessing/Source/GeometryAlgorithms/Private/ConstrainedDelaunay2.cpp
+++ b/Engine/Plugins/Runtime/GeometryProcessing/Source/GeometryAlgorithms/Private/ConstrainedDelaunay2.cpp
@@ -1,11 +1,9 @@
 // Copyright Epic Games, Inc. All Rights Reserved.
 
 #include "ConstrainedDelaunay2.h"
-#include "ThirdParty/GTEngine/Mathematics/GteConstrainedDelaunay2.h"
-#include "ThirdParty/GTEngine/Mathematics/GteBSNumber.h"
-#include "ThirdParty/GTEngine/Mathematics/GteUIntegerFP32.h"
+#include "CompGeom/Delaunay2.h"
 #include "Async/ParallelFor.h"
-#include <vector>
+
 
 using namespace UE::Geometry;
 
@@ -119,15 +117,15 @@ bool HasUnorderedEdge(const TMap<TPair<int, int>, bool>& EdgeMap, int VertA, int
 namespace ConstrainedDelaunay2Internal
 {
 	template<class RealType>
-	void SplitBowtiesHelper(TArray<TPair<int, int>>& NeedUpdates, TArray<TVector2<RealType>>& Vertices, TArray<int8>& Keep, const std::vector<int>& Indices, const std::vector<int>& Adj)
+	void SplitBowtiesHelper(TArray<TPair<int, int>>& NeedUpdates, TArray<TVector2<RealType>>& Vertices, TArray<int8>& Keep, const TArray<FIndex3i>& Indices, const TArray<FIndex3i>& Adj)
 	{
-		int TriNum = int(Adj.size() / 3);
+		int TriNum = Adj.Num();
 		int32 OrigNumVertices = Vertices.Num();
 		// track all wedge verts that are seen by walking local tris
 		auto OtherEdgeOnTri = [&Indices](int VertID, int TriID, int EdgeIdx)
 		{
 			int StepNext = 1;
-			if (Indices[TriID * 3 + EdgeIdx] == VertID) {
+			if (Indices[TriID][EdgeIdx] == VertID) {
 				StepNext = 2;
 			}
 			return (EdgeIdx + StepNext) % 3;
@@ -137,7 +135,7 @@ namespace ConstrainedDelaunay2Internal
 		{
 			for (int AdjEdgeIdx = 0; AdjEdgeIdx < 3; AdjEdgeIdx++)
 			{
-				if (Adj[ToTriID * 3 + AdjEdgeIdx] == FromTriID)
+				if (Adj[ToTriID][AdjEdgeIdx] == FromTriID)
 				{
 					return AdjEdgeIdx;
 				}
@@ -149,7 +147,7 @@ namespace ConstrainedDelaunay2Internal
 		{
 			for (int VertSubIdx = 0; VertSubIdx < 3; VertSubIdx++)
 			{
-				if (Indices[TriID * 3 + VertSubIdx] == VertID)
+				if (Indices[TriID][VertSubIdx] == VertID)
 				{
 					return VertSubIdx;
 				}
@@ -160,7 +158,7 @@ namespace ConstrainedDelaunay2Internal
 		auto Walk = [&Adj, &OtherEdgeOnTri](int VertID, int TriID, int EdgeSubIdx)
 		{
 			int OtherEdgeSubIdx = OtherEdgeOnTri(VertID, TriID, EdgeSubIdx);
-			int AdjTri = Adj[TriID * 3 + OtherEdgeSubIdx];
+			int AdjTri = Adj[TriID][OtherEdgeSubIdx];
 			return AdjTri;
 		};
 		TArray<bool> Seen, SeenSource; Seen.SetNumZeroed(TriNum * 3); SeenSource.SetNumZeroed(Vertices.Num());
@@ -173,7 +171,7 @@ namespace ConstrainedDelaunay2Internal
 			for (int SubIdx = 0, OtherSubIdx = 2; SubIdx < 3; OtherSubIdx = SubIdx++)
 			{
 				int WedgeIdx = TriID * 3 + SubIdx;
-				int VertID = Indices[WedgeIdx];
+				int VertID = Indices[TriID][SubIdx];
 
 				if (Seen[WedgeIdx]) // already been walked over & therefore covered by previous pass
 				{
@@ -201,7 +199,7 @@ namespace ConstrainedDelaunay2Internal
 						int VertSubIdx = GetVertSubIdx(WalkVertID, WalkTriID);
 						int WalkWedgeIdx = WalkTriID * 3 + VertSubIdx;
 						ensure(!Seen[WalkWedgeIdx]);
-						checkSlow(Indices[WalkWedgeIdx] == WalkVertID);
+						checkSlow(Indices[WalkTriID][VertSubIdx] == WalkVertID);
 						Seen[WalkWedgeIdx] = true;
 						if (bSeenSource)
 						{
@@ -239,19 +237,19 @@ namespace ConstrainedDelaunay2Internal
 		}
 	}
 
-	void BuildFinalTriangles(TArray<FIndex3i>& Triangles, TArray<TPair<int, int>>& NeedUpdates, int& AddedVerticesStartIndex, TArray<int8>& Keep, const std::vector<int>& Indices, int OrigNumVertices, bool bOutputCCW)
+	void BuildFinalTriangles(TArray<FIndex3i>& Triangles, TArray<TPair<int, int>>& NeedUpdates, int& AddedVerticesStartIndex, TArray<int8>& Keep, const TArray<FIndex3i>& Indices, int OrigNumVertices, bool bOutputCCW)
 	{
-		int TriNum = int(Indices.size() / 3);
+		int TriNum = Indices.Num();
 
 		// function to build output triangles out of an indices array
 		// normally called directly on the const indices from the CDT, but will be called on an updated copy if bowtie splits happen
-		auto BuildTriangles = [&Triangles, &TriNum, &Keep, &bOutputCCW](const std::vector<int>& IndicesIn)
+		auto BuildTriangles = [&Triangles, &TriNum, &Keep, &bOutputCCW](const TArray<FIndex3i>& IndicesIn)
 		{
 			for (int i = 0; i < TriNum; i++)
 			{
 				if (Keep[i] > 0)
 				{
-					FIndex3i& Tri = Triangles.Emplace_GetRef(IndicesIn[i * 3], IndicesIn[i * 3 + 1], IndicesIn[i * 3 + 2]);
+					FIndex3i& Tri = Triangles.Emplace_GetRef(IndicesIn[i].A, IndicesIn[i].B, IndicesIn[i].C);
 					if (!bOutputCCW)
 					{
 						Swap(Tri.B, Tri.C);
@@ -261,10 +259,12 @@ namespace ConstrainedDelaunay2Internal
 		};
 		if (NeedUpdates.Num() > 0)
 		{
-			std::vector<int> UpdatedIndices = Indices;
+			TArray<FIndex3i> UpdatedIndices = Indices;
 			for (const TPair<int, int>& Update : NeedUpdates)
 			{
-				UpdatedIndices[Update.Key] = Update.Value;
+				int TriID = (int)Update.Key / 3;
+				int SubIdx = Update.Key % 3;
+				UpdatedIndices[TriID][SubIdx] = Update.Value;
 			}
 
 			AddedVerticesStartIndex = OrigNumVertices;
@@ -284,51 +284,27 @@ bool TConstrainedDelaunay2<RealType>::Triangulate(TFunctionRef<bool(const TArray
 {
 	Triangles.Empty();
 
-	gte::ConstrainedDelaunay2<double, gte::BSNumber<gte::UIntegerFP32<263>>> Delaunay; // Value of 263 is from comment in GTEngine/Mathematics/GteDelaunay2.h
-
-	std::vector<gte::Vector2<double>> InputVertices;
-	for (int i = 0; i < Vertices.Num(); i++)
-	{
-		InputVertices.push_back(gte::Vector2<double> {{Vertices[i].X, Vertices[i].Y}});
-	}
-
-	if (!Delaunay(Vertices.Num(), &InputVertices[0], 0))
+	FDelaunay2 Delaunay;
+	Delaunay.bAutomaticallyFixEdgesToDuplicateVertices = true;
+	if (!Delaunay.Triangulate(Vertices))
 	{
 		return false;
 	}
+	Delaunay.bKeepFastEdgeAdjacencyData = true;
+	bool bEdgesFailed = !Delaunay.ConstrainEdges(Vertices, Edges);
+	bool bHoleEdgesFailed = !Delaunay.ConstrainEdges(Vertices, HoleEdges);
+	bool bBoundaryTrackingFailure = bEdgesFailed || bHoleEdgesFailed;
 
+	TArray<FIndex3i> Indices, Adj;
+	Delaunay.GetTrianglesAndAdjacency(Indices, Adj);
 
-	const std::vector<int>& Duplicates = Delaunay.GetDuplicates();
-
-	std::vector<int> OutEdges;
-	bool bBoundaryTrackingFailure = false;
-
-	TArray<FIndex2i>* InputEdgesAndHoles[2] = { &Edges, &HoleEdges };
-	for (int EdgeOrHole = 0; EdgeOrHole < 2; EdgeOrHole++)
-	{
-		TArray<FIndex2i>& Input = *InputEdgesAndHoles[EdgeOrHole];
-		for (const FIndex2i& Edge : Input)
-		{
-			int A = Duplicates[Edge.A];
-			int B = Duplicates[Edge.B];
-			if (!Delaunay.Insert({{A, B}}, OutEdges))
-			{
-				// Note the failed edge; we will try to proceed anyway, just without this edge.  This can happen for example if the edge is exactly on top of another edge
-				bBoundaryTrackingFailure = true;
-			}
-		}
-	}
-
-	const std::vector<int>& Indices = Delaunay.GetIndices();
-	const std::vector<int>& Adj = Delaunay.GetAdjacencies();
-	int TriNum = int(Adj.size() / 3);
+	int TriNum = Adj.Num();
 	TArray<int8> Keep;  // values: 0->unprocessed (delete), 1->yes keep, -1->processed, delete
 	Keep.SetNumZeroed(TriNum);
 
 	ParallelFor(TriNum, [this, &KeepTriangle, &Indices, &Keep](int32 Index)
 	{
-		int32 TriIndStart = Index * 3;
-		bool bKeepTri = KeepTriangle(Vertices, FIndex3i(Indices[TriIndStart], Indices[TriIndStart + 1], Indices[TriIndStart + 2]));
+		bool bKeepTri = KeepTriangle(Vertices, Indices[Index]);
 		Keep[Index] = bKeepTri ? 1 : -1;
 	});
 
@@ -351,22 +327,18 @@ bool TConstrainedDelaunay2<RealType>::Triangulate()
 
 	check(FillRule <= EFillRule::Odd || bOrientedEdges);
 
-	gte::ConstrainedDelaunay2<double, gte::BSNumber<gte::UIntegerFP32<263>>> Delaunay; // Value of 263 is from comment in GTEngine/Mathematics/GteDelaunay2.h
-
-	std::vector<gte::Vector2<double>> InputVertices;
-	for (int i = 0; i < Vertices.Num(); i++)
-	{
-		InputVertices.push_back(gte::Vector2<double> {{Vertices[i].X, Vertices[i].Y}});
-	}
-
-	if (!Delaunay(Vertices.Num(), &InputVertices[0], 0))
+	FDelaunay2 Delaunay;
+	Delaunay.bAutomaticallyFixEdgesToDuplicateVertices = true;
+	if (!Delaunay.Triangulate(Vertices))
 	{
 		return false;
 	}
 
-	const std::vector<int>& Duplicates = Delaunay.GetDuplicates();
+	Delaunay.bKeepFastEdgeAdjacencyData = true;
+	Delaunay.bValidateEdges = false; // edge validation will be done manually later
+	Delaunay.ConstrainEdges(Vertices, Edges);
+	Delaunay.ConstrainEdges(Vertices, HoleEdges);
 
-	std::vector<int> OutEdges;
 	TMap<TPair<int, int>, bool> BoundaryMap, HoleMap; // tracks all the boundary edges as they are added, so we can later flood fill across them for inside/outside decisions
 	TMap<TPair<int, int>, bool>* EdgeAndHoleMaps[2] = { &BoundaryMap, &HoleMap };
 	bool bBoundaryTrackingFailure = false;
@@ -376,28 +348,25 @@ bool TConstrainedDelaunay2<RealType>::Triangulate()
 	{
 		TArray<FIndex2i>& Input = *InputEdgesAndHoles[EdgeOrHole];
 		TMap<TPair<int, int>, bool>& InputMap = *EdgeAndHoleMaps[EdgeOrHole];
-		for (const FIndex2i& Edge : Input)
+		for (FIndex2i Edge : Input)
 		{
-			int A = Duplicates[Edge.A];
-			int B = Duplicates[Edge.B];
-			if (!Delaunay.Insert({{A, B}}, OutEdges))
+			Delaunay.FixDuplicatesOnEdge(Edge);
+			if (!Delaunay.HasEdge(Edge, false))
 			{
 				// Note the failed edge; we will try to proceed anyway, just without this edge.  This can happen for example if the edge is exactly on top of another edge
 				bBoundaryTrackingFailure = true;
 			}
 			else
 			{
-				for (size_t i = 0; i + 1 < OutEdges.size(); i++)
-				{
-					AddOrderedEdge(InputMap, OutEdges[i], OutEdges[i + 1]);
-				}
+				AddOrderedEdge(InputMap, Edge.A, Edge.B);
 			}
 		}
 	}
 
-	const std::vector<int>& Indices = Delaunay.GetIndices();
-	const std::vector<int>& Adj = Delaunay.GetAdjacencies();
-	int TriNum = int(Adj.size() / 3);
+	TArray<FIndex3i> Indices, Adj;
+	Delaunay.GetTrianglesAndAdjacency(Indices, Adj);
+
+	int TriNum = Adj.Num();
 	TArray<int8> Keep;  // values: 0->unprocessed (delete), 1->yes keep, -1->processed, delete
 	Keep.SetNumZeroed(TriNum);
 
@@ -406,12 +375,11 @@ bool TConstrainedDelaunay2<RealType>::Triangulate()
 	// note: need *all* not just *one* because of the strategy of refusing to cross hole edges; if using pure winding number classification would just need one boundary triangle to start
 	for (int TriIdx = 0; TriIdx < TriNum; TriIdx++)
 	{
-		int BaseIdx = TriIdx * 3;
 		for (int SubIdx = 0, NextIdx = 2; SubIdx < 3; NextIdx = SubIdx++)
 		{
-			if (Adj[BaseIdx + NextIdx] < 0) // on hull
+			if (Adj[TriIdx][NextIdx] < 0) // on hull
 			{
-				int VertA = Indices[BaseIdx + SubIdx], VertB = Indices[BaseIdx + NextIdx];
+				int VertA = Indices[TriIdx][SubIdx], VertB = Indices[TriIdx][NextIdx];
 				if (HasUnorderedEdge(HoleMap, VertA, VertB))
 				{
 					continue; // cannot cross hole edges
@@ -431,16 +399,16 @@ bool TConstrainedDelaunay2<RealType>::Triangulate()
 		SelIdx = (SelIdx + 1) % ToWalkQ.Num();
 		TPair<int, int> TriWithWinding = ToWalkQ[SelIdx];
 		ToWalkQ.RemoveAtSwap(SelIdx);
-		int BaseIdx = TriWithWinding.Key * 3;
+		int TriIdx = TriWithWinding.Key;
 		int LastWinding = TriWithWinding.Value;
 		for (int SubIdx = 0, NextIdx = 2; SubIdx < 3; NextIdx = SubIdx++)
 		{
-			int VertA = Indices[BaseIdx + SubIdx], VertB = Indices[BaseIdx + NextIdx];
+			int VertA = Indices[TriIdx][SubIdx], VertB = Indices[TriIdx][NextIdx];
 			if (HasUnorderedEdge(HoleMap, VertA, VertB))
 			{
 				continue; // cannot cross hole edges
 			}
-			int AdjTri = Adj[BaseIdx + NextIdx];
+			int AdjTri = Adj[TriIdx][NextIdx];
 			if (AdjTri >= 0 && Keep[AdjTri] == 0)
 			{
 				int WindingChange = WindingAcross(BoundaryMap, VertA, VertB);
diff --git a/Engine/Plugins/Runtime/GeometryProcessing/Source/GeometryAlgorithms/Public/Arrangement2d.h b/Engine/Plugins/Runtime/GeometryProcessing/Source/GeometryAlgorithms/Public/Arrangement2d.h
--- a/Engine/Plugins/Runtime/GeometryProcessing/Source/GeometryAlgorithms/Public/Arrangement2d.h
+++ b/Engine/Plugins/Runtime/GeometryProcessing/Source/GeometryAlgorithms/Public/Arrangement2d.h
@@ -53,16 +53,17 @@ struct FArrangement2d
 	}
 
 	/**
-	 * Attempts to triangulates the arrangement with a constrained delaunay triangulation
-	 * NOTE: Not robust; generates invalid triangulations and fails to insert edges sometimes, even if the arrangement has no self-intersections
-	 * But should always do *something* that at least covers the point set, if the point set is not degenerate
-	 * TODO: Make a robust triangulation algo
+	 * Attempts to triangulates the arrangement with a constrained Delaunay triangulation
+	 * NOTE: May fail if arrangement has self-intersections
 	 * 
 	 * Triangles: Output triangles (as indices into Graph vertices)
 	 * SkippedEdges: Output indices of edges that the algorithm failed to insert
 	 * BoundaryEdgeGroupID: ID of edges corresponding to a boundary; if we have a closed loop of these boundary edges on output triangulation, will discard triangles outside this
-	 * return: false if triangulation algo knows it failed (NOTE: triangulation may still be invalid when function returns true, as function is not robust)
+	 * return: false if triangulation algo knows it failed (note Triangles may still have some triangulation of the input in this case; for example it may just be missing some required edges)
 	 */
+	bool GEOMETRYALGORITHMS_API AttemptTriangulate(TArray<FIndex3i>& Triangles, TArray<int32>& SkippedEdges, int32 BoundaryEdgeGroupID);
+
+	// Variant of AttemptTriangulate using FIntVector instead of FIndex3i; Note this incurs an extra copy of the triangle array
 	bool GEOMETRYALGORITHMS_API AttemptTriangulate(TArray<FIntVector>& Triangles, TArray<int32>& SkippedEdges, int32 BoundaryEdgeGroupID);
 
 	/**
diff --git a/Engine/Plugins/Runtime/GeometryProcessing/Source/GeometryAlgorithms/Public/ConstrainedDelaunay2.h b/Engine/Plugins/Runtime/GeometryProcessing/Source/GeometryAlgorithms/Public/ConstrainedDelaunay2.h
--- a/Engine/Plugins/Runtime/GeometryProcessing/Source/GeometryAlgorithms/Public/ConstrainedDelaunay2.h
+++ b/Engine/Plugins/Runtime/GeometryProcessing/Source/GeometryAlgorithms/Public/ConstrainedDelaunay2.h
@@ -16,7 +16,12 @@ namespace Geometry {
 
 using namespace UE::Math;
 
-
+/**
+ * This is a version of FDelaunay (which can also do Constrained Delaunay triangulation), but with added support for:
+ *  1. Explicit "Hole" edges that must be cut out of the result (FDelaunay relies only on winding rules to create holes)
+ *  2. Option to duplicate vertices to split any 'bowtie' vertices
+ * If you do not need either of these features, FDelaunay may be the faster option.
+ */
 template<typename RealType>
 struct TConstrainedDelaunay2
 {
diff --git a/Engine/Source/Runtime/GeometryCore/Private/CompGeom/Delaunay2.cpp b/Engine/Source/Runtime/GeometryCore/Private/CompGeom/Delaunay2.cpp
--- a/Engine/Source/Runtime/GeometryCore/Private/CompGeom/Delaunay2.cpp
+++ b/Engine/Source/Runtime/GeometryCore/Private/CompGeom/Delaunay2.cpp
@@ -1218,5 +1218,24 @@ bool FDelaunay2::HasEdges(TArrayView<const FIndex2i> Edges) const
 	return false; // if no triangulation was performed, function should not be called
 }
 
+void FDelaunay2::FixDuplicatesOnEdge(FIndex2i& Edge)
+{
+	checkSlow(Connectivity.IsValid() && Connectivity->HasDuplicateTracking());
+	Connectivity->FixDuplicatesOnEdge(Edge);
+}
+
+bool FDelaunay2::HasEdge(const FIndex2i& Edge, bool bRemapDuplicates)
+{
+	checkSlow(Connectivity.IsValid());
+	if (bRemapDuplicates)
+	{
+		checkSlow(Connectivity->HasDuplicateTracking());
+		FIndex2i RemapEdge = Edge;
+		Connectivity->FixDuplicatesOnEdge(RemapEdge);
+		return Connectivity->HasEdge(RemapEdge);
+	}
+	return Connectivity->HasEdge(Edge);
+}
+
 } // end namespace UE::Geometry
 } // end namespace UE
\ No newline at end of file
diff --git a/Engine/Source/Runtime/GeometryCore/Public/CompGeom/Delaunay2.h b/Engine/Source/Runtime/GeometryCore/Public/CompGeom/Delaunay2.h
--- a/Engine/Source/Runtime/GeometryCore/Public/CompGeom/Delaunay2.h
+++ b/Engine/Source/Runtime/GeometryCore/Public/CompGeom/Delaunay2.h
@@ -203,6 +203,12 @@ public:
 	// @return true if triangulation has the given edges, useful for validating results
 	bool HasEdges(TArrayView<const FIndex2i> Edges) const;
 
+	// Remap any references to duplicate vertices to only reference the vertices used in the triangulation
+	void FixDuplicatesOnEdge(FIndex2i& Edge);
+
+	// @return true if the edge is in the triangulation
+	bool HasEdge(const FIndex2i& Edge, bool bRemapDuplicates);
+
 protected:
 	TPimplPtr<FDelaunay2Connectivity> Connectivity;
 
