Modeling Tools Baking: Remove superfluous FindConnectedTriangles call

#rnx

diff --git a/Engine/Plugins/Runtime/GeometryProcessing/Source/DynamicMesh/Private/Sampling/MeshMapBaker.cpp b/Engine/Plugins/Runtime/GeometryProcessing/Source/DynamicMesh/Private/Sampling/MeshMapBaker.cpp
--- a/Engine/Plugins/Runtime/GeometryProcessing/Source/DynamicMesh/Private/Sampling/MeshMapBaker.cpp
+++ b/Engine/Plugins/Runtime/GeometryProcessing/Source/DynamicMesh/Private/Sampling/MeshMapBaker.cpp
@@ -709,7 +709,6 @@ void FMeshMapBaker::ComputeUVCharts(const FDynamicMesh3& Mesh, TArray<int32>& Me
 	if (const FDynamicMeshUVOverlay* UVOverlay = Mesh.Attributes() ? Mesh.Attributes()->PrimaryUV() : nullptr)
 	{
 		FMeshConnectedComponents UVComponents(&Mesh);
-		UVComponents.FindConnectedTriangles();
 		UVComponents.FindConnectedTriangles([UVOverlay](int32 Triangle0, int32 Triangle1) {
 			return UVOverlay ? UVOverlay->AreTrianglesConnected(Triangle0, Triangle1) : false;
 		});
