Modeling Tools Baking/Approximate Actors: Minor refactoring of MeshGenericWorldPositionBaker

#rnx

diff --git a/Engine/Plugins/Runtime/GeometryProcessing/Source/DynamicMesh/Private/Sampling/MeshGenericWorldPositionBaker.cpp b/Engine/Plugins/Runtime/GeometryProcessing/Source/DynamicMesh/Private/Sampling/MeshGenericWorldPositionBaker.cpp
--- a/Engine/Plugins/Runtime/GeometryProcessing/Source/DynamicMesh/Private/Sampling/MeshGenericWorldPositionBaker.cpp
+++ b/Engine/Plugins/Runtime/GeometryProcessing/Source/DynamicMesh/Private/Sampling/MeshGenericWorldPositionBaker.cpp
@@ -9,30 +9,19 @@ void FMeshGenericWorldPositionColorBaker::Bake()
 	const FMeshImageBakingCache* BakeCache = GetCache();
 	check(BakeCache);
 
-	auto SampleFunction = [&](const FMeshImageBakingCache::FCorrespondenceSample& SampleData)
-	{
-		FVector3d Position = SampleData.BaseSample.SurfacePoint;
-		FVector3d Normal = SampleData.BaseNormal;
-		int32 TriangleID = SampleData.BaseSample.TriangleIndex;
-		FVector3d BaryCoords = SampleData.BaseSample.BaryCoords;
-
-		FVector4f Color = ColorSampleFunction(Position, Normal);
-
-		return Color;
-	};
-
 	ResultBuilder = MakeUnique<TImageBuilder<FVector4f>>();
 	ResultBuilder->SetDimensions(BakeCache->GetDimensions());
 	ResultBuilder->Clear(DefaultColor);
 
-	BakeCache->EvaluateSamples([&](const FVector2i& Coords, const FMeshImageBakingCache::FCorrespondenceSample& Sample)
+	BakeCache->EvaluateSamples([this](const FVector2i& Coords, const FMeshImageBakingCache::FCorrespondenceSample& Sample)
 	{
-		FVector4f Color = SampleFunction(Sample);
+		FVector3d Position = Sample.BaseSample.SurfacePoint;
+		FVector3d Normal = Sample.BaseNormal;
+		FVector4f Color = ColorSampleFunction(Position, Normal);
+
 		ResultBuilder->SetPixel(Coords, Color);
 	});
 
-
-
 	const FImageOccupancyMap& Occupancy = *BakeCache->GetOccupancyMap();
 	for (int64 k = 0; k < Occupancy.GutterTexels.Num(); k++)
 	{
@@ -53,43 +42,39 @@ void FMeshGenericWorldPositionNormalBaker::Bake()
 	check(BakeCache);
 
 	// make sure we have per-triangle tangents computed so that GetInterpolatedTriangleTangent() below works
-	int32 MaxTriID = BakeCache->GetBakeTargetMesh()->MaxTriangleID();
 	check(BaseMeshTangents->GetTangents().Num() >= 3 * BakeCache->GetBakeTargetMesh()->MaxTriangleID());
 
-	auto SampleFunction = [&](const FMeshImageBakingCache::FCorrespondenceSample& SampleData)
-	{
-		FVector3d Position = SampleData.BaseSample.SurfacePoint;
-		FVector3d Normal = SampleData.BaseNormal;
-		int32 TriangleID = SampleData.BaseSample.TriangleIndex;
-		FVector3d BaryCoords = SampleData.BaseSample.BaryCoords;
-
-		// get tangents on base mesh
-		FVector3d BaseTangentX, BaseTangentY;
-		BaseMeshTangents->GetInterpolatedTriangleTangent(
-			SampleData.BaseSample.TriangleIndex,
-			SampleData.BaseSample.BaryCoords,
-			BaseTangentX, BaseTangentY);
-
-		// sample world normal at world position
-		FVector3f DetailNormal = NormalSampleFunction(Position, Normal);
-		Normalize(DetailNormal);
-
-		// compute normal in tangent space
-		float dx = DetailNormal.Dot((FVector3f)BaseTangentX);
-		float dy = DetailNormal.Dot((FVector3f)BaseTangentY);
-		float dz = DetailNormal.Dot((FVector3f)SampleData.BaseNormal);
-
-		return FVector3f(dx, dy, dz);
-	};
-
 	NormalsBuilder = MakeUnique<TImageBuilder<FVector3f>>();
 	NormalsBuilder->SetDimensions(BakeCache->GetDimensions());
 	FVector3f DefaultMapNormal = (DefaultNormal + FVector3f::One()) * 0.5f;
 	NormalsBuilder->Clear(DefaultMapNormal);
 
-	BakeCache->EvaluateSamples([&](const FVector2i& Coords, const FMeshImageBakingCache::FCorrespondenceSample& Sample)
+	BakeCache->EvaluateSamples([this](const FVector2i& Coords, const FMeshImageBakingCache::FCorrespondenceSample& Sample)
 	{
-		FVector3f RelativeDetailNormal = SampleFunction(Sample);
+		// Compute tangent space normal
+		FVector3f RelativeDetailNormal;
+		{
+			FVector3d BasePosition = Sample.BaseSample.SurfacePoint;
+			FVector3d BaseNormal = Sample.BaseNormal;
+			int32 TriangleID = Sample.BaseSample.TriangleIndex;
+			FVector3d BaryCoords = Sample.BaseSample.BaryCoords;
+
+			// get tangents on base mesh
+			FVector3d BaseTangentX, BaseTangentY;
+			BaseMeshTangents->GetInterpolatedTriangleTangent(TriangleID, BaryCoords, BaseTangentX, BaseTangentY);
+
+			// sample world normal at world position
+			FVector3f DetailNormal = NormalSampleFunction(BasePosition, BaseNormal);
+			Normalize(DetailNormal);
+
+			// compute normal in tangent space
+			float dx = DetailNormal.Dot((FVector3f)BaseTangentX);
+			float dy = DetailNormal.Dot((FVector3f)BaseTangentY);
+			float dz = DetailNormal.Dot((FVector3f)BaseNormal);
+			RelativeDetailNormal = FVector3f(dx, dy, dz);
+		}
+
+		// Remap normal components [-1,1] to color components [0,1]
 		FVector3f MapNormal = (RelativeDetailNormal + FVector3f::One()) * 0.5f;
 		NormalsBuilder->SetPixel(Coords, MapNormal);
 	});
