-Fixed a bug in the UV conformal unit test by allowing setting the starting vector in the SpectralConformal Mesh UV solver. This allows to deterministically generate the same result given the same initial vector.
-Added a DenseMatrix.h which contains wrappers around Eigen's DenseMatrix class.

diff --git a/Engine/Plugins/Runtime/GeometryProcessing/Source/DynamicMesh/Private/Solvers/Internal/DenseMatrix.h b/Engine/Plugins/Runtime/GeometryProcessing/Source/DynamicMesh/Private/Solvers/Internal/DenseMatrix.h
new file mode 100644
--- /dev/null
+++ b/Engine/Plugins/Runtime/GeometryProcessing/Source/DynamicMesh/Private/Solvers/Internal/DenseMatrix.h
@@ -0,0 +1,69 @@
+// Copyright Epic Games, Inc. All Rights Reserved.
+
+#pragma once
+
+#ifndef EIGEN_MPL2_ONLY
+#define EIGEN_MPL2_ONLY
+#endif
+
+#if defined(_MSC_VER) && USING_CODE_ANALYSIS
+#pragma warning(push)
+#pragma warning(disable : 6011)
+#pragma warning(disable : 6387)
+#pragma warning(disable : 6313)
+#pragma warning(disable : 6294)
+#endif
+PRAGMA_DEFAULT_VISIBILITY_START
+THIRD_PARTY_INCLUDES_START
+#include <Eigen/Dense>
+THIRD_PARTY_INCLUDES_END
+PRAGMA_DEFAULT_VISIBILITY_END
+#if defined(_MSC_VER) && USING_CODE_ANALYSIS
+#pragma warning(pop)
+#endif
+
+#include "Math/RandomStream.h"
+
+namespace UE
+{
+namespace Geometry
+{
+
+// Template wrapper around Eigen's Dense Matrix class 
+template <typename ScalarType, int RowsAtCompileTime, int ColsAtCompileTime>
+using TDenseMatrix = Eigen::Matrix<ScalarType, RowsAtCompileTime, ColsAtCompileTime>;
+
+// Dense matrix of arbitrary dimensions 
+template <typename ScalarType>
+using TDenseMatrixX = TDenseMatrix<ScalarType, Eigen::Dynamic, Eigen::Dynamic>;
+using FDenseMatrixD = TDenseMatrixX<double>;
+using FDenseMatrixF = TDenseMatrixX<float>;
+using FDenseMatrixI = TDenseMatrixX<int32>;
+
+// Dense matrix representing a column vector
+template <class ScalarType>
+using TColumnVector = TDenseMatrix<ScalarType, Eigen::Dynamic, 1>;
+using FColumnVectorD = TColumnVector<double>;
+using FColumnVectorF = TColumnVector<float>;
+using FColumnVectorI = TColumnVector<int32>;
+
+/**
+ * Initialize a Dense Matrix with random values from the stream.
+ * 
+ * @param OutMatrix        Matrix will be resized and populated with the random values from the stream.
+ * @param InRowsAtRunTime  Number of rows in the output matrix.
+ * @param InColsAtRunTime  Number of columns in the output matrix.
+ * @param InRandomStream   Random stream used for fetching random values.
+ */
+template <typename ScalarType, int RowsAtCompileTime, int ColsAtCompileTime>
+void RandomDenseMatrix(TDenseMatrix<ScalarType, RowsAtCompileTime, ColsAtCompileTime>& OutMatrix, 
+					   const int32 InRowsAtRunTime,
+					   const int32 InColsAtRunTime,
+					   const FRandomStream& InRandomStream) 
+{
+	auto RandNumFunc = [&InRandomStream]() { return InRandomStream.FRandRange(-1.0f, 1.0f); };
+	OutMatrix = TDenseMatrix<ScalarType, RowsAtCompileTime, ColsAtCompileTime>::NullaryExpr(InRowsAtRunTime, InColsAtRunTime, RandNumFunc);
+}
+
+} // end namespace UE::Geometry
+} // end namespace UE
\ No newline at end of file
diff --git a/Engine/Plugins/Runtime/GeometryProcessing/Source/DynamicMesh/Private/Solvers/Internal/MeshUVSolver.cpp b/Engine/Plugins/Runtime/GeometryProcessing/Source/DynamicMesh/Private/Solvers/Internal/MeshUVSolver.cpp
--- a/Engine/Plugins/Runtime/GeometryProcessing/Source/DynamicMesh/Private/Solvers/Internal/MeshUVSolver.cpp
+++ b/Engine/Plugins/Runtime/GeometryProcessing/Source/DynamicMesh/Private/Solvers/Internal/MeshUVSolver.cpp
@@ -499,6 +499,30 @@ bool FSpectralConformalMeshUVSolver::SolveParameterization(const FSparseMatrixD&
 	FPowerMethod::Parameters Parms;
 	Parms.Tolerance = 1e-10;
 	
+	// check if we need to set a custom starting vector
+	if (InitialUVs.IsSet() && ensure(InitialUVs.GetValue().Num() == Mesh.MaxVertexID())) 
+	{
+		const int32 NumVerts = VtxLinearization.NumVerts();
+
+		FPowerMethod::RealVectorType InitialSolution;
+		InitialSolution.resize(2*NumVerts);
+		
+		for (int32 Idx = 0; Idx < NumVerts; ++Idx)
+		{
+			const int32 VertexID = VtxLinearization.GetId(Idx);
+			InitialSolution(Idx) = InitialUVs.GetValue()[VertexID].X;
+			InitialSolution(Idx + NumVerts) = InitialUVs.GetValue()[VertexID].Y;
+		}
+		Parms.InitialSolution = InitialSolution;
+	}
+	else if (RandomStream.IsSet()) 
+	{
+		const int32 NumVerts = VtxLinearization.NumVerts();
+		FPowerMethod::RealVectorType InitialSolution;
+        RandomDenseMatrix(InitialSolution, 2*NumVerts, 1, RandomStream.GetValue());
+		Parms.InitialSolution = InitialSolution;
+	}
+	
 	FPowerMethod::ScalarType EigenValue;
 	FPowerMethod::RealVectorType EigenVector;
 	TUniquePtr<FPowerMethod> PMSolverPtr = MakeUnique<FSparsePowerMethod>(SystemMatrixPSD, 
@@ -564,3 +588,15 @@ bool FSpectralConformalMeshUVSolver::SolveUVs(const FDynamicMesh3* InMeshPtr, TA
 
 	return bSolveOK;;
 }
+
+void FSpectralConformalMeshUVSolver::SetInitialSolution(const TArray<FVector2d>& InInitialUVs) 
+{
+	RandomStream.Reset();
+	InitialUVs = InInitialUVs;
+}
+
+void FSpectralConformalMeshUVSolver::SetInitialSolution(const FRandomStream& InRandomStream) 
+{
+	InitialUVs.Reset();
+	RandomStream = InRandomStream;
+}
\ No newline at end of file
diff --git a/Engine/Plugins/Runtime/GeometryProcessing/Source/DynamicMesh/Private/Solvers/Internal/MeshUVSolver.h b/Engine/Plugins/Runtime/GeometryProcessing/Source/DynamicMesh/Private/Solvers/Internal/MeshUVSolver.h
--- a/Engine/Plugins/Runtime/GeometryProcessing/Source/DynamicMesh/Private/Solvers/Internal/MeshUVSolver.h
+++ b/Engine/Plugins/Runtime/GeometryProcessing/Source/DynamicMesh/Private/Solvers/Internal/MeshUVSolver.h
@@ -7,8 +7,10 @@
 #include "Solvers/MatrixInterfaces.h"
 #include "Solvers/MeshLinearization.h"
 #include "Solvers/ConstrainedMeshSolver.h"
-
+#include "Math/RandomStream.h"
 #include "MatrixBase.h"
+#include "DenseMatrix.h"
+#include "FSparseMatrixD.h"
 
 namespace UE
 {
@@ -29,19 +31,19 @@ public:
 	{
 	};
 
-	// Add constraint associated with given vertex id.  Boundary vertices will be ignored
+	// Add constraint associated with given vertex id.
 	void AddConstraint(const int32 VtxId, const double Weight, const FVector2d& Pos, const bool bPostFix) override;
 
-	// Update the position of an existing constraint.  Bool return if a corresponding constraint weight exists. Boundary vertices will be ignored (and return false).  
+	// Update the position of an existing constraint. Bool return if a corresponding constraint weight exists.
 	bool UpdateConstraintPosition(const int32 VtxId, const FVector2d& Position, const bool bPostFix) override;
 
-	// The underlying solver will have to refactor the matrix if this is done. Bool return if a corresponding constraint position exists. Boundary vertices will be ignored (and return false).  
+	// The underlying solver will have to refactor the matrix if this is done. Bool return if a corresponding constraint position exists.
 	bool UpdateConstraintWeight(const int32 VtxId, const double Weight) override;
 
-	// Clear all constraints associated with this smoother
+	// Clear all constraints associated with this uv solver.
 	void ClearConstraints() override;
 
-	// Test if for constraint associated with given vertex id. Will return false for any boundary vert.
+	// Test if for constraint associated with given vertex id. 
 	bool IsConstrained(const int32 VtxId) const override;
 
 protected:
@@ -99,8 +101,6 @@ public:
 										FSparseMatrixD* OutAreaMatrix = nullptr);
 protected:
 
-	using FColumnVectorD = Eigen::Matrix<FSparseMatrixD::Scalar, Eigen::Dynamic, 1>;
-	
 	/**
 	 * Construct the 2Vx2V (V is the number of vertices) symmetric vector area matrix A. Given the vector X = [u;v] of 
 	 * the stacked UV coordinates, we can compute the signed 2D area of the parameterization using the quadratic form 
@@ -212,6 +212,35 @@ public:
 
 	virtual bool SolveUVs(const FDynamicMesh3* InMesh, TArray<FVector2d>& OutUVBuffer) override;
 
+	
+	//
+	// Debug utility functions
+	//
+public:
+
+	/**
+	 * Spectral conformal solver uses an iterative method to compute the final UVs. As most iterative methods 
+	 * it requires the initial guess of the solution. By default, we randomly set UV values between [-1, 1] and 
+	 * then converge to the correct solution. 
+	 *  
+	 * However, the solution is rotation invariant and it can be different given different initial guesses. Hence, 
+	 * to help make this method more deterministic, you can manually specify the initial per-vertex UVs. 
+	 * Same initial values will always produce the same solution.
+	 * 
+	 * @param InInitialUVs vector of UV coordinates for every vertex in the mesh.
+	 * 
+	 * @note resets RandomStream optional variable to avoid confusion as to which one to use.
+	 */ 
+	void SetInitialSolution(const TArray<FVector2d>& InInitialUVs);
+ 
+	/**
+	 * @param InRandomStream  The stream will be used to sample random values for the initial UV solution. 
+	 * 					      You can set the desired seed value in the stream.
+	 *
+	 * @note resets InitialUVs optional variable to avoid confusion as to which one to use.
+	 */
+	void SetInitialSolution(const FRandomStream& InRandomStream);
+
 protected: 
 
 	/** 
@@ -242,6 +271,10 @@ protected:
 	// The mesh can contain triangles of variable sizes. By default this difference in size is not considered during 
 	// the UV solve. Enabling this allows to preserve the triangle size irregularity in the resulting parameterization.
 	bool bPreserveIrregularity = false;
+
+	// Optional debug varibles used for setting the initial UV values. 
+	TOptional<TArray<FVector2d>> InitialUVs; 
+	TOptional<FRandomStream> RandomStream;
 };
 
 
diff --git a/Engine/Plugins/Runtime/GeometryProcessing/Source/DynamicMesh/Private/Solvers/Internal/PowerMethodSolver.cpp b/Engine/Plugins/Runtime/GeometryProcessing/Source/DynamicMesh/Private/Solvers/Internal/PowerMethodSolver.cpp
--- a/Engine/Plugins/Runtime/GeometryProcessing/Source/DynamicMesh/Private/Solvers/Internal/PowerMethodSolver.cpp
+++ b/Engine/Plugins/Runtime/GeometryProcessing/Source/DynamicMesh/Private/Solvers/Internal/PowerMethodSolver.cpp
@@ -1,12 +1,13 @@
 // Copyright Epic Games, Inc. All Rights Reserved.
 
 #include "PowerMethodSolver.h"
+#include "DenseMatrix.h"
 
 using namespace UE::Geometry;
 
 FPowerMethod::ESolveStatus FPowerMethod::Solve(ScalarType& OutEigenValue, RealVectorType& OutEigenVector, bool bComputeLargest)
 {
-    if (Verify(bComputeLargest) == false) 
+    if (VerifyUserParameters(bComputeLargest) == false) 
     {
         return ESolveStatus::InvalidParameters;
     }
@@ -18,8 +19,17 @@ FPowerMethod::ESolveStatus FPowerMethod::Solve(ScalarType& OutEigenValue, RealVe
     }
 
     // Initial guess for the eigenvector
-    int32 Rows = OpMatrixA.Rows();
-    OutEigenVector = RealVectorType::Random(Rows, 1); 
+    if (Parms.InitialSolution.IsSet())
+    {
+        OutEigenVector = Parms.InitialSolution.GetValue();
+    }
+    else 
+    {
+        const int32 Rows = OpMatrixA.Rows();
+        FRandomStream Stream;
+        Stream.GenerateNewSeed();
+        RandomDenseMatrix(OutEigenVector, Rows, 1, Stream); 
+    }
 
     bool bConverged = false;
     for (IterationNum = 0; IterationNum < Parms.MaxIterations; ++IterationNum) 
@@ -181,7 +191,7 @@ bool FPowerMethod::Iteration(RealVectorType& EVector, bool bComputeLargest)
     return true;
 }
 
-bool FPowerMethod::Verify(bool bComputeLargest) const
+bool FPowerMethod::VerifyUserParameters(bool bComputeLargest) const
 {	
     if (OpMatrixA.Product == nullptr || OpMatrixA.Rows == nullptr) 
     {
@@ -208,6 +218,12 @@ bool FPowerMethod::Verify(bool bComputeLargest) const
         return false;
     } 
 
+    // check that the initial vector (if set) has the correct dimensions
+    if (Parms.InitialSolution.IsSet() && Parms.InitialSolution.GetValue().rows() != OpMatrixA.Rows())
+    {
+        return false; 
+    }
+
     return true;
 }
 
diff --git a/Engine/Plugins/Runtime/GeometryProcessing/Source/DynamicMesh/Private/Solvers/Internal/PowerMethodSolver.h b/Engine/Plugins/Runtime/GeometryProcessing/Source/DynamicMesh/Private/Solvers/Internal/PowerMethodSolver.h
--- a/Engine/Plugins/Runtime/GeometryProcessing/Source/DynamicMesh/Private/Solvers/Internal/PowerMethodSolver.h
+++ b/Engine/Plugins/Runtime/GeometryProcessing/Source/DynamicMesh/Private/Solvers/Internal/PowerMethodSolver.h
@@ -33,6 +33,7 @@ public:
 	{
 		ScalarType Tolerance = 1e-6; // Convergence parameter, see Converged()
 		int32 MaxIterations = 1000;  // Maximum number of times the power iteration is run 
+		TOptional<RealVectorType> InitialSolution; // Optionally specify the starting vector for the first iteration.
 	};
 
 	// The result of the solve
@@ -133,7 +134,7 @@ protected:
 	virtual bool SetupSolver(bool bComputeLargest);
 
 	/** Check that the input user parameters are valid and matrix operators are set */
-	virtual bool Verify(bool bComputeLargest) const;
+	virtual bool VerifyUserParameters(bool bComputeLargest) const;
 
 	/** Is matrix B not an identity matrix*/
 	inline bool IsGeneralProblem() const 
