ModelingTools: Fixed MaterialID bake to overwrite rather than accumulate.

diff --git a/Engine/Plugins/Runtime/GeometryProcessing/Source/DynamicMesh/Private/Sampling/MeshPropertyMapEvaluator.cpp b/Engine/Plugins/Runtime/GeometryProcessing/Source/DynamicMesh/Private/Sampling/MeshPropertyMapEvaluator.cpp
--- a/Engine/Plugins/Runtime/GeometryProcessing/Source/DynamicMesh/Private/Sampling/MeshPropertyMapEvaluator.cpp
+++ b/Engine/Plugins/Runtime/GeometryProcessing/Source/DynamicMesh/Private/Sampling/MeshPropertyMapEvaluator.cpp
@@ -63,6 +63,7 @@ void FMeshPropertyMapEvaluator::Setup(const FMeshBaseBaker& Baker, FEvaluationCo
 		break;
 	case EMeshPropertyMapType::MaterialID:
 		DefaultValue = FVector3f(LinearColors::LightPink3f());
+		Context.AccumulateMode = EAccumulateMode::Overwrite;
 		break;
 	case EMeshPropertyMapType::VertexColor:
 		DefaultValue = FVector3f::One();
