ModelingTools: Add support for blending when evaluating SampleFilterMask in BakeTexture/BakeAll

diff --git a/Engine/Plugins/Experimental/GeometryScripting/Source/GeometryScriptingCore/Private/MeshBakeFunctions.cpp b/Engine/Plugins/Experimental/GeometryScripting/Source/GeometryScriptingCore/Private/MeshBakeFunctions.cpp
--- a/Engine/Plugins/Experimental/GeometryScripting/Source/GeometryScriptingCore/Private/MeshBakeFunctions.cpp
+++ b/Engine/Plugins/Experimental/GeometryScripting/Source/GeometryScriptingCore/Private/MeshBakeFunctions.cpp
@@ -518,8 +518,8 @@ namespace GeometryScriptBakeLocals
 		{
 			Baker.SampleFilterF = [&SampleFilterMask](const FVector2i& ImageCoords, const FVector2d& UV, int32 TriID)
 			{
-				const FVector4f Color = SampleFilterMask->BilinearSampleUV<float>(UV, FVector4f(0, 0, 0, 1));
-				return Color.IsNearlyZero3();
+				const FVector4f Mask = SampleFilterMask->BilinearSampleUV<float>(UV, FVector4f::One());
+				return (Mask.X + Mask.Y + Mask.Z) / 3;
 			};
 		}
 		if (bIsBakeToSelf)
diff --git a/Engine/Plugins/Experimental/MeshModelingToolsetExp/Source/MeshModelingToolsExp/Private/BakeMeshAttributeMapsTool.cpp b/Engine/Plugins/Experimental/MeshModelingToolsetExp/Source/MeshModelingToolsExp/Private/BakeMeshAttributeMapsTool.cpp
--- a/Engine/Plugins/Experimental/MeshModelingToolsetExp/Source/MeshModelingToolsExp/Private/BakeMeshAttributeMapsTool.cpp
+++ b/Engine/Plugins/Experimental/MeshModelingToolsetExp/Source/MeshModelingToolsExp/Private/BakeMeshAttributeMapsTool.cpp
@@ -146,8 +146,8 @@ public:
 		{
 			Baker->SampleFilterF = [this](const FVector2i& ImageCoords, const FVector2d& UV, int32 TriID)
 			{
-				const FVector4f Color = SampleFilterMask->BilinearSampleUV<float>(UV, FVector4f(0, 0, 0, 1));
-				return Color.IsNearlyZero3();
+				const FVector4f Mask = SampleFilterMask->BilinearSampleUV<float>(UV, FVector4f::One());
+				return (Mask.X + Mask.Y + Mask.Z) / 3;
 			};
 		}
 		
diff --git a/Engine/Plugins/Experimental/MeshModelingToolsetExp/Source/MeshModelingToolsExp/Private/BakeMultiMeshAttributeMapsTool.cpp b/Engine/Plugins/Experimental/MeshModelingToolsetExp/Source/MeshModelingToolsExp/Private/BakeMultiMeshAttributeMapsTool.cpp
--- a/Engine/Plugins/Experimental/MeshModelingToolsetExp/Source/MeshModelingToolsExp/Private/BakeMultiMeshAttributeMapsTool.cpp
+++ b/Engine/Plugins/Experimental/MeshModelingToolsetExp/Source/MeshModelingToolsExp/Private/BakeMultiMeshAttributeMapsTool.cpp
@@ -355,8 +355,8 @@ public:
 		{
 			Baker->SampleFilterF = [this](const FVector2i& ImageCoords, const FVector2d& UV, int32 TriID)
 			{
-				const FVector4f Color = SampleFilterMask->BilinearSampleUV<float>(UV, FVector4f(0, 0, 0, 1));
-				return Color.IsNearlyZero3();
+				const FVector4f Mask = SampleFilterMask->BilinearSampleUV<float>(UV, FVector4f::One());
+				return (Mask.X + Mask.Y + Mask.Z) / 3;
 			};
 		}
 		
diff --git a/Engine/Plugins/Runtime/GeometryProcessing/Source/DynamicMesh/Private/Sampling/MeshMapBaker.cpp b/Engine/Plugins/Runtime/GeometryProcessing/Source/DynamicMesh/Private/Sampling/MeshMapBaker.cpp
--- a/Engine/Plugins/Runtime/GeometryProcessing/Source/DynamicMesh/Private/Sampling/MeshMapBaker.cpp
+++ b/Engine/Plugins/Runtime/GeometryProcessing/Source/DynamicMesh/Private/Sampling/MeshMapBaker.cpp
@@ -512,19 +512,9 @@ void FMeshMapBaker::BakeSample(
 	float* Buffer = static_cast<float*>(FMemory_Alloca(sizeof(float) * BakeSampleBufferSize));
 	float* BufferPtr = Buffer;
 	const int32 NumBakers = Bakers.Num();
-	if (SampleFilterF && SampleFilterF(ImageCoords, UVPosition, Sample.DetailTriID))
+	for (int32 Idx = 0; Idx < NumBakers; ++Idx)
 	{
-		for (int32 Idx = 0; Idx < NumBakers; ++Idx)
-		{
-			BakeContexts[Idx].EvaluateDefault(BufferPtr, BakeContexts[Idx].EvalData);
-		}
-	}
-	else
-	{
-		for (int32 Idx = 0; Idx < NumBakers; ++Idx)
-		{
-			BakeContexts[Idx].Evaluate(BufferPtr, Sample, BakeContexts[Idx].EvalData);
-		}
+		BakeContexts[Idx].Evaluate(BufferPtr, Sample, BakeContexts[Idx].EvalData);
 	}
 	checkSlow((BufferPtr - Buffer) == BakeSampleBufferSize);
 
@@ -533,7 +523,7 @@ void FMeshMapBaker::BakeSample(
 	const int64 OccupancyMapSampleIdx = OccupancyMap.Tile.GetIndexFromSourceCoords(ImageCoords);
 	const int32 SampleUVChart = OccupancyMap.TexelQueryUVChart[OccupancyMapSampleIdx];
 
-	auto AddFn = [this, &ImageCoords, &UVPosition, &Tile, &TileBuffer, Buffer, &OccupancyMap, SampleUVChart](const TArray<int32>& BakeIds) -> void
+	auto AddFn = [this, &ImageCoords, &UVPosition, &Tile, &TileBuffer, &OccupancyMap, SampleUVChart](const TArray<int32>& BakeIds, const float* SourceBuffer, float Weight) -> void
 	{
 		const FVector2i BoxFilterStart(
 			FMath::Clamp(ImageCoords.X - FilterKernelSize, 0, Dimensions.GetWidth()),
@@ -561,6 +551,7 @@ void FMeshMapBaker::BakeSample(
 
 			float FilterWeight = TextureFilterEval(TexelDistance);
 			FilterWeight *= (SampleUVChart == BufferTilePixelUVChart);
+			FilterWeight *= Weight;
 			PixelWeight += FilterWeight;
 			for (const int32 BakeIdx : BakeIds)
 			{
@@ -574,14 +565,14 @@ void FMeshMapBaker::BakeSample(
 					const int32 Stride = static_cast<int32>(Context.DataLayout[DataIdx]);
 					for (int32 BufIdx = Offset; BufIdx < Offset + Stride; ++BufIdx)
 					{
-						PixelBuffer[BufIdx] += Buffer[BufIdx] * FilterWeight;
+						PixelBuffer[BufIdx] += SourceBuffer[BufIdx] * FilterWeight;
 					}
 				}
 			}
 		}
 	};
 
-	auto OverwriteFn = [this, &ImageCoords, &Tile, &TileBuffer, Buffer](const TArray<int32>& BakeIds) -> void
+	auto OverwriteFn = [this, &ImageCoords, &Tile, &TileBuffer](const TArray<int32>& BakeIds, const float* SourceBuffer) -> void
 	{
 		const int64 BufferTilePixelLinearIdx = Tile.GetIndexFromSourceCoords(ImageCoords); 
 		float* PixelBuffer = TileBuffer.GetPixel(BufferTilePixelLinearIdx);
@@ -598,14 +589,24 @@ void FMeshMapBaker::BakeSample(
 				const int32 Stride = static_cast<int32>(Context.DataLayout[DataIdx]);
 				for (int32 BufIdx = Offset; BufIdx < Offset + Stride; ++BufIdx)
 				{
-					PixelBuffer[BufIdx] = Buffer[BufIdx];
+					PixelBuffer[BufIdx] = SourceBuffer[BufIdx];
 				}
 			}
 		}
 	};
 
-	AddFn(BakeAccumulateLists[static_cast<int32>(FMeshMapEvaluator::EAccumulateMode::Add)]);
-	OverwriteFn(BakeAccumulateLists[static_cast<int32>(FMeshMapEvaluator::EAccumulateMode::Overwrite)]);
+	if (SampleFilterF)
+	{
+		const float SampleMaskWeight = FMath::Clamp(SampleFilterF(ImageCoords, UVPosition, Sample.BaseSample.TriangleIndex), 0.0f, 1.0f);
+		AddFn(BakeAccumulateLists[static_cast<int32>(FMeshMapEvaluator::EAccumulateMode::Add)], Buffer, SampleMaskWeight);
+		AddFn(BakeAccumulateLists[static_cast<int32>(FMeshMapEvaluator::EAccumulateMode::Add)], BakeDefaults.GetData(), 1.0f - SampleMaskWeight);
+		OverwriteFn(BakeAccumulateLists[static_cast<int32>(FMeshMapEvaluator::EAccumulateMode::Overwrite)], FMath::IsNearlyZero(SampleMaskWeight) ? BakeDefaults.GetData() : Buffer);
+	}
+	else
+	{
+		AddFn(BakeAccumulateLists[static_cast<int32>(FMeshMapEvaluator::EAccumulateMode::Add)], Buffer, 1.0f);
+		OverwriteFn(BakeAccumulateLists[static_cast<int32>(FMeshMapEvaluator::EAccumulateMode::Overwrite)], Buffer);
+	}
 }
 
 int32 FMeshMapBaker::AddEvaluator(const TSharedPtr<FMeshMapEvaluator, ESPMode::ThreadSafe>& Eval)
diff --git a/Engine/Plugins/Runtime/GeometryProcessing/Source/DynamicMesh/Public/Sampling/MeshMapBaker.h b/Engine/Plugins/Runtime/GeometryProcessing/Source/DynamicMesh/Public/Sampling/MeshMapBaker.h
--- a/Engine/Plugins/Runtime/GeometryProcessing/Source/DynamicMesh/Public/Sampling/MeshMapBaker.h
+++ b/Engine/Plugins/Runtime/GeometryProcessing/Source/DynamicMesh/Public/Sampling/MeshMapBaker.h
@@ -56,7 +56,7 @@ public:
 	 * @param TriID the target mesh triangle ID to be evaluated
 	 * @return true if the given sample should be filtered (return default evaluator result)
 	 */
-	TFunction<bool(const FVector2i&, const FVector2d&, int32)> SampleFilterF = nullptr;
+	TFunction<float(const FVector2i&, const FVector2d&, int32)> SampleFilterF = nullptr;
 
 	//
 	// Parameters
