Modeling Tools: Fix some incorrect comments

#rnx

diff --git a/Engine/Plugins/Runtime/GeometryProcessing/Source/DynamicMesh/Public/Sampling/MeshMapBakerQueue.h b/Engine/Plugins/Runtime/GeometryProcessing/Source/DynamicMesh/Public/Sampling/MeshMapBakerQueue.h
--- a/Engine/Plugins/Runtime/GeometryProcessing/Source/DynamicMesh/Public/Sampling/MeshMapBakerQueue.h
+++ b/Engine/Plugins/Runtime/GeometryProcessing/Source/DynamicMesh/Public/Sampling/MeshMapBakerQueue.h
@@ -70,7 +70,6 @@ public:
 	 *
 	 * @param Idx The slot index to update.
 	 * @param Data The data to post to the queue.
-	 * @return true if the operation succeeded (the slot was empty).
 	 */
 	void Post(SIZE_T Idx, void* Data)
 	{
diff --git a/Engine/Source/Runtime/GeometryCore/Public/Spatial/DenseGrid2.h b/Engine/Source/Runtime/GeometryCore/Public/Spatial/DenseGrid2.h
--- a/Engine/Source/Runtime/GeometryCore/Public/Spatial/DenseGrid2.h
+++ b/Engine/Source/Runtime/GeometryCore/Public/Spatial/DenseGrid2.h
@@ -16,7 +16,7 @@ namespace Geometry
 {
 
 /**
- * 2D dense grid of floating-point scalar values.
+ * 2D dense grid of scalar values.
  */
 template<typename ElemType>
 class TDenseGrid2
