MeshMapBaker: Change criterion used to switch between computed and default values in overwrite mode evaluators when a sample filter function is used

This change means that if an overwrite mode evaluator (e.g., one computing Material IDs) used a default value for a given sample then all the add mode evaluators will have only added contributions from the evaluator default value for that sample too.

#rnx

diff --git a/Engine/Plugins/Runtime/GeometryProcessing/Source/DynamicMesh/Private/Sampling/MeshMapBaker.cpp b/Engine/Plugins/Runtime/GeometryProcessing/Source/DynamicMesh/Private/Sampling/MeshMapBaker.cpp
--- a/Engine/Plugins/Runtime/GeometryProcessing/Source/DynamicMesh/Private/Sampling/MeshMapBaker.cpp
+++ b/Engine/Plugins/Runtime/GeometryProcessing/Source/DynamicMesh/Private/Sampling/MeshMapBaker.cpp
@@ -600,7 +600,7 @@ void FMeshMapBaker::BakeSample(
 		const float SampleMaskWeight = FMath::Clamp(SampleFilterF(ImageCoords, UVPosition, Sample.BaseSample.TriangleIndex), 0.0f, 1.0f);
 		AddFn(BakeAccumulateLists[static_cast<int32>(FMeshMapEvaluator::EAccumulateMode::Add)], Buffer, SampleMaskWeight);
 		AddFn(BakeAccumulateLists[static_cast<int32>(FMeshMapEvaluator::EAccumulateMode::Add)], BakeDefaults.GetData(), 1.0f - SampleMaskWeight);
-		OverwriteFn(BakeAccumulateLists[static_cast<int32>(FMeshMapEvaluator::EAccumulateMode::Overwrite)], FMath::IsNearlyZero(SampleMaskWeight) ? BakeDefaults.GetData() : Buffer);
+		OverwriteFn(BakeAccumulateLists[static_cast<int32>(FMeshMapEvaluator::EAccumulateMode::Overwrite)], (SampleMaskWeight == 0) ? BakeDefaults.GetData() : Buffer);
 	}
 	else
 	{
