MeshMapBaker: Some refactoring/clarification

#rnx

diff --git a/Engine/Plugins/Runtime/GeometryProcessing/Source/DynamicMesh/Private/Sampling/MeshMapBaker.cpp b/Engine/Plugins/Runtime/GeometryProcessing/Source/DynamicMesh/Private/Sampling/MeshMapBaker.cpp
--- a/Engine/Plugins/Runtime/GeometryProcessing/Source/DynamicMesh/Private/Sampling/MeshMapBaker.cpp
+++ b/Engine/Plugins/Runtime/GeometryProcessing/Source/DynamicMesh/Private/Sampling/MeshMapBaker.cpp
@@ -68,7 +68,8 @@ void FMeshMapBaker::InitBake()
 			BakeResults[ResultIdx]->SetDimensions(Dimensions);
 
 			BakeSampleOffsets[ResultIdx] = SampleOffset;
-			SampleOffset += static_cast<int32>(BakeContexts[Idx].DataLayout[DataIdx]);
+			const int32 NumFloats = static_cast<int32>(BakeContexts[Idx].DataLayout[DataIdx]);
+			SampleOffset += NumFloats;
 		}
 	}
 	BakeSampleOffsets[NumResults] = SampleOffset;
@@ -182,7 +183,6 @@ void FMeshMapBaker::Bake()
 	{
 		NormalOverlay->GetTriBaryInterpolate<double>(SampleInfo.TriangleIndex, &SampleInfo.BaryCoords.X, &ValueOut.BaseNormal.X);
 		Normalize(ValueOut.BaseNormal);
-		const FVector3d RayDir = ValueOut.BaseNormal;
 
 		ValueOut.BaseSample = SampleInfo;
 		ValueOut.DetailMesh = nullptr;
@@ -206,6 +206,7 @@ void FMeshMapBaker::Bake()
 			const double SampleThickness = this->GetProjectionDistance();		// could modulate w/ a map here...
 
 			// Find detail mesh triangle point
+			const FVector3d RayDir = ValueOut.BaseNormal;
 			ValueOut.DetailMesh = GetDetailMeshTrianglePoint_Raycast(DetailSampler, SampleInfo.SurfacePoint, RayDir,
 				ValueOut.DetailTriID, ValueOut.DetailBaryCoords, SampleThickness,
 				(UseStrategy == ECorrespondenceStrategy::RaycastStandardThenNearest));
@@ -213,8 +214,8 @@ void FMeshMapBaker::Bake()
 	});
 
 	// Create a temporary output float buffer for the full image dimensions.
-	const FImageTile ImageTile(FVector2i(0,0), FVector2i(Dimensions.GetWidth(), Dimensions.GetHeight()));
-	FMeshMapTileBuffer ImageTileBuffer(ImageTile, BakeSampleBufferSize);
+	const FImageTile FullImageTile(FVector2i(0,0), FVector2i(Dimensions.GetWidth(), Dimensions.GetHeight()));
+	FMeshMapTileBuffer FullImageTileBuffer(FullImageTile, BakeSampleBufferSize);
 
 	// Tile the image
 	FImageTiling Tiles(Dimensions, TileSize, TileSize);
@@ -222,22 +223,8 @@ void FMeshMapBaker::Bake()
 	TArray<TArray64<TTuple<int64, int64>>> GutterTexelsPerTile;
 	GutterTexelsPerTile.SetNum(NumTiles);
 
-	auto NoopFn = [](const float& In, float& Out)
-	{
-	};
-
-	auto AddFn = [](const float& In, float& Out)
-	{
-		Out += In;
-	};
-
-	auto OverwriteFn = [](const float& In, float& Out)
-	{
-		Out = In;
-	};
-
-	// WriteToOutputBuffer transfers local tile data (TileBuffer) to the image output buffer (ImageTileBuffer).
-	auto WriteToOutputBuffer = [this, &ImageTileBuffer] (FMeshMapTileBuffer& TileBufferIn, const FImageTile& TargetTile, const TArray<int32>& BakeIds, auto&& Op, auto&& WeightOp)
+	// WriteToOutputBuffer transfers local tile data (TileBuffer) to the image output buffer (FullImageTileBuffer).
+	auto WriteToOutputBuffer = [this, &FullImageTileBuffer] (FMeshMapTileBuffer& TileBufferIn, const FImageTile& TargetTile, const TArray<int32>& EvaluatorIds, auto&& Op, auto&& WeightOp)
 	{
 		TRACE_CPUPROFILER_EVENT_SCOPE(FMeshMapBaker::Bake_WriteToOutputBuffer);
 		
@@ -249,22 +236,23 @@ void FMeshMapBaker::Bake()
 			{
 				if (CancelF())
 				{
-					return;
+					return; // WriteToOutputBuffer
 				}
 
 				const FVector2i ImageCoords = TargetTile.GetSourceCoords(TileCoords);
 
+				const int64 ImageLinearIdx = Dimensions.GetIndex(ImageCoords);
+				float& ImagePixelWeight = FullImageTileBuffer.GetPixelWeight(ImageLinearIdx);
+				float* ImagePixelBuffer = FullImageTileBuffer.GetPixel(ImageLinearIdx);
+
 				const FImageTile& BufferTile = TileBufferIn.GetTile();
+
 				const int64 TilePixelLinearIdx = BufferTile.GetIndexFromSourceCoords(ImageCoords);
 				const float& TilePixelWeight = TileBufferIn.GetPixelWeight(TilePixelLinearIdx);
 				float* TilePixelBuffer = TileBufferIn.GetPixel(TilePixelLinearIdx);
 
-				const int64 ImageLinearIdx = Dimensions.GetIndex(ImageCoords);
-				float& ImagePixelWeight = ImageTileBuffer.GetPixelWeight(ImageLinearIdx);
-				float* ImagePixelBuffer = ImageTileBuffer.GetPixel(ImageLinearIdx);
-
 				WeightOp(TilePixelWeight, ImagePixelWeight);
-				for( int32 Idx : BakeIds )
+				for( int32 Idx : EvaluatorIds )
 				{
 					const FMeshMapEvaluator::FEvaluationContext& Context = BakeContexts[Idx];
 					const int32 NumData = Context.DataLayout.Num();
@@ -275,9 +263,9 @@ void FMeshMapBaker::Bake()
 						const int32 Offset = BakeSampleOffsets[ResultIdx];
 						float* BufferPtr = &TilePixelBuffer[Offset];
 						float* ImageBufferPtr = &ImagePixelBuffer[Offset];
-						const FMeshMapEvaluator::EComponents Stride = Context.DataLayout[DataIdx];
 
-						for (int32 FloatIdx = 0; FloatIdx < (int32)Stride; ++FloatIdx)
+						const int32 NumFloats = static_cast<int32>(Context.DataLayout[DataIdx]);
+						for (int32 FloatIdx = 0; FloatIdx < NumFloats; ++FloatIdx)
 						{
 							Op(BufferPtr[FloatIdx], ImageBufferPtr[FloatIdx]);
 						}
@@ -287,30 +275,35 @@ void FMeshMapBaker::Bake()
 		}
 	};
 
-	FMeshMapBakerQueue OutputQueue(NumTiles);
-	auto WriteQueuedOutput = [this, &WriteToOutputBuffer, &AddFn](FMeshMapBakerQueue& Queue)
+	auto WriteQueuedOutput = [this, &WriteToOutputBuffer](FMeshMapBakerQueue& Queue)
 	{
+		constexpr auto AddFn = [](const float& In, float& Out)
+		{
+			Out += In;
+		};
+
 		if (Queue.AcquireProcessLock())
 		{
 			void* OutputData = Queue.Process();
 			while (OutputData)
 			{
 				FMeshMapTileBuffer* TileBufferPtr = static_cast<FMeshMapTileBuffer*>(OutputData);
-				WriteToOutputBuffer(*TileBufferPtr, TileBufferPtr->GetTile(), BakeAccumulateLists[static_cast<int32>(FMeshMapEvaluator::EAccumulateMode::Add)], AddFn, AddFn);
+				WriteToOutputBuffer(*TileBufferPtr, TileBufferPtr->GetTile(), EvaluatorIdsForMode(FMeshMapEvaluator::EAccumulateMode::Add), AddFn, AddFn);
 				delete TileBufferPtr;
 				OutputData = Queue.Process();
 			}
 			Queue.ReleaseProcessLock();
 		}
 	};
-	
-	ParallelFor(NumTiles, [this, &Tiles, &GutterTexelsPerTile, &OutputQueue, &WriteToOutputBuffer, &OverwriteFn, &NoopFn, &WriteQueuedOutput](int32 TileIdx)
+
+	FMeshMapBakerQueue OutputQueue(NumTiles);
+	ParallelFor(NumTiles, [this, &Tiles, &GutterTexelsPerTile, &OutputQueue, &WriteToOutputBuffer, &WriteQueuedOutput](int32 TileIdx)
 	{
 		TRACE_CPUPROFILER_EVENT_SCOPE(FMeshMapBaker::Bake_EvalTile);
 
 		if (CancelF())
 		{
-			return;
+			return; // ParallelFor
 		}
 		
 		// Generate unpadded and padded tiles.
@@ -347,7 +340,7 @@ void FMeshMapBaker::Bake()
 					if (CancelF())
 					{
 						delete TileBuffer;
-						return;
+						return; // ParallelFor
 					}
 
 					const FVector2i ImageCoords = Tile.GetSourceCoords(TileCoords);
@@ -377,8 +370,17 @@ void FMeshMapBaker::Bake()
 			}
 		}
 
+		constexpr auto NoopFn = [](const float& In, float& Out)
+		{
+		};
+
+		constexpr auto OverwriteFn = [](const float& In, float& Out)
+		{
+			Out = In;
+		};
+
 		// Transfer 'Overwrite' float data to image tile buffer
-		WriteToOutputBuffer(*TileBuffer, Tile, BakeAccumulateLists[static_cast<int32>(FMeshMapEvaluator::EAccumulateMode::Overwrite)], OverwriteFn, NoopFn);
+		WriteToOutputBuffer(*TileBuffer, Tile, EvaluatorIdsForMode(FMeshMapEvaluator::EAccumulateMode::Overwrite), OverwriteFn, NoopFn);
 
 		// Accumulate 'Add' float data to image tile buffer
 		OutputQueue.Post(TileIdx, TileBuffer);
@@ -394,24 +396,30 @@ void FMeshMapBaker::Bake()
 			if (Data)
 			{
 				const FMeshMapTileBuffer* TileBuffer = static_cast<FMeshMapTileBuffer*>(Data);
-				delete TileBuffer;				
+				delete TileBuffer;
 			}
 		}
+
+		return;
 	}
-	else
+
 	{
 		// The queue only acquires the process lock if the next item in the queue
 		// is ready. This could mean that there are potential leftovers in the queue
 		// after the parallel for. Write them out now.
 		WriteQueuedOutput(OutputQueue);
 	}
+	
+	if (CancelF())
+	{
+		return;
+	}
 
-	if (!CancelF())
 	{
 		FScopedDurationTimer WriteToImageTimer(BakeAnalytics.WriteToImageDuration);
 		
 		// Normalize and convert ImageTileBuffer data to color data.
-		ParallelFor(NumTiles, [this, &Tiles, &ImageTileBuffer](int32 TileIdx)
+		ParallelFor(NumTiles, [this, &Tiles, &FullImageTileBuffer](int32 TileIdx)
 		{
 			TRACE_CPUPROFILER_EVENT_SCOPE(FMeshMapBaker::Bake_WriteToImageBuffer);
 		
@@ -424,17 +432,18 @@ void FMeshMapBaker::Bake()
 				{
 					if (CancelF())
 					{
-						return;
+						return; // ParallelFor
 					}
 
 					const FVector2i ImageCoords = Tile.GetSourceCoords(TileCoords);
+
 					const int64 ImageLinearIdx = Dimensions.GetIndex(ImageCoords);
-					const float& PixelWeight = ImageTileBuffer.GetPixelWeight(ImageLinearIdx);
-					float* PixelBuffer = ImageTileBuffer.GetPixel(ImageLinearIdx);
+					const float& PixelWeight = FullImageTileBuffer.GetPixelWeight(ImageLinearIdx);
+					float* PixelBuffer = FullImageTileBuffer.GetPixel(ImageLinearIdx);
 
-					auto WriteToPixel = [this, &PixelBuffer, &ImageLinearIdx](TArray<int32>& BakeIds, float OneOverWeight)
+					auto WriteToPixel = [this, &PixelBuffer, &ImageLinearIdx](const TArray<int32>& EvaluatorIds, float OneOverWeight)
 					{
-						for (const int32 Idx : BakeIds)
+						for (const int32 Idx : EvaluatorIds)
 						{
 							const FMeshMapEvaluator::FEvaluationContext& Context = BakeContexts[Idx];
 							const int32 NumData = Context.DataLayout.Num();
@@ -444,10 +453,10 @@ void FMeshMapBaker::Bake()
 								const int32 ResultIdx = ResultOffset + DataIdx;
 								const int32 Offset = BakeSampleOffsets[ResultIdx];
 								float* BufferPtr = &PixelBuffer[Offset];
-								const FMeshMapEvaluator::EComponents Stride = Context.DataLayout[DataIdx];
 
 								// Apply weight to raw float data.
-								for (int32 FloatIdx = 0; FloatIdx < static_cast<int32>(Stride); ++FloatIdx)
+								const int32 NumFloats = static_cast<int32>(Context.DataLayout[DataIdx]);
+								for (int32 FloatIdx = 0; FloatIdx < NumFloats; ++FloatIdx)
 								{
 									BufferPtr[FloatIdx] *= OneOverWeight;
 								}
@@ -461,16 +470,21 @@ void FMeshMapBaker::Bake()
 				
 					if (PixelWeight > 0.0)
 					{
-						WriteToPixel(BakeAccumulateLists[static_cast<int32>(FMeshMapEvaluator::EAccumulateMode::Add)], 1.0f / PixelWeight);
+						WriteToPixel(EvaluatorIdsForMode(FMeshMapEvaluator::EAccumulateMode::Add), 1.0f / PixelWeight);
 					}
-					WriteToPixel(BakeAccumulateLists[static_cast<int32>(FMeshMapEvaluator::EAccumulateMode::Overwrite)], 1.0f);
+					WriteToPixel(EvaluatorIdsForMode(FMeshMapEvaluator::EAccumulateMode::Overwrite), 1.0f);
 				}
 			}
 		}, !bParallel ? EParallelForFlags::ForceSingleThread : EParallelForFlags::None);
 	}
 
+	if (CancelF())
+	{
+		return;
+	}
+
 	// Gutter Texel processing
-	if( bGutterEnabled && !CancelF())
+	if (bGutterEnabled)
 	{
 		FScopedDurationTimer WriteToGutterTimer(BakeAnalytics.WriteToGutterDuration);
 		
@@ -481,7 +495,7 @@ void FMeshMapBaker::Bake()
 
 			if (CancelF())
 			{
-				return;
+				return; // ParallelFor
 			}
 
 			const int NumGutter = GutterTexelsPerTile[TileIdx].Num();
@@ -511,8 +525,8 @@ void FMeshMapBaker::BakeSample(
 	// Evaluate each baker into stack allocated float buffer
 	float* Buffer = static_cast<float*>(FMemory_Alloca(sizeof(float) * BakeSampleBufferSize));
 	float* BufferPtr = Buffer;
-	const int32 NumBakers = Bakers.Num();
-	for (int32 Idx = 0; Idx < NumBakers; ++Idx)
+	const int32 NumEvaluators = Bakers.Num();
+	for (int32 Idx = 0; Idx < NumEvaluators; ++Idx)
 	{
 		BakeContexts[Idx].Evaluate(BufferPtr, Sample, BakeContexts[Idx].EvalData);
 	}
@@ -523,7 +537,7 @@ void FMeshMapBaker::BakeSample(
 	const int64 OccupancyMapSampleIdx = OccupancyMap.Tile.GetIndexFromSourceCoords(ImageCoords);
 	const int32 SampleUVChart = OccupancyMap.TexelQueryUVChart[OccupancyMapSampleIdx];
 
-	auto AddFn = [this, &ImageCoords, &UVPosition, &Tile, &TileBuffer, &OccupancyMap, SampleUVChart](const TArray<int32>& BakeIds, const float* SourceBuffer, float Weight) -> void
+	auto AddFn = [this, &ImageCoords, &UVPosition, &Tile, &TileBuffer, &OccupancyMap, SampleUVChart](const TArray<int32>& EvaluatorIds, const float* SourceBuffer, float Weight) -> void
 	{
 		const FVector2i BoxFilterStart(
 			FMath::Clamp(ImageCoords.X - FilterKernelSize, 0, Dimensions.GetWidth()),
@@ -538,9 +552,10 @@ void FMeshMapBaker::BakeSample(
 		for (int64 FilterIdx = 0; FilterIdx < BoxFilterTile.Num(); FilterIdx++)
 		{
 			const FVector2i SourceCoords = BoxFilterTile.GetSourceCoords(FilterIdx);
-			const int64 BufferTilePixelLinearIdx = Tile.GetIndexFromSourceCoords(SourceCoords);
 			const int64 OccupancyMapFilterIdx = OccupancyMap.Tile.GetIndexFromSourceCoords(SourceCoords);
 			const int32 BufferTilePixelUVChart = OccupancyMap.TexelQueryUVChart[OccupancyMapFilterIdx];
+
+			const int64 BufferTilePixelLinearIdx = Tile.GetIndexFromSourceCoords(SourceCoords);
 			float* PixelBuffer = TileBuffer.GetPixel(BufferTilePixelLinearIdx);
 			float& PixelWeight = TileBuffer.GetPixelWeight(BufferTilePixelLinearIdx);
 
@@ -550,34 +565,36 @@ void FMeshMapBaker::BakeSample(
 			TexelDistance.Y *= Dimensions.GetHeight();
 
 			float FilterWeight = TextureFilterEval(TexelDistance);
-			FilterWeight *= (SampleUVChart == BufferTilePixelUVChart);
+			FilterWeight *= (SampleUVChart == BufferTilePixelUVChart); // Zero the contribution if the sample and pixel are on different UV charts
 			FilterWeight *= Weight;
 			PixelWeight += FilterWeight;
-			for (const int32 BakeIdx : BakeIds)
+
+			for (const int32 Idx : EvaluatorIds)
 			{
-				const FMeshMapEvaluator::FEvaluationContext& Context = BakeContexts[BakeIdx];
+				const FMeshMapEvaluator::FEvaluationContext& Context = BakeContexts[Idx];
 				const int32 NumData = Context.DataLayout.Num();
-				const int32 ResultOffset = BakeOffsets[BakeIdx];
+				const int32 ResultOffset = BakeOffsets[Idx];
 				for (int32 DataIdx = 0; DataIdx < NumData; ++DataIdx)
 				{
 					const int32 ResultIdx = ResultOffset + DataIdx;
 					const int32 Offset = BakeSampleOffsets[ResultIdx];
-					const int32 Stride = static_cast<int32>(Context.DataLayout[DataIdx]);
-					for (int32 BufIdx = Offset; BufIdx < Offset + Stride; ++BufIdx)
+
+					const int32 NumFloats = static_cast<int32>(Context.DataLayout[DataIdx]);
+					for (int32 FloatIdx = Offset; FloatIdx < Offset + NumFloats; ++FloatIdx)
 					{
-						PixelBuffer[BufIdx] += SourceBuffer[BufIdx] * FilterWeight;
+						PixelBuffer[FloatIdx] += SourceBuffer[FloatIdx] * FilterWeight;
 					}
 				}
 			}
 		}
 	};
 
-	auto OverwriteFn = [this, &ImageCoords, &Tile, &TileBuffer](const TArray<int32>& BakeIds, const float* SourceBuffer) -> void
+	auto OverwriteFn = [this, &ImageCoords, &Tile, &TileBuffer](const TArray<int32>& EvaluatorIds, const float* SourceBuffer) -> void
 	{
 		const int64 BufferTilePixelLinearIdx = Tile.GetIndexFromSourceCoords(ImageCoords); 
 		float* PixelBuffer = TileBuffer.GetPixel(BufferTilePixelLinearIdx);
 		
-		for (const int32 Idx : BakeIds)
+		for (const int32 Idx : EvaluatorIds)
 		{
 			const FMeshMapEvaluator::FEvaluationContext& Context = BakeContexts[Idx];
 			const int32 NumData = Context.DataLayout.Num();
@@ -586,10 +603,11 @@ void FMeshMapBaker::BakeSample(
 			{
 				const int32 ResultIdx = ResultOffset + DataIdx;
 				const int32 Offset = BakeSampleOffsets[ResultIdx];
-				const int32 Stride = static_cast<int32>(Context.DataLayout[DataIdx]);
-				for (int32 BufIdx = Offset; BufIdx < Offset + Stride; ++BufIdx)
+				
+				const int32 NumFloats = static_cast<int32>(Context.DataLayout[DataIdx]);
+				for (int32 FloatIdx = Offset; FloatIdx < Offset + NumFloats; ++FloatIdx)
 				{
-					PixelBuffer[BufIdx] = SourceBuffer[BufIdx];
+					PixelBuffer[FloatIdx] = SourceBuffer[FloatIdx];
 				}
 			}
 		}
@@ -598,14 +616,14 @@ void FMeshMapBaker::BakeSample(
 	if (SampleFilterF)
 	{
 		const float SampleMaskWeight = FMath::Clamp(SampleFilterF(ImageCoords, UVPosition, Sample.BaseSample.TriangleIndex), 0.0f, 1.0f);
-		AddFn(BakeAccumulateLists[static_cast<int32>(FMeshMapEvaluator::EAccumulateMode::Add)], Buffer, SampleMaskWeight);
-		AddFn(BakeAccumulateLists[static_cast<int32>(FMeshMapEvaluator::EAccumulateMode::Add)], BakeDefaults.GetData(), 1.0f - SampleMaskWeight);
-		OverwriteFn(BakeAccumulateLists[static_cast<int32>(FMeshMapEvaluator::EAccumulateMode::Overwrite)], (SampleMaskWeight == 0) ? BakeDefaults.GetData() : Buffer);
+		AddFn(EvaluatorIdsForMode(FMeshMapEvaluator::EAccumulateMode::Add), Buffer, SampleMaskWeight);
+		AddFn(EvaluatorIdsForMode(FMeshMapEvaluator::EAccumulateMode::Add), BakeDefaults.GetData(), 1.0f - SampleMaskWeight);
+		OverwriteFn(EvaluatorIdsForMode(FMeshMapEvaluator::EAccumulateMode::Overwrite), (SampleMaskWeight == 0) ? BakeDefaults.GetData() : Buffer);
 	}
 	else
 	{
-		AddFn(BakeAccumulateLists[static_cast<int32>(FMeshMapEvaluator::EAccumulateMode::Add)], Buffer, 1.0f);
-		OverwriteFn(BakeAccumulateLists[static_cast<int32>(FMeshMapEvaluator::EAccumulateMode::Overwrite)], Buffer);
+		AddFn(EvaluatorIdsForMode(FMeshMapEvaluator::EAccumulateMode::Add), Buffer, 1.0f);
+		OverwriteFn(EvaluatorIdsForMode(FMeshMapEvaluator::EAccumulateMode::Overwrite), Buffer);
 	}
 }
 
diff --git a/Engine/Plugins/Runtime/GeometryProcessing/Source/DynamicMesh/Public/Sampling/MeshMapBaker.h b/Engine/Plugins/Runtime/GeometryProcessing/Source/DynamicMesh/Public/Sampling/MeshMapBaker.h
--- a/Engine/Plugins/Runtime/GeometryProcessing/Source/DynamicMesh/Public/Sampling/MeshMapBaker.h
+++ b/Engine/Plugins/Runtime/GeometryProcessing/Source/DynamicMesh/Public/Sampling/MeshMapBaker.h
@@ -54,7 +54,10 @@ public:
 	 * @param ImageCoords the output image coordinates to be evaluated
 	 * @param UV the target mesh UV coordinates to be evaluated
 	 * @param TriID the target mesh triangle ID to be evaluated
-	 * @return true if the given sample should be filtered (return default evaluator result)
+	 * @return a weight (clamped by the baker to the range [0,1]) used to combine evaluator sample values with evaluator
+	 *  defaults. For evaluators using EAccumulateMode::Overwrite the default is used when weight == 0, and the value is
+	 *  used otherwise. For evaluators using EAccumulateMode::Add the sample's value and default are blended using the
+	 *  expression: `weight * value + (1 - weight) * default`
 	 */
 	TFunction<float(const FVector2i&, const FVector2d&, int32)> SampleFilterF = nullptr;
 
@@ -162,6 +165,12 @@ protected:
 
 	FImageDimensions Dimensions = FImageDimensions(128, 128);
 
+	/** @return evaluator ids with the corresponding mode */
+	const TArray<int32>& EvaluatorIdsForMode(FMeshMapEvaluator::EAccumulateMode Mode) const
+	{
+		return BakeAccumulateLists[static_cast<int32>(Mode)];
+	}
+
 	/**
 	 * If true, the baker will pad the baked content past the UV borders by GutterSize.
 	 * This is useful to minimize artifacts when filtering or mipmapping.
diff --git a/Engine/Plugins/Runtime/GeometryProcessing/Source/DynamicMesh/Public/Sampling/MeshSurfaceSampler.h b/Engine/Plugins/Runtime/GeometryProcessing/Source/DynamicMesh/Public/Sampling/MeshSurfaceSampler.h
--- a/Engine/Plugins/Runtime/GeometryProcessing/Source/DynamicMesh/Public/Sampling/MeshSurfaceSampler.h
+++ b/Engine/Plugins/Runtime/GeometryProcessing/Source/DynamicMesh/Public/Sampling/MeshSurfaceSampler.h
@@ -221,35 +221,10 @@ bool TMeshSurfaceUVSampler<SampleType>::SampleUV(const FVector2d& UV, SampleType
 	check(QueryType == EMeshSurfaceSamplerQueryType::UVOnly);
 	check(bUVSpatialValid);
 
-	FMeshUVSampleInfo Sample;
-
 	FRay3d HitRay(FVector3d(UV.X, UV.Y, 100.0), -FVector3d::UnitZ());
-	Sample.TriangleIndex = UVBVTree.FindNearestHitTriangle(HitRay);
-	if ( Mesh->IsTriangle(Sample.TriangleIndex) == false )
-	{
-		ResultOut = ZeroValue;
-		return false;
-	}
-	check(UVOverlay->IsSetTriangle(Sample.TriangleIndex));
-
-	Sample.MeshVertices = Mesh->GetTriangle(Sample.TriangleIndex);
-	Sample.Triangle3D = FTriangle3d(
-		Mesh->GetVertex(Sample.MeshVertices.A),
-		Mesh->GetVertex(Sample.MeshVertices.B),
-		Mesh->GetVertex(Sample.MeshVertices.C));
+	const int32 UVTriangleID = UVBVTree.FindNearestHitTriangle(HitRay);
 
-	Sample.UVVertices = UVOverlay->GetTriangle(Sample.TriangleIndex);
-	Sample.TriangleUV = FTriangle2d(
-		(FVector2d)UVOverlay->GetElement(Sample.UVVertices.A),
-		(FVector2d)UVOverlay->GetElement(Sample.UVVertices.B),
-		(FVector2d)UVOverlay->GetElement(Sample.UVVertices.C));
-
-	Sample.BaryCoords = Sample.TriangleUV.GetBarycentricCoords(UV);
-	Sample.SurfacePoint = Mesh->GetTriBaryPoint(Sample.TriangleIndex, Sample.BaryCoords.X, Sample.BaryCoords.Y, Sample.BaryCoords.Z);
-
-	ValueFunction(Sample, ResultOut);
-
-	return true;
+	return SampleUV(UVTriangleID, UV, ResultOut);
 }
 
 
