Remesh Tool: initial implementation of (optional) boundary constraint projection:
- Constrained edges along a boundary are allowed to collapse, and vertices are allowed to move, but they are projected back to the original boundary curve geometry
- Additionally, vertices at boundary "corners" (defined using an angle threshold criterion) are constrained to not move or be deleted

diff --git a/Engine/Plugins/Experimental/MeshModelingToolsetExp/Source/MeshModelingToolsExp/Private/RemeshMeshTool.cpp b/Engine/Plugins/Experimental/MeshModelingToolsetExp/Source/MeshModelingToolsExp/Private/RemeshMeshTool.cpp
--- a/Engine/Plugins/Experimental/MeshModelingToolsetExp/Source/MeshModelingToolsExp/Private/RemeshMeshTool.cpp
+++ b/Engine/Plugins/Experimental/MeshModelingToolsetExp/Source/MeshModelingToolsExp/Private/RemeshMeshTool.cpp
@@ -64,6 +64,8 @@ URemeshMeshToolProperties::URemeshMeshToolProperties()
 	bSplits = true;
 	bCollapses = true;
 	bReproject = true;
+	bReprojectConstraints = false;
+	BoundaryCornerAngleThreshold = 45.0;
 	bPreventNormalFlips = true;
 	bPreventTinyTriangles = true;
 	bUseTargetEdgeLength = false;
@@ -219,6 +221,8 @@ TUniquePtr<FDynamicMeshOperator> URemeshMeshTool::MakeNewOperator()
 	Op->bPreventNormalFlips = BasicProperties->bPreventNormalFlips;
 	Op->bPreventTinyTriangles = BasicProperties->bPreventTinyTriangles;
 	Op->bReproject = BasicProperties->bReproject;
+	Op->bReprojectConstraints = BasicProperties->bReprojectConstraints;
+	Op->BoundaryCornerAngleThreshold = BasicProperties->BoundaryCornerAngleThreshold;
 	Op->bSplits = BasicProperties->bSplits;
 	Op->RemeshIterations = BasicProperties->RemeshIterations;
 	Op->MaxRemeshIterations = BasicProperties->MaxRemeshIterations;
diff --git a/Engine/Plugins/Experimental/MeshModelingToolsetExp/Source/MeshModelingToolsExp/Public/RemeshMeshTool.h b/Engine/Plugins/Experimental/MeshModelingToolsetExp/Source/MeshModelingToolsExp/Public/RemeshMeshTool.h
--- a/Engine/Plugins/Experimental/MeshModelingToolsetExp/Source/MeshModelingToolsExp/Public/RemeshMeshTool.h
+++ b/Engine/Plugins/Experimental/MeshModelingToolsetExp/Source/MeshModelingToolsExp/Public/RemeshMeshTool.h
@@ -87,6 +87,15 @@ public:
 	/** Enable projection back to input mesh */
 	UPROPERTY(EditAnywhere, Category = Remeshing, AdvancedDisplay)
 	bool bReproject;
+
+	/** Project constrained vertices back to original constraint curves */
+	UPROPERTY(EditAnywhere, Category = Remeshing, AdvancedDisplay)
+	bool bReprojectConstraints;
+
+	/** Angle threshold in degrees for classifying a boundary vertex as a corner. Corners will be fixed if Reproject Constraints is active. */
+	UPROPERTY(EditAnywhere, Category = Remeshing, AdvancedDisplay, meta = (EditCondition = "bReprojectConstraints", UIMin = "0", UIMax = "180", ClampMin = "0", ClampMax = "180"))
+	float BoundaryCornerAngleThreshold;
+
 };
 
 
diff --git a/Engine/Plugins/Runtime/GeometryProcessing/Source/DynamicMesh/Private/MeshConstraintsUtil.cpp b/Engine/Plugins/Runtime/GeometryProcessing/Source/DynamicMesh/Private/MeshConstraintsUtil.cpp
--- a/Engine/Plugins/Runtime/GeometryProcessing/Source/DynamicMesh/Private/MeshConstraintsUtil.cpp
+++ b/Engine/Plugins/Runtime/GeometryProcessing/Source/DynamicMesh/Private/MeshConstraintsUtil.cpp
@@ -2,6 +2,8 @@
 
 #include "MeshConstraintsUtil.h"
 #include "Async/ParallelFor.h"
+#include "GroupTopology.h"
+#include "MeshBoundaryLoops.h"
 
 using namespace UE::Geometry;
 
@@ -234,7 +236,7 @@ bool FMeshConstraintsUtil::ConstrainEdgeBoundariesAndSeams(const int EdgeID,
 	{
 		ApplyBoundaryConstraint(MaterialBoundaryConstraintFlags);
 	}
-	if (bIsSeam)
+	if (bIsSeam)	// TODO: mesh boundary edges are currently flagged as seams (based on the implementation of Attributes->IsSeamEdge), so are subject to seam constraints
 	{
 		CurVtxConstraint.bCannotDelete = CurVtxConstraint.bCannotDelete || !bAllowSeamCollapse;
 		CurVtxConstraint.bCanMove = CurVtxConstraint.bCanMove && (bAllowSeamSmoothing || bAllowSeamCollapse);
@@ -276,4 +278,149 @@ bool FMeshConstraintsUtil::ConstrainEdgeBoundariesAndSeams(const int EdgeID,
 		}
 	}
 	return false;
+}
+
+
+namespace
+{
+	// Utility class for finding mesh boundaries in the same way we find group and material boundaries
+	class FMeshBoundaryTopology : public FGroupTopology
+	{
+	public:
+
+		FMeshBoundaryTopology() {}
+
+		FMeshBoundaryTopology(const FDynamicMesh3* Mesh, bool bAutoBuild) :
+			FGroupTopology(Mesh, bAutoBuild)
+		{}
+
+		virtual int GetGroupID(int TriangleID) const override
+		{
+			return 1;
+		}
+	};
+}
+
+
+void FMeshConstraintsUtil::SetBoundaryConstraintsWithProjection(
+	FMeshConstraints& Constraints,
+	const EBoundaryType BoundaryToConstrain,
+	const FDynamicMesh3& Mesh,
+	double CornerAngleThreshold)
+{
+	TSet<int32> GroupCorners;
+
+	TUniquePtr<FGroupTopology> Topology;
+
+	switch (BoundaryToConstrain)
+	{
+	case EBoundaryType::Mesh:
+		Topology = MakeUnique<FMeshBoundaryTopology>(&Mesh, true);
+		break;
+	case EBoundaryType::Group:
+		// TODO: this will tag mesh boundaries as group boundaries, however Mesh.IsGroupBoundaryEdge() returns false at mesh boundary edges
+		Topology = MakeUnique<FGroupTopology>(&Mesh, true);		
+		break;
+	case EBoundaryType::MaterialID:
+		if (!ensure(Mesh.HasAttributes() && Mesh.Attributes()->GetMaterialID() != nullptr))
+		{
+			return;
+		}
+		// TODO: this will tag mesh boundaries as attribute boundaries, however Mesh.Attributes->IsMaterialBoundaryEdge() returns false at mesh boundary edges
+		Topology = MakeUnique<FGroupTopology>(&Mesh, Mesh.Attributes()->GetMaterialID(), true);
+		break;
+	}
+
+	const int32 NumEdges = Topology->Edges.Num();
+	const int32 OldNumCurves = Constraints.ProjectionData.ProjectionCurves.Num();
+	const int32 NewNumCurves = OldNumCurves + NumEdges;
+
+	Constraints.ProjectionData.ProjectionCurves.SetNum(NewNumCurves);
+
+	for (int32 TopologyEdgeIndex = 0; TopologyEdgeIndex < NumEdges; ++TopologyEdgeIndex)
+	{
+		const FGroupTopology::FGroupEdge& GroupEdge = Topology->Edges[TopologyEdgeIndex];
+
+		TSharedPtr<FMeshConstraintCurve> Curve = MakeShared<FMeshConstraintCurve>();
+		Constraints.ProjectionData.ProjectionCurves[OldNumCurves + TopologyEdgeIndex] = Curve;
+
+		const bool bIsLoop = (GroupEdge.EndpointCorners[0] == IndexConstants::InvalidID);
+
+		// mark corner vertices for pinning
+		if (!bIsLoop)
+		{
+			GroupCorners.Add(Topology->GetCornerVertexID(GroupEdge.EndpointCorners[0]));
+			GroupCorners.Add(Topology->GetCornerVertexID(GroupEdge.EndpointCorners[1]));
+		}
+
+		// extract polyline
+		GroupEdge.Span.GetPolyline(*Curve);
+
+		const TArray<int32>& EdgeVerts = Topology->GetGroupEdgeVertices(TopologyEdgeIndex);
+		const int32 NumEdgeVerts = EdgeVerts.Num();
+
+		for (int32 EdgeVertIndex = 0; EdgeVertIndex < NumEdgeVerts - 1; ++EdgeVertIndex)
+		{
+			// Check for angle-based corners
+			int32 PrevEdgeVert;
+			const int32 NextEdgeVert = EdgeVertIndex + 1;
+
+			if (EdgeVertIndex == 0)
+			{
+				if (!bIsLoop)
+				{
+					continue;
+				}
+				PrevEdgeVert = NumEdgeVerts - 2;		// Vertex at N-1 is the same vertex as at 0, for FGroupTopology loops
+			}
+			else
+			{
+				PrevEdgeVert = EdgeVertIndex - 1;
+			}
+
+			const int32 VertexID = EdgeVerts[EdgeVertIndex];
+
+			const FVector3d PrevRel = Mesh.GetVertex(VertexID) - Mesh.GetVertex(EdgeVerts[PrevEdgeVert]);
+			const FVector3d NextRel = Mesh.GetVertex(EdgeVerts[NextEdgeVert]) - Mesh.GetVertex(VertexID);
+
+			if (PrevRel.Size() > UE_SMALL_NUMBER && NextRel.Size() > UE_SMALL_NUMBER)
+			{
+				const double CornerCosAngle = PrevRel.Dot(NextRel) / PrevRel.Size() / NextRel.Size();
+				const double CornerAngleDeg = FMath::RadiansToDegrees(FMath::Acos(CornerCosAngle));
+
+				if (FMath::Abs(CornerAngleDeg) > CornerAngleThreshold)
+				{
+					// Found a corner
+					GroupCorners.Add(VertexID);
+					continue;
+				}
+			}
+
+			Constraints.SetOrCombineVertexConstraint(VertexID, FVertexConstraint(Curve.Get()));
+		}
+
+		const TArray<int32>& EdgeEdges = Topology->GetGroupEdgeEdges(TopologyEdgeIndex);
+		for (int32 EdgeEdgeIndex = 0; EdgeEdgeIndex < EdgeEdges.Num(); ++EdgeEdgeIndex)
+		{
+			const int32 EdgeID = EdgeEdges[EdgeEdgeIndex];
+
+			FEdgeConstraint FoundEdgeConstraint;
+			if (Constraints.GetEdgeConstraint(EdgeID, FoundEdgeConstraint))
+			{
+				FoundEdgeConstraint.Target = Curve.Get();
+				Constraints.SetOrUpdateEdgeConstraint(EdgeID, FoundEdgeConstraint);
+			}
+			else
+			{
+				FEdgeConstraint GroupEdgeConstraint(EEdgeRefineFlags::NoConstraint, Curve.Get());
+				Constraints.SetOrUpdateEdgeConstraint(EdgeID, GroupEdgeConstraint);
+			}
+		}
+	}
+
+
+	for (int32 CornerVertexID : GroupCorners)
+	{
+		Constraints.SetOrCombineVertexConstraint(CornerVertexID, FVertexConstraint::FullyConstrained());
+	}
 }
\ No newline at end of file
diff --git a/Engine/Plugins/Runtime/GeometryProcessing/Source/DynamicMesh/Public/MeshConstraints.h b/Engine/Plugins/Runtime/GeometryProcessing/Source/DynamicMesh/Public/MeshConstraints.h
--- a/Engine/Plugins/Runtime/GeometryProcessing/Source/DynamicMesh/Public/MeshConstraints.h
+++ b/Engine/Plugins/Runtime/GeometryProcessing/Source/DynamicMesh/Public/MeshConstraints.h
@@ -5,6 +5,7 @@
 #pragma once
 
 #include "Spatial/SpatialInterfaces.h"   // for projection target
+#include "Polyline3.h"
 
 namespace UE
 {
@@ -175,16 +176,16 @@ public:
 		Target(nullptr)
 	{}
 
-	explicit FVertexConstraint(IProjectionTarget* Target) :
+	explicit FVertexConstraint(IProjectionTarget* Target, int FixedSetID = InvalidSetID) :
 		bCannotDelete(false),
-		bCanMove(false),
-		FixedSetID(InvalidSetID),
+		bCanMove(true),
+		FixedSetID(FixedSetID),
 		Target(Target)
 	{}
 
 	bool IsUnconstrained() const 
 	{
-		return ( (bCanMove == true) && (bCannotDelete == false) );
+		return ( (bCanMove == true) && (bCannotDelete == false) && (Target == nullptr) );
 	}
 
 	/** @return an unconstrained vertex constraint (ie not constrained at all) */
@@ -202,11 +203,30 @@ public:
 	{
 		bCannotDelete = bCannotDelete || OtherConstraint.bCannotDelete;
 		bCanMove = bCanMove && OtherConstraint.bCanMove;
+		// TODO: What should happen if both constraints have projection targets?
+		Target = OtherConstraint.Target;
 	}
 
 };
 
 
+class FMeshConstraintCurve : public FPolyline3d, public IProjectionTarget
+{
+public:
+
+	virtual FVector3d Project(const FVector3d& Point, int /*Identifier*/) override
+	{
+		int32 NearestSegmentIdx;
+		double NearestSegmentParam;
+		double DistSqr = DistanceSquared(Point, NearestSegmentIdx, NearestSegmentParam);
+		return GetSegmentPoint(NearestSegmentIdx, NearestSegmentParam);
+	}
+};
+
+struct DYNAMICMESH_API FConstraintProjectionData
+{
+	TArray<TSharedPtr<FMeshConstraintCurve>> ProjectionCurves;
+};
 
 
 /**
@@ -227,6 +247,10 @@ protected:
 
 public:
 
+	/** Additional data for constraint curve projection */
+	FConstraintProjectionData ProjectionData;
+
+
 	FMeshConstraints()
 	{
 		FixedSetIDCounter = 0;
@@ -363,6 +387,19 @@ public:
 		Vertices.Add(VertexID, vc);
 	}
 
+	void SetOrCombineVertexConstraint(int VertexID, const FVertexConstraint& vc)
+	{
+		FVertexConstraint* Found = Vertices.Find(VertexID);
+		if (Found != nullptr)
+		{
+			Found->CombineConstraint(vc);
+		}
+		else
+		{
+			Vertices.Add(VertexID, vc);
+		}
+	}
+
 	/** Clear the constraint on the given VertexID */
 	void ClearVertexConstraint(int VertexID)
 	{
diff --git a/Engine/Plugins/Runtime/GeometryProcessing/Source/DynamicMesh/Public/MeshConstraintsUtil.h b/Engine/Plugins/Runtime/GeometryProcessing/Source/DynamicMesh/Public/MeshConstraintsUtil.h
--- a/Engine/Plugins/Runtime/GeometryProcessing/Source/DynamicMesh/Public/MeshConstraintsUtil.h
+++ b/Engine/Plugins/Runtime/GeometryProcessing/Source/DynamicMesh/Public/MeshConstraintsUtil.h
@@ -172,6 +172,27 @@ s	 * EdgeROI that are seams, as well as vertices incident to edges in EdgeROI th
 			                                    FVertexConstraint& VertexConstraintA,
 			                                    FVertexConstraint& VertexConstraintB);
 
+	enum class EBoundaryType
+	{
+		Mesh = 0,
+		Group = 1,
+		MaterialID = 2
+	};
+
+	/**
+	 * Create polycurve projection targets at group boundaries. 
+	 * Vertices on multiple boundaries are fixed, as well as vertices at boundary "corners" 
+	 * (where the angle formed between adjacent boundary edges is above a given threshold.)
+	 * 
+	 * @param Constraints (out) constraint data to add new constraints to
+	 * @param Mesh the mesh to constrain
+	 * @param CornerAngleThreshold angle threshold (in degrees) above which adjacent edges are considered a corner
+	 */
+	static void SetBoundaryConstraintsWithProjection(
+		FMeshConstraints& Constraints,
+		const EBoundaryType BoundaryToConstrain,
+		const FDynamicMesh3& Mesh,
+		double CornerAngleThreshold = 45.0);
 
 private:
 	FMeshConstraintsUtil() = delete;		// this class is not constructible
diff --git a/Engine/Plugins/Runtime/MeshModelingToolset/Source/ModelingOperators/Private/CleaningOps/RemeshMeshOp.cpp b/Engine/Plugins/Runtime/MeshModelingToolset/Source/ModelingOperators/Private/CleaningOps/RemeshMeshOp.cpp
--- a/Engine/Plugins/Runtime/MeshModelingToolset/Source/ModelingOperators/Private/CleaningOps/RemeshMeshOp.cpp
+++ b/Engine/Plugins/Runtime/MeshModelingToolset/Source/ModelingOperators/Private/CleaningOps/RemeshMeshOp.cpp
@@ -118,14 +118,41 @@ void FRemeshMeshOp::CalculateResult(FProgressCancel* Progress)
 
 	Remesher->DEBUG_CHECK_LEVEL = 0;
 
-	FMeshConstraints constraints;
-	FMeshConstraintsUtil::ConstrainAllBoundariesAndSeams(constraints, *TargetMesh,
-														 MeshBoundaryConstraint,
-														 GroupBoundaryConstraint,
-														 MaterialBoundaryConstraint,
-														 true, !bPreserveSharpEdges);
-
-	Remesher->SetExternalConstraints(MoveTemp(constraints));
+	FMeshConstraints Constraints;
+	constexpr bool bAllowSeamSplits = true;
+	constexpr bool bAllowSeamCollpase = false;
+	const bool bAllowSeamSmoothing = !bPreserveSharpEdges;
+
+	FMeshConstraintsUtil::ConstrainAllBoundariesAndSeams(Constraints,
+		*TargetMesh,
+		MeshBoundaryConstraint,
+		GroupBoundaryConstraint,
+		MaterialBoundaryConstraint,
+		bAllowSeamSplits,
+		bAllowSeamSmoothing,
+		bAllowSeamCollpase);
+
+	if (bReprojectConstraints)
+	{
+		Constraints.ProjectionData.ProjectionCurves.Reset();
+
+		if ( MeshBoundaryConstraint != EEdgeRefineFlags::NoConstraint )
+		{
+			FMeshConstraintsUtil::SetBoundaryConstraintsWithProjection(Constraints, FMeshConstraintsUtil::EBoundaryType::Mesh, *TargetMesh, BoundaryCornerAngleThreshold);
+		}
+
+		if (GroupBoundaryConstraint != EEdgeRefineFlags::NoConstraint)
+		{
+			FMeshConstraintsUtil::SetBoundaryConstraintsWithProjection(Constraints, FMeshConstraintsUtil::EBoundaryType::Group, *TargetMesh, BoundaryCornerAngleThreshold);
+		}
+
+		if (MaterialBoundaryConstraint != EEdgeRefineFlags::NoConstraint)
+		{
+			FMeshConstraintsUtil::SetBoundaryConstraintsWithProjection(Constraints, FMeshConstraintsUtil::EBoundaryType::MaterialID, *TargetMesh, BoundaryCornerAngleThreshold);
+		}
+	}
+
+	Remesher->SetExternalConstraints(MoveTemp(Constraints));
 
 	if (ProjectionTarget == nullptr)
 	{
diff --git a/Engine/Plugins/Runtime/MeshModelingToolset/Source/ModelingOperators/Public/CleaningOps/RemeshMeshOp.h b/Engine/Plugins/Runtime/MeshModelingToolset/Source/ModelingOperators/Public/CleaningOps/RemeshMeshOp.h
--- a/Engine/Plugins/Runtime/MeshModelingToolset/Source/ModelingOperators/Public/CleaningOps/RemeshMeshOp.h
+++ b/Engine/Plugins/Runtime/MeshModelingToolset/Source/ModelingOperators/Public/CleaningOps/RemeshMeshOp.h
@@ -65,6 +65,8 @@ public:
 	float SmoothingStrength, TargetEdgeLength;
 	ERemeshSmoothingType SmoothingType;
 	bool bDiscardAttributes, bPreserveSharpEdges, bFlips, bSplits, bCollapses, bReproject, bPreventNormalFlips, bPreventTinyTriangles;
+	bool bReprojectConstraints = false;
+	double BoundaryCornerAngleThreshold = 45.0;
 
 	// When true, result will have attributes object regardless of whether attributes 
 	// were discarded or present initially.
diff --git a/Engine/Source/Runtime/GeometryCore/Public/Polyline3.h b/Engine/Source/Runtime/GeometryCore/Public/Polyline3.h
--- a/Engine/Source/Runtime/GeometryCore/Public/Polyline3.h
+++ b/Engine/Source/Runtime/GeometryCore/Public/Polyline3.h
@@ -392,11 +392,11 @@ public:
 			T d = TNumericLimits<T>::Max();
 			if (t >= seg.Extent)
 			{
-				d = seg.EndPoint().DistanceSquared(QueryPoint);
+				d = TVector<T>::DistSquared(seg.EndPoint(), QueryPoint);
 			}
 			else if (t <= -seg.Extent)
 			{
-				d = seg.StartPoint().DistanceSquared(QueryPoint);
+				d = TVector<T>::DistSquared(seg.StartPoint(), QueryPoint);
 			}
 			else
 			{
