return the default normal/occlusion for occlusion baker samples that have a degenerate normal

diff --git a/Engine/Plugins/Runtime/GeometryProcessing/Source/DynamicMesh/Private/Sampling/MeshOcclusionMapBaker.cpp b/Engine/Plugins/Runtime/GeometryProcessing/Source/DynamicMesh/Private/Sampling/MeshOcclusionMapBaker.cpp
--- a/Engine/Plugins/Runtime/GeometryProcessing/Source/DynamicMesh/Private/Sampling/MeshOcclusionMapBaker.cpp
+++ b/Engine/Plugins/Runtime/GeometryProcessing/Source/DynamicMesh/Private/Sampling/MeshOcclusionMapBaker.cpp
@@ -88,7 +88,13 @@ void FMeshOcclusionMapBaker::Bake()
 			//FVector3d DetailTriNormal = DetailMesh.GetTriNormal(DetailTriID);
 			FVector3d DetailTriNormal;
 			DetailNormalOverlay->GetTriBaryInterpolate<double>(DetailTriID, &SampleData.DetailBaryCoords.X, &DetailTriNormal.X);
-			Normalize(DetailTriNormal);
+			if (!DetailTriNormal.Normalize())
+			{
+				// degenerate triangle normal
+				OutOcclusion = 0.0;
+				OutNormal = DefaultNormal;
+				return;
+			}
 
 			FVector3d BaseTangentX, BaseTangentY;
 			if (WantBentNormal() && NormalSpace == ESpace::Tangent)
diff --git a/Engine/Plugins/Runtime/GeometryProcessing/Source/DynamicMesh/Private/Sampling/MeshOcclusionMapEvaluator.cpp b/Engine/Plugins/Runtime/GeometryProcessing/Source/DynamicMesh/Private/Sampling/MeshOcclusionMapEvaluator.cpp
--- a/Engine/Plugins/Runtime/GeometryProcessing/Source/DynamicMesh/Private/Sampling/MeshOcclusionMapEvaluator.cpp
+++ b/Engine/Plugins/Runtime/GeometryProcessing/Source/DynamicMesh/Private/Sampling/MeshOcclusionMapEvaluator.cpp
@@ -194,7 +194,11 @@ FMeshOcclusionMapEvaluator::FOcclusionTuple FMeshOcclusionMapEvaluator::SampleFu
 
 	FVector3f DetailTriNormal;
 	DetailSampler->TriBaryInterpolateNormal(DetailMesh, DetailTriID, SampleData.DetailBaryCoords, DetailTriNormal);
-	Normalize(DetailTriNormal);
+	if (!DetailTriNormal.Normalize())
+	{
+		// degenerate triangle normal
+		return FOcclusionTuple(0.0, DefaultNormal);
+	}
 
 	FVector3d BaseTangentX, BaseTangentY;
 	if constexpr (WANT_BENT_NORMAL(ComputeType) && ComputeSpace == ESpace::Tangent)
