Render Capture Baking: Support correspondence sample validity checking, use BSpline filtering and enable multisampling

#rnx

diff --git a/Engine/Plugins/Experimental/MeshModelingToolsetExp/Source/MeshModelingToolsExp/Private/BakeRenderCaptureTool.cpp b/Engine/Plugins/Experimental/MeshModelingToolsetExp/Source/MeshModelingToolsExp/Private/BakeRenderCaptureTool.cpp
--- a/Engine/Plugins/Experimental/MeshModelingToolsetExp/Source/MeshModelingToolsExp/Private/BakeRenderCaptureTool.cpp
+++ b/Engine/Plugins/Experimental/MeshModelingToolsetExp/Source/MeshModelingToolsExp/Private/BakeRenderCaptureTool.cpp
@@ -28,6 +28,36 @@ using namespace UE::Geometry;
 // Implementation details
 //
 
+
+
+// TODO Reimplement this when the render capture algorithm is refactored to decouple the visibility checking from the evaluation
+class FSceneCapturePhotoSetSampler : public FMeshBakerDynamicMeshSampler
+{
+public:
+	FSceneCapturePhotoSetSampler(
+		const FSceneCapturePhotoSet* SceneCapture,
+		TFunctionRef<bool(const FVector3d&, const FVector3d&)> VisibilityFunction,
+		const FDynamicMesh3* Mesh,
+		const FDynamicMeshAABBTree3* Spatial,
+		const FMeshTangentsd* Tangents = nullptr) :
+		FMeshBakerDynamicMeshSampler(Mesh, Spatial, Tangents),
+		SceneCapture(SceneCapture),
+		VisibilityFunction(VisibilityFunction)
+	{
+	}
+
+	virtual bool IsValidCorrespondence(const FMeshMapEvaluator::FCorrespondenceSample& Sample) const override
+	{
+		return SceneCapture->IsValidSample(Sample.BaseSample.SurfacePoint, Sample.BaseNormal, VisibilityFunction);
+	}
+
+public:
+	const FSceneCapturePhotoSet* SceneCapture = nullptr;
+	TFunctionRef<bool(const FVector3d&, const FVector3d&)> VisibilityFunction;
+};
+
+
+
 static FName BaseColorTexParamName = FName("BaseColor");
 static FName RoughnessTexParamName = FName("Roughness");
 static FName MetallicTexParamName = FName("Metallic");
@@ -77,8 +107,7 @@ public:
 		// output texture options
 		int32 TextureImageSize = 1024;
 
-		// supersampling parameter
-		int32 AntiAliasMultiSampling = 0;
+		EBakeTextureSamplesPerPixel SamplesPerPixel = EBakeTextureSamplesPerPixel::Sample1;
 
 		//
 		// Mesh settings
@@ -115,7 +144,7 @@ public:
 		Options.bUsePackedMRS = UseSettings.MaterialSettings.bPackedMRSMap;
 
 		Options.TextureImageSize = UseSettings.MaterialSettings.TextureSize;
-		Options.AntiAliasMultiSampling = FMath::Max(1, UseSettings.MultiSamplingAA);
+		Options.SamplesPerPixel = UseSettings.SamplesPerPixel;
 
 		return Options;
 	}
@@ -174,33 +203,36 @@ static void ImageBuildersFromPhotoSet(
 {
 	TRACE_CPUPROFILER_EVENT_SCOPE(ImageBuildersFromPhotoSet);
 
-	FDynamicMeshAABBTree3 Spatial(BaseMesh, true);
+	FDynamicMeshAABBTree3 BaseMeshSpatial(BaseMesh, true);
 
 	double RayOffsetHackDist = (double)(100.0f * FMathf::ZeroTolerance * BaseMesh->GetBounds().MinDim() );
-	auto VisibilityFunction = [&Spatial, RayOffsetHackDist](const FVector3d& SurfPos, const FVector3d& ImagePosWorld)
+	auto VisibilityFunction = [&BaseMeshSpatial, RayOffsetHackDist](const FVector3d& SurfPos, const FVector3d& ImagePosWorld)
 	{
 		FVector3d RayDir = ImagePosWorld - SurfPos;
 		double Dist = Normalize(RayDir);
 		FVector3d RayOrigin = SurfPos + RayOffsetHackDist * RayDir;
-		int32 HitTID = Spatial.FindNearestHitTriangle(FRay3d(RayOrigin, RayDir), IMeshSpatial::FQueryOptions(Dist));
+		int32 HitTID = BaseMeshSpatial.FindNearestHitTriangle(FRay3d(RayOrigin, RayDir), IMeshSpatial::FQueryOptions(Dist));
 		return (HitTID == IndexConstants::InvalidID);
 	};
 
 	FMeshMapBaker Baker;
-	FMeshBakerDynamicMeshSampler Sampler(BaseMesh, &Spatial, BaseMeshTangents.Get());
-	Baker.SetDetailSampler(&Sampler);
 	Baker.SetTargetMesh(BaseMesh);
 	Baker.SetTargetMeshTangents(BaseMeshTangents);
 	Baker.SetDimensions(FImageDimensions(Options.TextureImageSize, Options.TextureImageSize));
-	Baker.SetSamplesPerPixel(1); // TODO Support multisampling (must work with infill)
-	Baker.SetFilter(FMeshMapBaker::EBakeFilterType::None); // TODO Support texture filtering (must work with infill)
+	Baker.SetSamplesPerPixel(static_cast<int32>(Options.SamplesPerPixel));
+	Baker.SetFilter(FMeshMapBaker::EBakeFilterType::BSpline);
 	Baker.SetTargetMeshUVLayer(Options.TargetUVLayer);
+
+	FSceneCapturePhotoSetSampler Sampler(SceneCapture, VisibilityFunction, BaseMesh, &BaseMeshSpatial, BaseMeshTangents.Get());
+	Baker.SetDetailSampler(&Sampler);
 	Baker.SetCorrespondenceStrategy(FMeshMapBaker::ECorrespondenceStrategy::Identity);
 
 	TMap<ERenderCaptureType, int32> CaptureTypeToEvaluatorIndexMap;
 
-	FVector4f InvalidColor(0, -1, 0, 1);
-	FVector3f DefaultNormal = FVector3f::UnitZ();
+	// Pixels in the output textures which don't map onto the mesh have a light grey color (except the normal map which
+	// will show a color corresponding to a unit z tangent space normal)
+	const FVector4f InvalidColor(.42, .42, .42, 1);
+	const FVector3f DefaultNormal = FVector3f::UnitZ();
 
 	FSceneCapturePhotoSet::FSceneSample DefaultColorSample;
 	DefaultColorSample.BaseColor = FVector3f(InvalidColor.X, InvalidColor.Y, InvalidColor.Z);
@@ -313,11 +345,13 @@ static void ImageBuildersFromPhotoSet(
 		}
 	}
 
+	// TODO Remove the FMeshImageBakingCache and reimplement infill to work with the new baking framework
+
 	FMeshImageBakingCache TempBakeCache;
 	{
 		TRACE_CPUPROFILER_EVENT_SCOPE(ImageBuildersFromPhotoSet_MakeMeshImageBakingCache);
 
-		TempBakeCache.SetDetailMesh(BaseMesh, &Spatial);
+		TempBakeCache.SetDetailMesh(BaseMesh, &BaseMeshSpatial);
 		TempBakeCache.SetBakeTargetMesh(BaseMesh);
 		TempBakeCache.SetDimensions(FImageDimensions(Options.TextureImageSize, Options.TextureImageSize));
 		TempBakeCache.SetUVLayer(Options.TargetUVLayer);
@@ -500,7 +534,7 @@ void UBakeRenderCaptureTool::Setup()
 
 	Settings->MapPreview = BaseColorTexParamName.ToString(); // We always bake the base color
 	Settings->WatchProperty(Settings->MapPreview, [this](FString) { UpdateVisualization(); GetToolManager()->PostInvalidation(); });
-	Settings->WatchProperty(Settings->MultiSamplingAA, [this](int32) { OpState |= EBakeOpState::Evaluate; });
+	Settings->WatchProperty(Settings->SamplesPerPixel, [this](EBakeTextureSamplesPerPixel) { OpState |= EBakeOpState::Evaluate; });
 	Settings->WatchProperty(Settings->RenderCaptureResolution, [this](int32) {  OpState |= EBakeOpState::Evaluate; });
 	Settings->WatchProperty(Settings->MaterialSettings, [this](FMaterialProxySettingsRC) { OpState |= EBakeOpState::Evaluate; });
 	Settings->WatchProperty(Settings->CaptureFieldOfView, [this](float) { OpState |= EBakeOpState::Evaluate; });
diff --git a/Engine/Plugins/Experimental/MeshModelingToolsetExp/Source/MeshModelingToolsExp/Public/BakeRenderCaptureTool.h b/Engine/Plugins/Experimental/MeshModelingToolsetExp/Source/MeshModelingToolsExp/Public/BakeRenderCaptureTool.h
--- a/Engine/Plugins/Experimental/MeshModelingToolsetExp/Source/MeshModelingToolsExp/Public/BakeRenderCaptureTool.h
+++ b/Engine/Plugins/Experimental/MeshModelingToolsetExp/Source/MeshModelingToolsExp/Public/BakeRenderCaptureTool.h
@@ -172,11 +172,9 @@ public:
 	// Material Baking Settings
 	//
 
-	/** If Value is > 1, Multisample output baked textures by this amount in each direction (eg 4 == 16x supersampling).
-	 *  WIP feature currently disabled
-	 */
-	UPROPERTY(EditAnywhere, Category = MaterialSettings, meta = (ClampMin = "0", ClampMax = "8", UIMin = "0", UIMax = "4", EditCondition="false"))
-	int32 MultiSamplingAA = 0;
+	/** Number of samples per pixel */
+	UPROPERTY(EditAnywhere, Category = MaterialSettings)
+	EBakeTextureSamplesPerPixel SamplesPerPixel = EBakeTextureSamplesPerPixel::Sample1;
 
 	/** If Value is zero, use MaterialSettings resolution, otherwise override the render capture resolution */
 	UPROPERTY(EditAnywhere, Category = MaterialSettings, meta = (ClampMin = "0"))
diff --git a/Engine/Plugins/Runtime/GeometryProcessing/Source/DynamicMesh/Private/Sampling/MeshMapBaker.cpp b/Engine/Plugins/Runtime/GeometryProcessing/Source/DynamicMesh/Private/Sampling/MeshMapBaker.cpp
--- a/Engine/Plugins/Runtime/GeometryProcessing/Source/DynamicMesh/Private/Sampling/MeshMapBaker.cpp
+++ b/Engine/Plugins/Runtime/GeometryProcessing/Source/DynamicMesh/Private/Sampling/MeshMapBaker.cpp
@@ -177,9 +177,10 @@ void FMeshMapBaker::Bake()
 		UseStrategy = ECorrespondenceStrategy::NearestPoint;
 	}
 
-	// This sampler finds the correspondence between target surface and detail surface.
-	DetailCorrespondenceSampler.Initialize(Mesh, UVOverlay, EMeshSurfaceSamplerQueryType::TriangleAndUV, FMeshMapEvaluator::FCorrespondenceSample(),
-		[Mesh, NormalOverlay, UseStrategy, this](const FMeshUVSampleInfo& SampleInfo, FMeshMapEvaluator::FCorrespondenceSample& ValueOut)
+	// Computes the correspondence sample assuming the SampleInfo is valid
+	// Returns true if the correspondence is valid and false otherwise
+	auto ComputeCorrespondenceSample
+		= [Mesh, NormalOverlay, UseStrategy, this](const FMeshUVSampleInfo& SampleInfo, FMeshMapEvaluator::FCorrespondenceSample& ValueOut)
 	{
 		NormalOverlay->GetTriBaryInterpolate<double>(SampleInfo.TriangleIndex, &SampleInfo.BaryCoords.X, &ValueOut.BaseNormal.X);
 		Normalize(ValueOut.BaseNormal);
@@ -199,6 +200,11 @@ void FMeshMapBaker::Bake()
 			ValueOut.DetailMesh = GetDetailMeshTrianglePoint_Nearest(DetailSampler, SampleInfo.SurfacePoint,
 				ValueOut.DetailTriID, ValueOut.DetailBaryCoords);
 		}
+		else if (UseStrategy == ECorrespondenceStrategy::Custom && DetailSampler->SupportsCustomCorrespondence())
+		{
+			ValueOut.DetailMesh = DetailSampler->ComputeCustomCorrespondence(SampleInfo, ValueOut);
+			check(false); // @Incomplete
+		}
 		else	// Fall back to raycast strategy
 		{
 			checkSlow(DetailSampler->SupportsRaycastCorrespondence());
@@ -211,7 +217,18 @@ void FMeshMapBaker::Bake()
 				ValueOut.DetailTriID, ValueOut.DetailBaryCoords, SampleThickness,
 				(UseStrategy == ECorrespondenceStrategy::RaycastStandardThenNearest));
 		}
-	});
+
+		if (ValueOut.DetailMesh && DetailSampler->IsTriangle(ValueOut.DetailMesh, ValueOut.DetailTriID))
+		{
+			return DetailSampler->IsValidCorrespondence(ValueOut);
+		}
+
+		return false;
+	};
+
+	// This computes a FMeshUVSampleInfo to pass to the ComputeCorrespondenceSample function, which will find the
+	// correspondence between the target surface and detail surface.
+	MeshUVSampler.Initialize(Mesh, UVOverlay, EMeshSurfaceSamplerQueryType::TriangleAndUV);
 
 	// Create a temporary output float buffer for the full image dimensions.
 	const FImageTile FullImageTile(FVector2i(0,0), FVector2i(Dimensions.GetWidth(), Dimensions.GetHeight()));
@@ -275,7 +292,7 @@ void FMeshMapBaker::Bake()
 		}
 	};
 
-	auto WriteQueuedOutput = [this, &WriteToOutputBuffer](FMeshMapBakerQueue& Queue)
+	auto WriteToOutputBufferQueued = [this, &WriteToOutputBuffer](FMeshMapBakerQueue& Queue)
 	{
 		constexpr auto AddFn = [](const float& In, float& Out)
 		{
@@ -297,7 +314,7 @@ void FMeshMapBaker::Bake()
 	};
 
 	FMeshMapBakerQueue OutputQueue(NumTiles);
-	ParallelFor(NumTiles, [this, &Tiles, &GutterTexelsPerTile, &OutputQueue, &WriteToOutputBuffer, &WriteQueuedOutput](int32 TileIdx)
+	ParallelFor(NumTiles, [this, &Tiles, &GutterTexelsPerTile, &OutputQueue, &WriteToOutputBuffer, &WriteToOutputBufferQueued, &ComputeCorrespondenceSample](int32 TileIdx)
 	{
 		TRACE_CPUPROFILER_EVENT_SCOPE(FMeshMapBaker::Bake_EvalTile);
 
@@ -350,6 +367,7 @@ void FMeshMapBaker::Bake()
 						continue;
 					}
 
+					// Iterate over all the samples in the pixel
 					for (int32 SampleIdx = 0; SampleIdx < NumSamples; ++SampleIdx)
 					{
 						const int64 LinearIdx = OccupancyMapLinearIdx * NumSamples + SampleIdx;
@@ -358,11 +376,22 @@ void FMeshMapBaker::Bake()
 							const FVector2d UVPosition = (FVector2d)OccupancyMap.TexelQueryUV[LinearIdx];
 							const int32 UVTriangleID = OccupancyMap.TexelQueryTriangle[LinearIdx];
 
-							FMeshMapEvaluator::FCorrespondenceSample Sample;
-							DetailCorrespondenceSampler.SampleUV(UVTriangleID, UVPosition, Sample);
-							if (Sample.DetailMesh && DetailSampler->IsTriangle(Sample.DetailMesh, Sample.DetailTriID))
+							// Compute the per-sample correspondence data
+							bool bValidSample = false;
+							FMeshUVSampleInfo SampleInfo;
+							if (MeshUVSampler.QuerySampleInfo(UVTriangleID, UVPosition, SampleInfo))
 							{
-								BakeSample(*TileBuffer, Sample, UVPosition, ImageCoords, OccupancyMap);
+								FMeshMapEvaluator::FCorrespondenceSample Sample;
+								if (ComputeCorrespondenceSample(SampleInfo, Sample))
+								{
+									BakeSample(*TileBuffer, Sample, UVPosition, ImageCoords, OccupancyMap);
+									bValidSample = true;
+								}
+							}
+
+							if (bValidSample == false)
+							{
+								// RecordInvalidSample(Sample, UVPosition, ImageCoords); // TODO
 							}
 						}
 					}
@@ -384,7 +413,7 @@ void FMeshMapBaker::Bake()
 
 		// Accumulate 'Add' float data to image tile buffer
 		OutputQueue.Post(TileIdx, TileBuffer);
-		WriteQueuedOutput(OutputQueue);
+		WriteToOutputBufferQueued(OutputQueue);
 	}, !bParallel ? EParallelForFlags::ForceSingleThread : EParallelForFlags::None);
 
 	if (CancelF())
@@ -407,7 +436,7 @@ void FMeshMapBaker::Bake()
 		// The queue only acquires the process lock if the next item in the queue
 		// is ready. This could mean that there are potential leftovers in the queue
 		// after the parallel for. Write them out now.
-		WriteQueuedOutput(OutputQueue);
+		WriteToOutputBufferQueued(OutputQueue);
 	}
 	
 	if (CancelF())
@@ -515,6 +544,7 @@ void FMeshMapBaker::Bake()
 	}
 }
 
+// Precondition: Must be passed a valid Sample
 void FMeshMapBaker::BakeSample(
 	FMeshMapTileBuffer& TileBuffer,
 	const FMeshMapEvaluator::FCorrespondenceSample& Sample,
@@ -555,20 +585,25 @@ void FMeshMapBaker::BakeSample(
 			const int64 OccupancyMapFilterIdx = OccupancyMap.Tile.GetIndexFromSourceCoords(SourceCoords);
 			const int32 BufferTilePixelUVChart = OccupancyMap.TexelQueryUVChart[OccupancyMapFilterIdx];
 
+			// Get the weight and value buffers for this pixel
 			const int64 BufferTilePixelLinearIdx = Tile.GetIndexFromSourceCoords(SourceCoords);
 			float* PixelBuffer = TileBuffer.GetPixel(BufferTilePixelLinearIdx);
 			float& PixelWeight = TileBuffer.GetPixelWeight(BufferTilePixelLinearIdx);
 
-			// Apply filter using double linear weighting (once per axis)
-			FVector2d TexelDistance = Dimensions.GetTexelUV(SourceCoords) - UVPosition;
-			TexelDistance.X *= Dimensions.GetWidth();
-			TexelDistance.Y *= Dimensions.GetHeight();
+			// Compute the filter weight based on the UV distance from the pixel center to the sample position
+			// Note: There will be no contribution if the sample and pixel are on different UV charts
+			float FilterWeight = Weight * static_cast<float>(SampleUVChart == BufferTilePixelUVChart);
+			{
+				FVector2d TexelDistance = Dimensions.GetTexelUV(SourceCoords) - UVPosition;
+				TexelDistance.X *= Dimensions.GetWidth();
+				TexelDistance.Y *= Dimensions.GetHeight();
+				FilterWeight *= TextureFilterEval(TexelDistance);
+			}
 
-			float FilterWeight = TextureFilterEval(TexelDistance);
-			FilterWeight *= (SampleUVChart == BufferTilePixelUVChart); // Zero the contribution if the sample and pixel are on different UV charts
-			FilterWeight *= Weight;
+			// Update the weight of this pixel
 			PixelWeight += FilterWeight;
 
+			// Update the value of this pixel for each evaluator
 			for (const int32 Idx : EvaluatorIds)
 			{
 				const FMeshMapEvaluator::FEvaluationContext& Context = BakeContexts[Idx];
diff --git a/Engine/Plugins/Runtime/GeometryProcessing/Source/DynamicMesh/Private/Sampling/MeshSurfaceSampler.cpp b/Engine/Plugins/Runtime/GeometryProcessing/Source/DynamicMesh/Private/Sampling/MeshSurfaceSampler.cpp
new file mode 100644
--- /dev/null
+++ b/Engine/Plugins/Runtime/GeometryProcessing/Source/DynamicMesh/Private/Sampling/MeshSurfaceSampler.cpp
@@ -0,0 +1,98 @@
+﻿// Copyright Epic Games, Inc. All Rights Reserved.
+
+#pragma once
+
+#include "Sampling/MeshSurfaceSampler.h"
+
+namespace UE
+{
+namespace Geometry
+{
+
+
+void FMeshSurfaceUVSampler::Initialize(
+	const FDynamicMesh3* MeshIn,
+	const FDynamicMeshUVOverlay* UVOverlayIn,
+	EMeshSurfaceSamplerQueryType QueryTypeIn)
+{
+	Mesh = MeshIn;
+	UVOverlay = UVOverlayIn;
+	QueryType = QueryTypeIn;
+
+	if (QueryType == EMeshSurfaceSamplerQueryType::UVOnly)
+	{
+		InitializeUVMeshSpatial();
+	}
+}
+
+
+
+
+bool FMeshSurfaceUVSampler::QuerySampleInfo(const FVector2d& UV, FMeshUVSampleInfo& SampleInfo)
+{
+	check(QueryType == EMeshSurfaceSamplerQueryType::UVOnly);
+	check(bUVMeshSpatialInitialized);
+
+	FRay3d HitRay(FVector3d(UV.X, UV.Y, 100.0), -FVector3d::UnitZ());
+	const int32 UVTriangleID = UVMeshSpatial.FindNearestHitTriangle(HitRay);
+
+	return QuerySampleInfo(UVTriangleID, UV, SampleInfo);
+}
+
+
+
+
+bool FMeshSurfaceUVSampler::QuerySampleInfo(int32 UVTriangleID, const FVector2d& UV, FMeshUVSampleInfo& SampleInfo)
+{
+	SampleInfo.TriangleIndex = UVTriangleID;
+	if (Mesh->IsTriangle(SampleInfo.TriangleIndex) == false)
+	{
+		return false;
+	}
+	check(UVOverlay->IsSetTriangle(SampleInfo.TriangleIndex));
+
+	SampleInfo.MeshVertices = Mesh->GetTriangle(SampleInfo.TriangleIndex);
+	SampleInfo.Triangle3D = FTriangle3d(
+		Mesh->GetVertex(SampleInfo.MeshVertices.A),
+		Mesh->GetVertex(SampleInfo.MeshVertices.B),
+		Mesh->GetVertex(SampleInfo.MeshVertices.C));
+
+	SampleInfo.UVVertices = UVOverlay->GetTriangle(SampleInfo.TriangleIndex);
+	SampleInfo.TriangleUV = FTriangle2d(
+		(FVector2d)UVOverlay->GetElement(SampleInfo.UVVertices.A),
+		(FVector2d)UVOverlay->GetElement(SampleInfo.UVVertices.B),
+		(FVector2d)UVOverlay->GetElement(SampleInfo.UVVertices.C));
+
+	SampleInfo.BaryCoords = SampleInfo.TriangleUV.GetBarycentricCoords(UV);
+	SampleInfo.SurfacePoint = Mesh->GetTriBaryPoint(
+		SampleInfo.TriangleIndex,
+		SampleInfo.BaryCoords.X,
+		SampleInfo.BaryCoords.Y,
+		SampleInfo.BaryCoords.Z);
+
+	return true;
+}
+
+
+
+
+void FMeshSurfaceUVSampler::InitializeUVMeshSpatial()
+{
+	if (bUVMeshSpatialInitialized)
+	{
+		return;
+	}
+
+	check(UVOverlay);
+
+	UVMeshAdapter.Mesh = Mesh;
+	UVMeshAdapter.UV = UVOverlay;
+	UVMeshSpatial.SetMesh(&UVMeshAdapter, true);
+
+	bUVMeshSpatialInitialized = true;
+}
+
+
+} // end namespace Geometry
+} // end namespace UE
+
diff --git a/Engine/Plugins/Runtime/GeometryProcessing/Source/DynamicMesh/Public/Sampling/MeshBakerCommon.h b/Engine/Plugins/Runtime/GeometryProcessing/Source/DynamicMesh/Public/Sampling/MeshBakerCommon.h
--- a/Engine/Plugins/Runtime/GeometryProcessing/Source/DynamicMesh/Public/Sampling/MeshBakerCommon.h
+++ b/Engine/Plugins/Runtime/GeometryProcessing/Source/DynamicMesh/Public/Sampling/MeshBakerCommon.h
@@ -9,6 +9,7 @@
 #include "Image/ImageTile.h"
 #include "Image/ImageBuilder.h"
 #include "MeshCurvature.h"
+#include "MeshMapEvaluator.h"
 
 
 namespace UE
@@ -147,6 +148,22 @@ public:
 	/** @return true if triangle ray intersection correspondence is supported */
 	virtual bool SupportsRaycastCorrespondence() const = 0;
 
+	/** @return true if a user-defined correspondence test is supported */
+	virtual bool SupportsCustomCorrespondence() const
+	{
+		return false;
+	}
+
+	virtual void* ComputeCustomCorrespondence(const FMeshUVSampleInfo& SampleInfo, FMeshMapEvaluator::FCorrespondenceSample& ValueOut) const
+	{
+		return nullptr;
+	}
+
+	virtual bool IsValidCorrespondence(const FMeshMapEvaluator::FCorrespondenceSample& CorrespondenceSample) const
+	{
+		return true;
+	}
+
 	/**
 	 * @param Point query point
 	 * @param NearestDistSqrOut [Out] returned nearest squared distance, if triangle is found
diff --git a/Engine/Plugins/Runtime/GeometryProcessing/Source/DynamicMesh/Public/Sampling/MeshBaseBaker.h b/Engine/Plugins/Runtime/GeometryProcessing/Source/DynamicMesh/Public/Sampling/MeshBaseBaker.h
--- a/Engine/Plugins/Runtime/GeometryProcessing/Source/DynamicMesh/Public/Sampling/MeshBaseBaker.h
+++ b/Engine/Plugins/Runtime/GeometryProcessing/Source/DynamicMesh/Public/Sampling/MeshBaseBaker.h
@@ -34,7 +34,9 @@ public:
 		/** Use RaycastStandard but fall back to NearestPoint if none of the rays hit */
 		RaycastStandardThenNearest,
 		/** Assume that BakeTarget == DetailMesh and so no mapping is necessary */
-		Identity
+		Identity,
+		/** Use user-defined strategy for finding correspondence */
+		Custom
 	};
 
 	// Setters
diff --git a/Engine/Plugins/Runtime/GeometryProcessing/Source/DynamicMesh/Public/Sampling/MeshMapBaker.h b/Engine/Plugins/Runtime/GeometryProcessing/Source/DynamicMesh/Public/Sampling/MeshMapBaker.h
--- a/Engine/Plugins/Runtime/GeometryProcessing/Source/DynamicMesh/Public/Sampling/MeshMapBaker.h
+++ b/Engine/Plugins/Runtime/GeometryProcessing/Source/DynamicMesh/Public/Sampling/MeshMapBaker.h
@@ -161,7 +161,7 @@ protected:
 	const bool bParallel = true;
 
 	FDynamicMesh3 FlatMesh;
-	TMeshSurfaceUVSampler<FMeshMapEvaluator::FCorrespondenceSample> DetailCorrespondenceSampler;
+	FMeshSurfaceUVSampler MeshUVSampler;
 
 	FImageDimensions Dimensions = FImageDimensions(128, 128);
 
diff --git a/Engine/Plugins/Runtime/GeometryProcessing/Source/DynamicMesh/Public/Sampling/MeshSurfaceSampler.h b/Engine/Plugins/Runtime/GeometryProcessing/Source/DynamicMesh/Public/Sampling/MeshSurfaceSampler.h
--- a/Engine/Plugins/Runtime/GeometryProcessing/Source/DynamicMesh/Public/Sampling/MeshSurfaceSampler.h
+++ b/Engine/Plugins/Runtime/GeometryProcessing/Source/DynamicMesh/Public/Sampling/MeshSurfaceSampler.h
@@ -102,7 +102,7 @@ struct FMeshUVSampleInfo
 };
 
 
-/** Types of query that TMeshSurfaceUVSampler supports */
+/** Types of query that FMeshSurfaceUVSampler/TMeshSurfaceUVSampler supports */
 enum class EMeshSurfaceSamplerQueryType
 {
 	/** Query with arbitrary UV value */
@@ -113,6 +113,57 @@ enum class EMeshSurfaceSamplerQueryType
 
 
 /**
+ * FMeshSurfaceUVSampler computes FMeshUVSampleInfo's on a given mesh at given UV-space positions, this info can then
+ * be used by an external function to compute some application-specific sample.
+ *
+ * Note: This class is similar to TMeshSurfaceUVSampler (the "old" class) but provides a more convenient API. This class
+ * computes the sample info and that's it. The old class computes the sample info internally but also immediately
+ * forwards it to an application-specific sampling function which must be provided as a callback and hence has a fixed
+ * signature and requires that the sample type is passed as a template parameter. Also the old class is more awkward to
+ * use when sampling multiple things, see the documentation of TMeshSurfaceUVSampler for more details.
+ */
+class FMeshSurfaceUVSampler
+{
+public:
+	virtual ~FMeshSurfaceUVSampler() {}
+
+	/**
+	 * Initialize the sampler.
+	 * @param Mesh      Mesh to sample
+	 * @param UVOverlay UV overlay of Mesh
+	 * @param QueryType Type of query functions we will call. Some queries are simpler and do not require spatial data structures
+	 */
+	virtual void Initialize(const FDynamicMesh3* Mesh, const FDynamicMeshUVOverlay* UVOverlay, EMeshSurfaceSamplerQueryType QueryType);
+
+	/**
+	 * Fill in the given SampleInfo for the given UV location.
+	 * @pre The Initialize function must have been called with EMeshSurfaceSamplerQueryType::UVOnly
+	 * @return If the query location is not on a triangle in the Mesh then return false (SampleInfo will be invalid)
+	 */
+	virtual bool QuerySampleInfo(const FVector2d& UV, FMeshUVSampleInfo& SampleInfo);
+
+	/**
+	 * Fill in the given SampleInfo for the given UV location in the given UVTriangleID.
+	 * @return If the query location is not on a triangle in the Mesh then return false (SampleInfo will be invalid)
+	 */
+	virtual bool QuerySampleInfo(int32 UVTriangleID, const FVector2d& UV, FMeshUVSampleInfo& SampleInfo);
+
+protected:
+
+	const FDynamicMesh3* Mesh = nullptr;
+	const FDynamicMeshUVOverlay* UVOverlay = nullptr;
+	EMeshSurfaceSamplerQueryType QueryType = EMeshSurfaceSamplerQueryType::TriangleAndUV;
+
+	// BV tree for finding triangle for a given UV. Not always initialized.
+	FDynamicMeshUVMesh UVMeshAdapter;
+	TMeshAABBTree3<FDynamicMeshUVMesh> UVMeshSpatial;
+	bool bUVMeshSpatialInitialized = false;
+	void InitializeUVMeshSpatial();
+};
+
+/**
+ * Consider using FMeshSurfaceUVSampler instead, it is more flexible and easier to understand!
+ *
  * TMeshSurfaceUVSampler computes point samples of the given SampleType at positions on the mesh
  * based on UV-space positions. The standard use case for this class is to compute samples used
  * in building Normal Maps, AO Maps, etc.
diff --git a/Engine/Plugins/Runtime/MeshModelingToolset/Source/ModelingComponents/Private/Scene/SceneCapturePhotoSet.cpp b/Engine/Plugins/Runtime/MeshModelingToolset/Source/ModelingComponents/Private/Scene/SceneCapturePhotoSet.cpp
--- a/Engine/Plugins/Runtime/MeshModelingToolset/Source/ModelingComponents/Private/Scene/SceneCapturePhotoSet.cpp
+++ b/Engine/Plugins/Runtime/MeshModelingToolset/Source/ModelingComponents/Private/Scene/SceneCapturePhotoSet.cpp
@@ -326,6 +326,14 @@ bool FSceneCapturePhotoSet::ComputeSample(
 	return true;
 }
 
+bool FSceneCapturePhotoSet::IsValidSample(
+	const FVector3d& Position,
+	const FVector3d& Normal,
+	TFunctionRef<bool(const FVector3d&, const FVector3d&)> VisibilityFunction) const
+{
+	return BaseColorPhotoSet.IsValidSample(Position, Normal, VisibilityFunction);
+}
+
 void FSceneCapturePhotoSet::SetEnableWriteDebugImages(bool bEnable, FString FolderName)
 {
 	bWriteDebugImages = bEnable;
diff --git a/Engine/Plugins/Runtime/MeshModelingToolset/Source/ModelingComponents/Public/Scene/SceneCapturePhotoSet.h b/Engine/Plugins/Runtime/MeshModelingToolset/Source/ModelingComponents/Public/Scene/SceneCapturePhotoSet.h
--- a/Engine/Plugins/Runtime/MeshModelingToolset/Source/ModelingComponents/Public/Scene/SceneCapturePhotoSet.h
+++ b/Engine/Plugins/Runtime/MeshModelingToolset/Source/ModelingComponents/Public/Scene/SceneCapturePhotoSet.h
@@ -112,6 +112,11 @@ public:
 		TFunctionRef<bool(const FVector3d&, const FVector3d&)> VisibilityFunction,
 		FSceneSample& DefaultsInResultsOut) const;
 
+	bool IsValidSample(
+		const FVector3d& Position,
+		const FVector3d& Normal,
+		TFunctionRef<bool(const FVector3d&, const FVector3d&)> VisibilityFunction) const;
+
 
 	const FSpatialPhotoSet3f& GetBaseColorPhotoSet() { return BaseColorPhotoSet; }
 	const FSpatialPhotoSet1f& GetRoughnessPhotoSet() { return RoughnessPhotoSet; }
diff --git a/Engine/Source/Runtime/GeometryCore/Public/Image/SpatialPhotoSet.h b/Engine/Source/Runtime/GeometryCore/Public/Image/SpatialPhotoSet.h
--- a/Engine/Source/Runtime/GeometryCore/Public/Image/SpatialPhotoSet.h
+++ b/Engine/Source/Runtime/GeometryCore/Public/Image/SpatialPhotoSet.h
@@ -81,6 +81,12 @@ public:
 		const PixelType& DefaultValue
 	) const;
 
+	bool IsValidSample(
+		const FVector3d& Position, 
+		const FVector3d& Normal, 
+		TFunctionRef<bool(const FVector3d&,const FVector3d&)> VisibilityFunction
+	) const;
+
 protected:
 	TArray<TSharedPtr<TSpatialPhoto<PixelType>, ESPMode::ThreadSafe>> Photos;
 };
@@ -154,5 +160,62 @@ PixelType TSpatialPhotoSet<PixelType, RealType>::ComputeSample(
 }
 
 
+// TODO Consider refactoring this to deduplicate code from ComputeSample, or not, since this is called in a tight inner loop
+template<typename PixelType, typename RealType>
+bool TSpatialPhotoSet<PixelType, RealType>::IsValidSample(
+	const FVector3d& Position, 
+	const FVector3d& Normal, 
+	TFunctionRef<bool(const FVector3d&, const FVector3d&)> VisibilityFunction) const
+{
+	double DotTolerance = -0.1;		// dot should be negative for normal pointing towards photo
+
+	double MinDot = 1.0;
+
+	int32 NumPhotos = Num();
+	for (int32 pi = 0; pi < NumPhotos; ++pi)
+	{
+		const TSpatialPhoto<PixelType>& Photo = *Photos[pi];
+		check(Photo.Dimensions.IsSquare());
+
+		FVector3d ViewDirection = Photo.Frame.X();
+		double ViewDot = ViewDirection.Dot(Normal);
+		if (ViewDot > DotTolerance || ViewDot > MinDot)
+		{
+			continue;
+		}
+
+		FFrame3d ViewPlane = Photo.Frame;
+		ViewPlane.Origin += Photo.NearPlaneDist * ViewDirection;
+
+		double ViewPlaneWidthWorld = Photo.NearPlaneDist * FMathd::Tan(Photo.HorzFOVDegrees * 0.5 * FMathd::DegToRad);
+		double ViewPlaneHeightWorld = ViewPlaneWidthWorld;
+
+		FVector3d RayOrigin = Photo.Frame.Origin;
+		FVector3d RayDir = Normalized(Position - RayOrigin);
+		FVector3d HitPoint;
+		bool bHit = ViewPlane.RayPlaneIntersection(RayOrigin, RayDir, 0, HitPoint);
+		if (bHit)
+		{
+			bool bVisible = VisibilityFunction(Position, HitPoint);
+			if ( bVisible )
+			{
+				double PlaneX = (HitPoint - ViewPlane.Origin).Dot(ViewPlane.Y());
+				double PlaneY = (HitPoint - ViewPlane.Origin).Dot(ViewPlane.Z());
+
+				//FVector2d PlanePos = ViewPlane.ToPlaneUV(HitPoint, 0);
+				double u = PlaneX / ViewPlaneWidthWorld;
+				double v = -(PlaneY / ViewPlaneHeightWorld);
+				if (FMathd::Abs(u) < 1 && FMathd::Abs(v) < 1)
+				{
+					return true;
+				}
+			}
+		}
+	}
+
+	return false;
+}
+
+
 } // end namespace UE::Geometry
 } // end namespace UE
\ No newline at end of file
