Remove unnecessary pragma once

diff --git a/Engine/Plugins/Runtime/GeometryProcessing/Source/DynamicMesh/Private/Sampling/MeshSurfaceSampler.cpp b/Engine/Plugins/Runtime/GeometryProcessing/Source/DynamicMesh/Private/Sampling/MeshSurfaceSampler.cpp
--- a/Engine/Plugins/Runtime/GeometryProcessing/Source/DynamicMesh/Private/Sampling/MeshSurfaceSampler.cpp
+++ b/Engine/Plugins/Runtime/GeometryProcessing/Source/DynamicMesh/Private/Sampling/MeshSurfaceSampler.cpp
@@ -1,7 +1,5 @@
 ﻿// Copyright Epic Games, Inc. All Rights Reserved.
 
-#pragma once
-
 #include "Sampling/MeshSurfaceSampler.h"
 
 namespace UE
