Fix crash in PolyEdit offset w/ bowties

diff --git a/Engine/Plugins/Runtime/GeometryProcessing/Source/DynamicMesh/Private/Operations/OffsetMeshRegion.cpp b/Engine/Plugins/Runtime/GeometryProcessing/Source/DynamicMesh/Private/Operations/OffsetMeshRegion.cpp
--- a/Engine/Plugins/Runtime/GeometryProcessing/Source/DynamicMesh/Private/Operations/OffsetMeshRegion.cpp
+++ b/Engine/Plugins/Runtime/GeometryProcessing/Source/DynamicMesh/Private/Operations/OffsetMeshRegion.cpp
@@ -282,6 +282,7 @@ bool FOffsetMeshRegion::ApplyOffset(FOffsetInfo& Region, FMeshNormals* UseNormal
 
 	// Stitch the loops
 
+	bool bSuccess = true;
 	int NumInitialLoops = InitialLoops.GetLoopCount();
 	Region.BaseLoops.SetNum(NumInitialLoops);
 	Region.OffsetLoops.SetNum(NumInitialLoops);
@@ -299,7 +300,12 @@ bool FOffsetMeshRegion::ApplyOffset(FOffsetInfo& Region, FMeshNormals* UseNormal
 
 		// stitch the loops
 		FDynamicMeshEditResult StitchResult;
-		Editor.StitchVertexLoopsMinimal(OffsetLoopV, BaseLoopV, StitchResult);
+		bool bStitchSuccess = Editor.StitchVertexLoopsMinimal(OffsetLoopV, BaseLoopV, StitchResult);
+		if (!bStitchSuccess)
+		{
+			bSuccess = false;
+			continue;
+		}
 
 		// set the groups of the new quads along the stitch
 		int NumNewQuads = StitchResult.NewQuads.Num();
@@ -373,7 +379,7 @@ bool FOffsetMeshRegion::ApplyOffset(FOffsetInfo& Region, FMeshNormals* UseNormal
 		}
 	}
 
-	return true;
+	return bSuccess;
 }
 
 bool FOffsetMeshRegion::EdgesSeparateSameGroupsAndAreColinearAtBorder(FDynamicMesh3* Mesh, 
