GeometryCore: Fix static analysis warning in FDynamicMeshEditor

diff --git a/Engine/Plugins/Runtime/GeometryProcessing/Source/DynamicMesh/Private/DynamicMeshEditor.cpp b/Engine/Plugins/Runtime/GeometryProcessing/Source/DynamicMesh/Private/DynamicMeshEditor.cpp
--- a/Engine/Plugins/Runtime/GeometryProcessing/Source/DynamicMesh/Private/DynamicMeshEditor.cpp
+++ b/Engine/Plugins/Runtime/GeometryProcessing/Source/DynamicMesh/Private/DynamicMeshEditor.cpp
@@ -460,7 +460,8 @@ bool FDynamicMeshEditor::DisconnectTriangles(const TArray<int>& Triangles, TArra
 			{
 				checkSlow(!Mesh->SplitVertexWouldLeaveIsolated(VertID, FilteredTriangles));
 				DynamicMeshInfo::FVertexSplitInfo SplitInfo;
-				ensure(EMeshResult::Ok == Mesh->SplitVertex(VertID, FilteredTriangles, SplitInfo));
+				const EMeshResult MeshResult = Mesh->SplitVertex(VertID, FilteredTriangles, SplitInfo);
+				ensure(MeshResult == EMeshResult::Ok);
 
 				int NewVertID = SplitInfo.NewVertex;
 				LoopVertexMap.Add(Loop.Vertices[vi], NewVertID);
