Fixed the PVS static analyzer error caused by a dangling reference.
Added extra comments clarifying the expected lifetime of the passed objects to the FSparsePowerMethod class constructors and the CreateSparseMatrixOperator() static method.

#rnx

diff --git a/Engine/Plugins/Runtime/GeometryProcessing/Source/DynamicMesh/Private/Solvers/Internal/PowerMethodSolver.cpp b/Engine/Plugins/Runtime/GeometryProcessing/Source/DynamicMesh/Private/Solvers/Internal/PowerMethodSolver.cpp
--- a/Engine/Plugins/Runtime/GeometryProcessing/Source/DynamicMesh/Private/Solvers/Internal/PowerMethodSolver.cpp
+++ b/Engine/Plugins/Runtime/GeometryProcessing/Source/DynamicMesh/Private/Solvers/Internal/PowerMethodSolver.cpp
@@ -239,7 +239,7 @@ void FSparsePowerMethod::CreateSparseMatrixOperator(const FSparseMatrixD& InMatr
         OutVector = InMatrix * InVector; 
     };
 
-    OutMatrixOp.Factorize = [&InMatrix, &MatrixSolverType, bIsSymmetric]() 
+    OutMatrixOp.Factorize = [&InMatrix, MatrixSolverType, bIsSymmetric]() 
     {	
         TUniquePtr<IMatrixSolverBase> Solver = ContructMatrixSolver(MatrixSolverType);
         Solver->SetUp(InMatrix, bIsSymmetric);
diff --git a/Engine/Plugins/Runtime/GeometryProcessing/Source/DynamicMesh/Private/Solvers/Internal/PowerMethodSolver.h b/Engine/Plugins/Runtime/GeometryProcessing/Source/DynamicMesh/Private/Solvers/Internal/PowerMethodSolver.h
--- a/Engine/Plugins/Runtime/GeometryProcessing/Source/DynamicMesh/Private/Solvers/Internal/PowerMethodSolver.h
+++ b/Engine/Plugins/Runtime/GeometryProcessing/Source/DynamicMesh/Private/Solvers/Internal/PowerMethodSolver.h
@@ -178,6 +178,9 @@ protected:
  *		{
  *			// OutValue, OutVector is a valid eigenpair
  *		}
+ *
+ * @note The MatrixOperator TFunctions capture the passed-in matrices by reference, so the lifetime of these matrices 
+ * must extend through the use of the FSparsePowerMethod class methods (specifically the Solve() method).
  */
 class DYNAMICMESH_API FSparsePowerMethod : public FPowerMethod
 {
@@ -230,7 +233,12 @@ public:
 	{
 	}
 
-	/** Helper method to setup Matrix Operator for the sparse real matrix */
+    /** 
+     * Convenience method to set up the Product(), Factorize() and Rows() TFunctions for the MatrixOperator.
+     * 
+     * @note InMatrix is captured by reference by all TFunctions and hence we assume that the lifetime of the
+     * referenced matrix extends past the use of the OutMatrixOp.
+     */
 	static void CreateSparseMatrixOperator(const FSparseMatrixD& InMatrix, 
 										   const FMatrixHints& InHints, 
 										   MatrixOperator& OutMatrixOp);
