Fix fracture messing up normals for some meshes by making the mesh boolean's simplifier preserve vertex normals

diff --git a/Engine/Plugins/Runtime/GeometryProcessing/Source/DynamicMesh/Private/Operations/MeshBoolean.cpp b/Engine/Plugins/Runtime/GeometryProcessing/Source/DynamicMesh/Private/Operations/MeshBoolean.cpp
--- a/Engine/Plugins/Runtime/GeometryProcessing/Source/DynamicMesh/Private/Operations/MeshBoolean.cpp
+++ b/Engine/Plugins/Runtime/GeometryProcessing/Source/DynamicMesh/Private/Operations/MeshBoolean.cpp
@@ -885,8 +885,8 @@ bool FMeshBoolean::CollapseWouldHurtTriangleQuality(
 bool FMeshBoolean::CollapseWouldChangeShapeOrUVs(
 	const FDynamicMesh3& Mesh, const TSet<int>& CutBoundaryEdgeSet, double DotTolerance, int SourceEID,
 	int32 RemoveV, const FVector3d& RemoveVPos, int32 KeepV, const FVector3d& KeepVPos, const FVector3d& EdgeDir,
-	bool bPreserveTriangleGroups, bool bPreserveUVsForMesh, bool bPreserveVertexUVs, bool bPreserveOverlayUVs, 
-	float UVEqualThresholdSq)
+	bool bPreserveTriangleGroups, bool bPreserveUVsForMesh, bool bPreserveVertexUVs, bool bPreserveOverlayUVs,
+	float UVEqualThresholdSq, bool bPreserveVertexNormals, float NormalEqualCosThreshold)
 {
 	// Search the edges connected to the vertex to find one in the boundary set that points in the opposite direction
 	// If we don't find that edge, or if there are other boundary/seam edges attached, we can't remove this vertex
@@ -960,6 +960,17 @@ bool FMeshBoolean::CollapseWouldChangeShapeOrUVs(
 						return;
 					}
 				}
+				if (bPreserveVertexNormals && Mesh.HasVertexNormals())
+				{
+					FVector3f OtherN = Mesh.GetVertexNormal(OtherV);
+					FVector3f RemoveN = Mesh.GetVertexNormal(RemoveV);
+					FVector3f KeepN = Mesh.GetVertexNormal(KeepV);
+					if (Normalized(Lerp(OtherN, KeepN, LerpT)).Dot(Normalized(RemoveN)) < NormalEqualCosThreshold)
+					{
+						bHasBadEdge = true;
+						return;
+					}
+				}
 				if (bPreserveOverlayUVs && Mesh.HasAttributes())
 				{
 					int NumLayers = Mesh.Attributes()->NumUVLayers();
@@ -1154,7 +1165,8 @@ void FMeshBoolean::SimplifyAlongNewEdges(int NumMeshesToProcess, FDynamicMesh3*
 						*CutMesh[MeshIdx], CutBoundaryEdgeSets[MeshIdx], DotTolerance,
 						SourceEID, RemoveV, RemoveVPos, KeepV, KeepVPos, EdgeDir, bPreserveTriangleGroups,
 						PreserveUVsOnlyForMesh == -1 || MeshIdx == PreserveUVsOnlyForMesh,
-						bPreserveVertexUVs, bPreserveOverlayUVs, UVDistortTolerance* UVDistortTolerance);
+						bPreserveVertexUVs, bPreserveOverlayUVs, UVDistortTolerance * UVDistortTolerance,
+						bPreserveVertexNormals, FMathf::Cos(NormalDistortTolerance * FMathf::DegToRad));
 				};
 
 				if (bHasBadEdge)
diff --git a/Engine/Plugins/Runtime/GeometryProcessing/Source/DynamicMesh/Private/Operations/MeshSelfUnion.cpp b/Engine/Plugins/Runtime/GeometryProcessing/Source/DynamicMesh/Private/Operations/MeshSelfUnion.cpp
--- a/Engine/Plugins/Runtime/GeometryProcessing/Source/DynamicMesh/Private/Operations/MeshSelfUnion.cpp
+++ b/Engine/Plugins/Runtime/GeometryProcessing/Source/DynamicMesh/Private/Operations/MeshSelfUnion.cpp
@@ -597,7 +597,8 @@ void FMeshSelfUnion::SimplifyAlongNewEdges(TArray<int>& CutBoundaryEdges, TMap<i
 					bHasBadEdge = bHasBadEdge || FMeshBoolean::CollapseWouldChangeShapeOrUVs(
 						*Mesh, CutBoundaryEdgeSet, DotTolerance,
 						SourceEID, RemoveV, RemoveVPos, KeepV, KeepVPos, EdgeDir, bPreserveTriangleGroups,
-						true, bPreserveVertexUVs, bPreserveOverlayUVs, UVDistortTolerance * UVDistortTolerance);
+						true, bPreserveVertexUVs, bPreserveOverlayUVs, UVDistortTolerance * UVDistortTolerance,
+						bPreserveVertexNormals, FMathf::Cos(NormalDistortTolerance * FMathf::DegToRad));
 				}
 
 				if (bHasBadEdge)
diff --git a/Engine/Plugins/Runtime/GeometryProcessing/Source/DynamicMesh/Public/Operations/MeshBoolean.h b/Engine/Plugins/Runtime/GeometryProcessing/Source/DynamicMesh/Public/Operations/MeshBoolean.h
--- a/Engine/Plugins/Runtime/GeometryProcessing/Source/DynamicMesh/Public/Operations/MeshBoolean.h
+++ b/Engine/Plugins/Runtime/GeometryProcessing/Source/DynamicMesh/Public/Operations/MeshBoolean.h
@@ -94,6 +94,10 @@ public:
 	bool bPreserveOverlayUVs = true;
 	/** When preserving UVs, sets maximum allowed change in UV coordinates from collapsing an edge, measured at the removed vertex */
 	float UVDistortTolerance = FMathf::ZeroTolerance;
+	/** Prevent simplification from distorting vertex normals */
+	bool bPreserveVertexNormals = true;
+	/** When preserving normals, sets maximum allowed change in normals from collapsing an edge, measured at the removed vertex in degrees */
+	float NormalDistortTolerance = .01f;
 	/** If > -1, then only preserve the UVs of one of the input meshes.  Useful when cutting an artist-created asset w/ procedural geometry. */
 	int PreserveUVsOnlyForMesh = -1;
 
@@ -223,7 +227,8 @@ public:
 		const FDynamicMesh3& Mesh, const TSet<int>& CutBoundaryEdgeSet, double DotTolerance, int SourceEID,
 		int32 RemoveV, const FVector3d& RemoveVPos, int32 KeepV, const FVector3d& KeepVPos,
 		const FVector3d& EdgeDir, bool bPreserveTriangleGroups, bool bPreserveUVsForMesh,
-		bool bPreserveVertexUVs, bool bPreserveOverlayUVs, float UVToleranceSq);
+		bool bPreserveVertexUVs, bool bPreserveOverlayUVs, float UVToleranceSq,
+		bool bPreserveVertexNormals, float NormalEqualCosThreshold);
 };
 
 
diff --git a/Engine/Plugins/Runtime/GeometryProcessing/Source/DynamicMesh/Public/Operations/MeshSelfUnion.h b/Engine/Plugins/Runtime/GeometryProcessing/Source/DynamicMesh/Public/Operations/MeshSelfUnion.h
--- a/Engine/Plugins/Runtime/GeometryProcessing/Source/DynamicMesh/Public/Operations/MeshSelfUnion.h
+++ b/Engine/Plugins/Runtime/GeometryProcessing/Source/DynamicMesh/Public/Operations/MeshSelfUnion.h
@@ -80,6 +80,10 @@ public:
 	bool bPreserveOverlayUVs = true;
 	/** When preserving UVs, sets maximum allowed change in UV coordinates from collapsing an edge, measured at the removed vertex */
 	float UVDistortTolerance = FMathf::ZeroTolerance;
+	/** Prevent simplification from distorting vertex normals */
+	bool bPreserveVertexNormals = true;
+	/** When preserving normals, sets maximum allowed change in normals from collapsing an edge, measured at the removed vertex in degrees */
+	float NormalDistortTolerance = .01f;
 
 
 	//
