Misc fixes for console module builds

diff --git a/Engine/Plugins/FX/Niagara/Source/Niagara/Public/NiagaraWorldManager.h b/Engine/Plugins/FX/Niagara/Source/Niagara/Public/NiagaraWorldManager.h
--- a/Engine/Plugins/FX/Niagara/Source/Niagara/Public/NiagaraWorldManager.h
+++ b/Engine/Plugins/FX/Niagara/Source/Niagara/Public/NiagaraWorldManager.h
@@ -55,7 +55,7 @@ struct TStructOpsTypeTraits<FNiagaraWorldManagerTickFunction> : public TStructOp
 
 using FNiagaraSystemSimulationRef = TSharedRef<FNiagaraSystemSimulation, ESPMode::ThreadSafe>;
 
-enum class NIAGARA_API ENiagaraScalabilityCullingMode : uint8
+enum class ENiagaraScalabilityCullingMode : uint8
 {
 	/* Scalability culling is enabled as normal. */
 	Enabled,
diff --git a/Engine/Plugins/Runtime/GeometryProcessing/Source/DynamicMesh/Public/Sampling/MeshNormalMapEvaluator.h b/Engine/Plugins/Runtime/GeometryProcessing/Source/DynamicMesh/Public/Sampling/MeshNormalMapEvaluator.h
--- a/Engine/Plugins/Runtime/GeometryProcessing/Source/DynamicMesh/Public/Sampling/MeshNormalMapEvaluator.h
+++ b/Engine/Plugins/Runtime/GeometryProcessing/Source/DynamicMesh/Public/Sampling/MeshNormalMapEvaluator.h
@@ -19,6 +19,12 @@ namespace Geometry
 class DYNAMICMESH_API FMeshNormalMapEvaluator : public FMeshMapEvaluator
 {
 public:
+	PRAGMA_DISABLE_DEPRECATION_WARNINGS
+	FMeshNormalMapEvaluator() = default;
+	FMeshNormalMapEvaluator(const FMeshNormalMapEvaluator&) = default;
+	FMeshNormalMapEvaluator& operator=( const FMeshNormalMapEvaluator& ) = default;
+	PRAGMA_ENABLE_DEPRECATION_WARNINGS
+
 	// Begin FMeshMapEvaluator interface
 	virtual void Setup(const FMeshBaseBaker& Baker, FEvaluationContext& Context) override;
 
diff --git a/Engine/Plugins/Runtime/GeometryProcessing/Source/DynamicMesh/Public/Sampling/MeshPropertyMapEvaluator.h b/Engine/Plugins/Runtime/GeometryProcessing/Source/DynamicMesh/Public/Sampling/MeshPropertyMapEvaluator.h
--- a/Engine/Plugins/Runtime/GeometryProcessing/Source/DynamicMesh/Public/Sampling/MeshPropertyMapEvaluator.h
+++ b/Engine/Plugins/Runtime/GeometryProcessing/Source/DynamicMesh/Public/Sampling/MeshPropertyMapEvaluator.h
@@ -32,6 +32,12 @@ public:
 	EMeshPropertyMapType Property = EMeshPropertyMapType::Normal;
 
 public:
+	PRAGMA_DISABLE_DEPRECATION_WARNINGS
+	FMeshPropertyMapEvaluator() = default;
+	FMeshPropertyMapEvaluator(const FMeshPropertyMapEvaluator&) = default;
+	FMeshPropertyMapEvaluator& operator=(const FMeshPropertyMapEvaluator&) = default;
+	PRAGMA_ENABLE_DEPRECATION_WARNINGS
+
 	// Begin FMeshMapEvaluator interface
 	virtual void Setup(const FMeshBaseBaker& Baker, FEvaluationContext& Context) override;
 
diff --git a/Engine/Source/Runtime/Engine/Public/SkeletalMeshMerge.h b/Engine/Source/Runtime/Engine/Public/SkeletalMeshMerge.h
--- a/Engine/Source/Runtime/Engine/Public/SkeletalMeshMerge.h
+++ b/Engine/Source/Runtime/Engine/Public/SkeletalMeshMerge.h
@@ -158,6 +158,10 @@ public:
     		PRAGMA_ENABLE_DEPRECATION_WARNINGS
     		);
 
+	PRAGMA_DISABLE_DEPRECATION_WARNINGS
+	FSkeletalMeshMerge(const FSkeletalMeshMerge&) = default;
+	PRAGMA_ENABLE_DEPRECATION_WARNINGS
+
 	/**
 	 * Merge/Composite skeleton and meshes together from the list of source meshes.
 	 * @param RefPoseOverrides - An optional override for the merged skeleton's reference pose.
diff --git a/Engine/Source/Runtime/Experimental/Dataflow/Core/Public/Dataflow/DataflowInputOutput.h b/Engine/Source/Runtime/Experimental/Dataflow/Core/Public/Dataflow/DataflowInputOutput.h
--- a/Engine/Source/Runtime/Experimental/Dataflow/Core/Public/Dataflow/DataflowInputOutput.h
+++ b/Engine/Source/Runtime/Experimental/Dataflow/Core/Public/Dataflow/DataflowInputOutput.h
@@ -29,7 +29,7 @@ namespace Dataflow
 	};
 
 	template<class T>
-	class DATAFLOWCORE_API TInput : public FConnection
+	class TInput : public FConnection
 	{
 		typedef FConnection Super;
 		friend class FConnection;
