MeshModelingTools: Make offset and inset work in pathological bowtie-related cases. Fix DisconnectTriangles in cases of bowties.

diff --git a/Engine/Plugins/Runtime/GeometryProcessing/Source/DynamicMesh/Private/DynamicMeshEditor.cpp b/Engine/Plugins/Runtime/GeometryProcessing/Source/DynamicMesh/Private/DynamicMeshEditor.cpp
--- a/Engine/Plugins/Runtime/GeometryProcessing/Source/DynamicMesh/Private/DynamicMeshEditor.cpp
+++ b/Engine/Plugins/Runtime/GeometryProcessing/Source/DynamicMesh/Private/DynamicMeshEditor.cpp
@@ -48,35 +48,29 @@ bool FDynamicMeshEditor::RemoveIsolatedVertices()
 	return bSuccess;
 }
 
-
-bool FDynamicMeshEditor::StitchVertexLoopsMinimal(const TArray<int>& Loop1, const TArray<int>& Loop2, FDynamicMeshEditResult& ResultOut)
+namespace DynamicMeshEditorLocals
 {
-	int N = Loop1.Num();
-	checkf(N == Loop2.Num(), TEXT("FDynamicMeshEditor::StitchVertexLoopsMinimal: loops are not the same length!"));
-	if (N != Loop2.Num())
-	{
-		return false;
-	}
-
-	ResultOut.NewQuads.Reserve(N);
-	ResultOut.NewGroups.Reserve(N);
+// Does all the work of StitchVertexLoopsMinimal and StitchVertexLoopToTriVidPairSequence since they only differ in the way
+// that they pick verts per quad.
+template <typename FuncType>
+bool StitchLoopsInternal(FDynamicMeshEditor& Editor, int32 NumQuads, FuncType&& GetQuadVidsForIndex, FDynamicMeshEditResult& ResultOut)
+{
+	ResultOut.NewQuads.Reserve(NumQuads);
+	ResultOut.NewGroups.Reserve(NumQuads);
 
-	int i = 0;
-	for (; i < N; ++i) 
+	for (int i = 0; i < NumQuads; ++i)
 	{
-		int a = Loop1[i];
-		int b = Loop1[(i + 1) % N];
-		int c = Loop2[i];
-		int d = Loop2[(i + 1) % N];
+		int32 a, b, c, d;
+		GetQuadVidsForIndex(i, a, b, c, d);
 
-		int NewGroupID = Mesh->AllocateTriangleGroup();
+		int NewGroupID = Editor.Mesh->AllocateTriangleGroup();
 		ResultOut.NewGroups.Add(NewGroupID);
 
 		FIndex3i t1(b, a, d);
-		int tid1 = Mesh->AppendTriangle(t1, NewGroupID);
+		int tid1 = Editor.Mesh->AppendTriangle(t1, NewGroupID);
 
 		FIndex3i t2(a, c, d);
-		int tid2 = Mesh->AppendTriangle(t2, NewGroupID);
+		int tid2 = Editor.Mesh->AppendTriangle(t2, NewGroupID);
 
 		ResultOut.NewQuads.Add(FIndex2i(tid1, tid2));
 
@@ -90,7 +84,7 @@ bool FDynamicMeshEditor::StitchVertexLoopsMinimal(const TArray<int>& Loop1, cons
 
 operation_failed:
 	// remove what we added so far
-	if (ResultOut.NewQuads.Num()) 
+	if (ResultOut.NewQuads.Num())
 	{
 		TArray<int> Triangles; Triangles.Reserve(2*ResultOut.NewQuads.Num());
 		for (const FIndex2i& QuadTriIndices : ResultOut.NewQuads)
@@ -98,14 +92,81 @@ operation_failed:
 			Triangles.Add(QuadTriIndices.A);
 			Triangles.Add(QuadTriIndices.B);
 		}
-		if (!RemoveTriangles(Triangles, false))
+		if (!Editor.RemoveTriangles(Triangles, false))
 		{
-			checkf(false, TEXT("FDynamicMeshEditor::StitchVertexLoopsMinimal: failed to add all triangles, and also failed to back out changes."));
+			ensureMsgf(false, TEXT("FDynamicMeshEditor::StitchVertexLoopsMinimal: failed to add all triangles, and also failed to back out changes."));
 		}
 	}
 	return false;
 }
+}
+
+bool FDynamicMeshEditor::StitchVertexLoopsMinimal(const TArray<int>& Loop1, const TArray<int>& Loop2, FDynamicMeshEditResult& ResultOut)
+{
+	int N = Loop1.Num();
+	if (!ensureMsgf(N == Loop2.Num(), TEXT("FDynamicMeshEditor::StitchVertexLoopsMinimal: loops are not the same length!")))
+	{
+		return false;
+	}
+	return DynamicMeshEditorLocals::StitchLoopsInternal(*this, N, [N, &Loop1, &Loop2]
+	(int32 Index, int32& VertA, int32& VertB, int32& VertC, int32& VertD) {
+		VertA = Loop1[Index];
+		VertB = Loop1[(Index + 1) % N];
+		VertC = Loop2[Index];
+		VertD = Loop2[(Index + 1) % N];
+		}, ResultOut);
+}
+
+bool FDynamicMeshEditor::StitchVertexLoopToTriVidPairSequence(
+	const TArray<TPair<int32, TPair<int8, int8>>>& TriVidPairs,
+	const TArray<int>& VertexLoop, FDynamicMeshEditResult& ResultOut)
+{
+	int N = TriVidPairs.Num();
+	if (!ensureMsgf(N == VertexLoop.Num(), TEXT("FDynamicMeshEditor::StitchVertexLoopToEdgeSequence: sequences are not the same length!")))
+	{
+		return false;
+	}
+	return DynamicMeshEditorLocals::StitchLoopsInternal(*this, N, [this, N, &TriVidPairs, &VertexLoop]
+	(int32 Index, int32& VertA, int32& VertB, int32& VertC, int32& VertD) {
+		FIndex3i TriVids1 = Mesh->GetTriangle(TriVidPairs[Index].Key);
+		VertA = TriVids1[TriVidPairs[Index].Value.Key];
+		VertB = TriVids1[TriVidPairs[Index].Value.Value];
+
+		VertC = VertexLoop[Index];
+		VertD = VertexLoop[(Index + 1) % N];
+		}, ResultOut);
+}
+
+bool FDynamicMeshEditor::ConvertLoopToTriVidPairSequence(const FDynamicMesh3& Mesh, const TArray<int32>& VidLoop,
+	const TArray<int32>& EdgeLoop, TArray<TPair<int32, TPair<int8, int8>>>& TriVertPairsOut)
+{
+	if (!ensure(EdgeLoop.Num() == VidLoop.Num()))
+	{
+		return false;
+	}
 
+	for (int32 QuadIndex = 0; QuadIndex < EdgeLoop.Num(); ++QuadIndex)
+	{
+		int32 Tid = Mesh.GetEdgeT(EdgeLoop[QuadIndex]).A;
+		int32 FirstVid = VidLoop[QuadIndex];
+		int32 SecondVid = VidLoop[(QuadIndex + 1) % VidLoop.Num()];
+
+		FIndex3i TriVids = Mesh.GetTriangle(Tid);
+		int8 SubIdx1 = IndexUtil::FindTriIndex(FirstVid, TriVids);
+		int8 SubIdx2 = IndexUtil::FindTriIndex(SecondVid, TriVids);
+		if (ensure(SubIdx1 >= 0 && SubIdx2 >= 0))
+		{
+			TriVertPairsOut.Emplace(Tid,
+				TPair<int8, int8>(SubIdx1, SubIdx2));
+		}
+		else
+		{
+			return false;
+		}
+
+	}
+	return true;
+}
 
 
 
@@ -428,6 +489,7 @@ bool FDynamicMeshEditor::DisconnectTriangles(const TArray<int>& Triangles, TArra
 	int NumLoops = Loops.Num();
 	LoopSetOut.SetNum(NumLoops);
 	TArray<int> FilteredTriangles;
+	TMap<int32, int32> OldVidsToNewVids; 
 	for ( int li = 0; li < NumLoops; ++li)
 	{
 		FEdgeLoop& Loop = Loops[li];
@@ -439,12 +501,40 @@ bool FDynamicMeshEditor::DisconnectTriangles(const TArray<int>& Triangles, TArra
 
 		// duplicate the vertices
 		int NumVertices = Loop.Vertices.Num();
-		TMap<int, int> LoopVertexMap; LoopVertexMap.Reserve(NumVertices);
 		TArray<int> NewVertexLoop; NewVertexLoop.SetNum(NumVertices);
 		for (int vi = 0; vi < NumVertices; ++vi)
 		{
 			int VertID = Loop.Vertices[vi];
 
+			// See if we've processed this vertex already as part of a bowtie
+			int32* ExistingNewVertID = OldVidsToNewVids.Find(VertID);
+			if (ExistingNewVertID)
+			{
+				// We've already done the split, but we still need to update the loop pairs. The way we
+				// do this depends on whether the vert was originally a boundary vert or not, because
+				// we treat these cases differently when we perform splits.
+				// If it was originally a boundary vert, the new vertex ended up as an isolated
+				// vertex.
+				if (!Mesh->IsReferencedVertex(*ExistingNewVertID))
+				{
+					// The isolated vertex should be in the outer loop
+					LoopPair.OuterVertices[vi] = *ExistingNewVertID;
+					LoopPair.OuterEdges[vi] = FDynamicMesh3::InvalidID;
+					LoopPair.OuterEdges[(vi == 0) ? NumVertices - 1 : vi - 1] = FDynamicMesh3::InvalidID;
+					NewVertexLoop[vi] = VertID;
+				}
+				else
+				{
+					// Otherwise, the new vertex should end up in the inner loop (and not isolated, because
+					// it is attached to the selected triangles).
+					NewVertexLoop[vi] = *ExistingNewVertID;
+				}
+
+				continue;
+			}
+			// If we get here, we still need to split the vertex.
+
+			// See how many of the adjacent triangles are in our selection.
 			FilteredTriangles.Reset();
 			int TriRingCount = 0;
 			for (int RingTID : Mesh->VtxTrianglesItr(VertID))
@@ -455,16 +545,17 @@ bool FDynamicMeshEditor::DisconnectTriangles(const TArray<int>& Triangles, TArra
 				}
 				TriRingCount++;
 			}
-			bool bIsSubset = (FilteredTriangles.Num() < TriRingCount);
-			if ( bIsSubset )
+
+			if (FilteredTriangles.Num() < TriRingCount)
 			{
+				// This is a non-mesh-boundary vert
 				checkSlow(!Mesh->SplitVertexWouldLeaveIsolated(VertID, FilteredTriangles));
 				DynamicMeshInfo::FVertexSplitInfo SplitInfo;
 				const EMeshResult MeshResult = Mesh->SplitVertex(VertID, FilteredTriangles, SplitInfo);
 				ensure(MeshResult == EMeshResult::Ok);
 
 				int NewVertID = SplitInfo.NewVertex;
-				LoopVertexMap.Add(Loop.Vertices[vi], NewVertID);
+				OldVidsToNewVids.Add(VertID, NewVertID);
 				NewVertexLoop[vi] = NewVertID;
 			}
 			else if (bHandleBoundaryVertices)
@@ -473,7 +564,7 @@ bool FDynamicMeshEditor::DisconnectTriangles(const TArray<int>& Triangles, TArra
 				// vertex as the "old" one, and just keep the existing one on the "inner" loop.
 				// This means we have to rewrite vertex in the "outer" loop, and that loop will no longer actually be an EdgeLoop, so we set those edges to invalid
 				int32 NewVertID = Mesh->AppendVertex(*Mesh, VertID);
-				LoopVertexMap.Add(NewVertID, VertID);
+				OldVidsToNewVids.Add(VertID, NewVertID);
 				LoopPair.OuterVertices[vi] = NewVertID;
 				LoopPair.OuterEdges[vi] = FDynamicMesh3::InvalidID;
 				LoopPair.OuterEdges[(vi == 0) ? NumVertices-1 : vi-1] = FDynamicMesh3::InvalidID;
@@ -608,6 +699,25 @@ void FDynamicMeshEditor::SplitBowties(int VertexID, FDynamicMeshEditResult& Resu
 	}
 }
 
+void FDynamicMeshEditor::SplitBowtiesAtTriangles(const TArray<int32>& TriangleIDs, FDynamicMeshEditResult& ResultOut)
+{
+	TSet<int32> VidsProcessed;
+	for (int32 Tid : TriangleIDs)
+	{
+		FIndex3i TriVids = Mesh->GetTriangle(Tid);
+		for (int i = 0; i < 3; ++i)
+		{
+			int32 Vid = TriVids[i];
+			bool bWasAlreadyProcessed = false;
+			VidsProcessed.Add(Vid, &bWasAlreadyProcessed);
+			if (!bWasAlreadyProcessed)
+			{
+				SplitBowties(Vid, ResultOut);
+			}
+		}
+	}
+}
+
 
 bool FDynamicMeshEditor::ReinsertSubmesh(const FDynamicSubmesh3& Region, FOptionallySparseIndexMap& SubToNewV, TArray<int>* new_tris, EDuplicateTriBehavior DuplicateBehavior)
 {
diff --git a/Engine/Plugins/Runtime/GeometryProcessing/Source/DynamicMesh/Private/Operations/InsetMeshRegion.cpp b/Engine/Plugins/Runtime/GeometryProcessing/Source/DynamicMesh/Private/Operations/InsetMeshRegion.cpp
--- a/Engine/Plugins/Runtime/GeometryProcessing/Source/DynamicMesh/Private/Operations/InsetMeshRegion.cpp
+++ b/Engine/Plugins/Runtime/GeometryProcessing/Source/DynamicMesh/Private/Operations/InsetMeshRegion.cpp
@@ -61,15 +61,6 @@ bool FInsetMeshRegion::Apply()
 
 bool FInsetMeshRegion::ApplyInset(FInsetInfo& Region, FMeshNormals* UseNormals)
 {
-	FMeshRegionBoundaryLoops InitialLoops(Mesh, Region.InitialTriangles, false);
-	bool bOK = InitialLoops.Compute();
-	if (bOK == false)
-	{
-		return false;
-	}
-
-	int32 NumInitialLoops = InitialLoops.GetLoopCount();
-
 	if (ChangeTracker)
 	{
 		ChangeTracker->SaveTriangles(Region.InitialTriangles, true);
@@ -78,14 +69,32 @@ bool FInsetMeshRegion::ApplyInset(FInsetInfo& Region, FMeshNormals* UseNormals)
 	FDynamicMeshEditor Editor(Mesh);
 
 	TArray<FDynamicMeshEditor::FLoopPairSet> LoopPairs;
-	bOK = Editor.DisconnectTriangles(Region.InitialTriangles, LoopPairs, true);
+	bool bOK = Editor.DisconnectTriangles(Region.InitialTriangles, LoopPairs, true);
 	if (bOK == false)
 	{
 		return false;
 	}
 
+	// Before we start insetting, prep things for stitching. We might end up breaking bowties in the disconnected
+	// region, which will mess up the Vids on the inner loop. So convert the inner loop to a sequence of identifiers
+	// that won't change across vertex splits (tri and subindex pairs). Note that we don't need to do the same to the
+	// outer loop because we are not splitting bowties there, and in fact the outer loop may end up with detached verts
+	// (without a triangle), so we would have needed a slightly messier system to account for that there.
+	typedef TPair<int32, TPair<int8, int8>> TriVertPair;
+	TArray<TArray<TriVertPair>> InsetStitchSides;
+	InsetStitchSides.SetNum(LoopPairs.Num());
+	for (int32 i = 0; i < LoopPairs.Num(); ++i)
+	{
+		FDynamicMeshEditor::ConvertLoopToTriVidPairSequence(*Mesh, LoopPairs[i].InnerVertices, LoopPairs[i].InnerEdges, InsetStitchSides[i]);
+	}
+
+	// Now that we've stored information that we need for stitching, split bowties in the disconnected region.
+	FDynamicMeshEditResult SplitBowtiesResult;
+	Editor.SplitBowtiesAtTriangles(Region.InitialTriangles, SplitBowtiesResult);
+
 	// make copy of separated submesh for deformation
 	// (could we defer this copy until we know we need it?)
+	// This should happen after splitting bowties.
 	FDynamicSubmesh3 SubmeshCalc(Mesh, Region.InitialTriangles, (int)EMeshComponents::None, false);
 	FDynamicMesh3& Submesh = SubmeshCalc.GetSubmesh();
 	bool bHaveInteriorVerts = false;
@@ -98,38 +107,70 @@ bool FInsetMeshRegion::ApplyInset(FInsetInfo& Region, FMeshNormals* UseNormals)
 		}
 	}
 
-	// inset vertices
-	for (FDynamicMeshEditor::FLoopPairSet& LoopPair : LoopPairs)
+	// If we did split bowties, some of the inner loops are now broken/merged (think of a small circle tangent inside
+	// a large circle- now you have a thick C shape instead). We'll need to generate new loops there for doing the 
+	// inset.
+	bool bSomeLoopsBroken = SplitBowtiesResult.NewVertices.Num() > 0;
+
+	Region.InsetLoops.Reset(LoopPairs.Num());
+	auto InsetLoop = [this, &Region](const TArray<int32>& LoopVids, const TArray<int32>& LoopEids)
 	{
-		int32 NumEdges = LoopPair.InnerEdges.Num();
+		int32 NumEdges = LoopEids.Num();
 
 		TArray<FLine3d> InsetLines;
-		UE::Geometry::ComputeInsetLineSegmentsFromEdges(*Mesh, LoopPair.InnerEdges, InsetDistance, InsetLines);
+		UE::Geometry::ComputeInsetLineSegmentsFromEdges(*Mesh, LoopEids, InsetDistance, InsetLines);
 
 		TArray<FVector3d> NewPositions;
-		UE::Geometry::SolveInsetVertexPositionsFromInsetLines(*Mesh, InsetLines, LoopPair.InnerVertices, NewPositions, true);
+		UE::Geometry::SolveInsetVertexPositionsFromInsetLines(*Mesh, InsetLines, LoopVids, NewPositions, true);
 
-		if (NewPositions.Num() == LoopPair.InnerVertices.Num())
+		if (NewPositions.Num() == LoopVids.Num())
 		{
-			for (int32 k = 0; k < LoopPair.InnerVertices.Num(); ++k)
+			for (int32 k = 0; k < LoopVids.Num(); ++k)
 			{
-				Mesh->SetVertex(LoopPair.InnerVertices[k], NewPositions[k]);
+				Mesh->SetVertex(LoopVids[k], NewPositions[k]);
 			}
 		}
+
+		Region.InsetLoops.Emplace();
+		Region.InsetLoops.Last().InitializeFromVertices(Mesh, LoopVids);
 	};
 
-	// stitch each loop
+	if (bSomeLoopsBroken)
+	{
+		// Recalc the loops
+		FMeshRegionBoundaryLoops UpdatedInsetLoops(Mesh, Region.InitialTriangles, false);
+		bOK = UpdatedInsetLoops.Compute();
+		if (!bOK)
+		{
+			return false;
+		}
+
+		for (const FEdgeLoop& Loop : UpdatedInsetLoops.Loops)
+		{
+			InsetLoop(Loop.Vertices, Loop.Edges);
+		}
+	}
+	else
+	{
+		// Can use the original loops we got while disconnecting
+		for (const FDynamicMeshEditor::FLoopPairSet& LoopPair : LoopPairs)
+		{
+			InsetLoop(LoopPair.InnerVertices, LoopPair.InnerEdges);
+		}
+	}
+
+	// Stitch the bands we saved.
+	int32 NumInitialLoops = LoopPairs.Num();
 	Region.BaseLoops.SetNum(NumInitialLoops);
-	Region.InsetLoops.SetNum(NumInitialLoops);
 	Region.StitchTriangles.SetNum(NumInitialLoops);
 	Region.StitchPolygonIDs.SetNum(NumInitialLoops);
-	TArray<TArray<FIndex2i>> QuadLoops;
-	QuadLoops.Reserve(NumInitialLoops);
-	int32 LoopIndex = 0;
-	for (FDynamicMeshEditor::FLoopPairSet& LoopPair : LoopPairs)
+	TArray<TArray<FIndex2i>> QuadStrips;
+	QuadStrips.Reserve(NumInitialLoops);
+	for (int32 LoopIndex = 0; LoopIndex < NumInitialLoops; ++LoopIndex)
 	{
-		TArray<int32>& BaseLoopV = LoopPair.OuterVertices;
-		TArray<int32>& InsetLoopV = LoopPair.InnerVertices;
+		const FDynamicMeshEditor::FLoopPairSet& LoopPair = LoopPairs[LoopIndex];
+		const TArray<int32>& BaseLoopV = LoopPair.OuterVertices;
+
 		int32 NumLoopV = BaseLoopV.Num();
 
 		// allocate a new group ID for each pair of input group IDs, and build up list of new group IDs along loop
@@ -138,8 +179,7 @@ bool FInsetMeshRegion::ApplyInset(FInsetInfo& Region, FMeshNormals* UseNormals)
 		TMap<TPair<int32, int32>, int32> NewGroupsMap;
 		for (int32 k = 0; k < NumLoopV; ++k)
 		{
-			int32 InsetEdgeID = Mesh->FindEdge(InsetLoopV[k], InsetLoopV[(k + 1) % NumLoopV]);
-			int32 InsetGroupID = Mesh->GetTriangleGroup(Mesh->GetEdgeT(InsetEdgeID).A);
+			int32 InsetGroupID = Mesh->GetTriangleGroup(InsetStitchSides[LoopIndex][k].Key);
 
 			// base edge may not exist if we inset entire region. In that case just use single GroupID
 			int32 BaseEdgeID = Mesh->FindEdge(BaseLoopV[k], BaseLoopV[(k + 1) % NumLoopV]);
@@ -155,9 +195,10 @@ bool FInsetMeshRegion::ApplyInset(FInsetInfo& Region, FMeshNormals* UseNormals)
 			EdgeGroups.Add(NewGroupsMap[GroupPair]);
 		}
 
-		// stitch the loops
+		// Stitch the loops. In case of broken loops, we'll still end up stitching the proper things, just
+		// in bands that fill in the proper regions together.
 		FDynamicMeshEditResult StitchResult;
-		Editor.StitchVertexLoopsMinimal(InsetLoopV, BaseLoopV, StitchResult);
+		Editor.StitchVertexLoopToTriVidPairSequence(InsetStitchSides[LoopIndex], LoopPair.OuterVertices, StitchResult);
 
 		// set the groups of the new quads along the stitch
 		int32 NumNewQuads = StitchResult.NewQuads.Num();
@@ -171,11 +212,9 @@ bool FInsetMeshRegion::ApplyInset(FInsetInfo& Region, FMeshNormals* UseNormals)
 		StitchResult.GetAllTriangles(Region.StitchTriangles[LoopIndex]);
 		Region.StitchPolygonIDs[LoopIndex] = NewGroupIDs;
 
-		QuadLoops.Add(MoveTemp(StitchResult.NewQuads));
+		QuadStrips.Add(MoveTemp(StitchResult.NewQuads));
 
 		Region.BaseLoops[LoopIndex].InitializeFromVertices(Mesh, BaseLoopV);
-		Region.InsetLoops[LoopIndex].InitializeFromVertices(Mesh, InsetLoopV);
-		LoopIndex++;
 	}
 
 	// if we have interior vertices or just want to try to resolve foldovers we
@@ -209,18 +248,18 @@ bool FInsetMeshRegion::ApplyInset(FInsetInfo& Region, FMeshNormals* UseNormals)
 
 		// compute area of inserted quad-strip border
 		double TotalBorderQuadArea = 0;
-		int32 NumLoops = LoopPairs.Num();
-		for (int32 li = 0; li < NumLoops; ++li)
+		int32 NumStrips = QuadStrips.Num();
+		for (int32 StripIndex = 0; StripIndex < NumStrips; ++StripIndex)
 		{
-			int32 NumQuads = QuadLoops[li].Num();
+			int32 NumQuads = QuadStrips[StripIndex].Num();
 			for (int32 k = 0; k < NumQuads; k++)
 			{
-				TotalBorderQuadArea += Mesh->GetTriArea(QuadLoops[li][k].A);
-				TotalBorderQuadArea += Mesh->GetTriArea(QuadLoops[li][k].B);
+				TotalBorderQuadArea += Mesh->GetTriArea(QuadStrips[StripIndex][k].A);
+				TotalBorderQuadArea += Mesh->GetTriArea(QuadStrips[StripIndex][k].B);
 			}
 		}
 
-		// Figure how much area chnaged by subtracting area of quad-strip from original area.
+		// Figure how much area changed by subtracting area of quad-strip from original area.
 		// (quad-strip area seems implausibly high at larger distances, ie becomes larger than initial area. Possibly due to sawtooth-shaped profile
 		//  of non-planar quads - measure each quad in planar projection?)
 		FVector2d VolArea = TMeshQueries<FDynamicMesh3>::GetVolumeArea(Submesh);
@@ -283,12 +322,11 @@ bool FInsetMeshRegion::ApplyInset(FInsetInfo& Region, FMeshNormals* UseNormals)
 	// calculate UVs/etc
 	if (Mesh->HasAttributes())
 	{
-		int32 NumLoops = LoopPairs.Num();
-		for ( int32 li = 0; li < NumLoops; ++li)
+		int32 NumStrips = QuadStrips.Num();
+		for ( int32 StripIndex = 0; StripIndex < NumStrips; ++StripIndex)
 		{
-			FDynamicMeshEditor::FLoopPairSet& LoopPair = LoopPairs[li];
-			TArray<int32>& BaseLoopV = LoopPair.OuterVertices;
-			TArray<int32>& InsetLoopV = LoopPair.InnerVertices;
+			FDynamicMeshEditor::FLoopPairSet& OriginalLoopPair = LoopPairs[StripIndex];
+			TArray<int32>& BaseLoopV = OriginalLoopPair.OuterVertices;
 
 			// for each polygon we created in stitch, set UVs and normals
 			// TODO copied from FExtrudeMesh, doesn't really make sense in this context...
@@ -296,10 +334,10 @@ bool FInsetMeshRegion::ApplyInset(FInsetInfo& Region, FMeshNormals* UseNormals)
 			FFrame3d FirstProjectFrame;
 			FVector3d FrameUp;
 
-			int32 NumQuads = QuadLoops[li].Num();
+			int32 NumQuads = QuadStrips[StripIndex].Num();
 			for (int32 k = 0; k < NumQuads; k++)
 			{
-				FVector3f Normal = Editor.ComputeAndSetQuadNormal( QuadLoops[li][k], true);
+				FVector3f Normal = Editor.ComputeAndSetQuadNormal( QuadStrips[StripIndex][k], true);
 
 				// align axis 0 of projection frame to first edge, then for further edges,
 				// rotate around 'up' axis to keep normal aligned and frame horizontal
@@ -326,7 +364,7 @@ bool FInsetMeshRegion::ApplyInset(FInsetInfo& Region, FMeshNormals* UseNormals)
 
 				// translate horizontally such that vertical spans are adjacent in UV space (so textures tile/wrap properly)
 				float TranslateU = UVScaleFactor * AccumUVTranslation;
-				Editor.SetQuadUVsFromProjection(QuadLoops[li][k], ProjectFrame, UVScaleFactor, FVector2f(TranslateU, 0));
+				Editor.SetQuadUVsFromProjection(QuadStrips[StripIndex][k], ProjectFrame, UVScaleFactor, FVector2f(TranslateU, 0));
 			}
 		}
 	}
diff --git a/Engine/Plugins/Runtime/GeometryProcessing/Source/DynamicMesh/Private/Operations/OffsetMeshRegion.cpp b/Engine/Plugins/Runtime/GeometryProcessing/Source/DynamicMesh/Private/Operations/OffsetMeshRegion.cpp
--- a/Engine/Plugins/Runtime/GeometryProcessing/Source/DynamicMesh/Private/Operations/OffsetMeshRegion.cpp
+++ b/Engine/Plugins/Runtime/GeometryProcessing/Source/DynamicMesh/Private/Operations/OffsetMeshRegion.cpp
@@ -1,6 +1,8 @@
 // Copyright Epic Games, Inc. All Rights Reserved.
 
 #include "Operations/OffsetMeshRegion.h"
+
+#include "Algo/Reverse.h"
 #include "DynamicMesh/MeshNormals.h"
 #include "DynamicMeshEditor.h"
 #include "Selections/MeshVertexSelection.h"
@@ -92,7 +94,7 @@ bool FOffsetMeshRegion::ApplyOffset(FOffsetInfo& Region, FMeshNormals* UseNormal
 
 	// Before we start changing triangles, prepare by allocating group IDs that we'll use
 	// for the stitched sides (doing it before changes to the mesh allows user-provided
-	// functions to operate on the original mesh).
+	// LoopEdgesShouldHaveSameGroup functions to operate on the original mesh).
 	TArray<TArray<int32>> LoopsEdgeGroups;
 	TArray<int32> NewGroupIDs;
 	LoopsEdgeGroups.SetNum(InitialLoops.Loops.Num());
@@ -143,45 +145,47 @@ bool FOffsetMeshRegion::ApplyOffset(FOffsetInfo& Region, FMeshNormals* UseNormal
 	FDynamicMeshEditResult DuplicateResult;
 	if (Region.bIsSolid)
 	{
-		// In the solid case, we want to duplicate the region.
+		// In the solid case, we want to duplicate the region so we can cap it.
 		FMeshIndexMappings IndexMap;
 		Editor.DuplicateTriangles(Region.OffsetTids, IndexMap, DuplicateResult);
 
 		AllModifiedAndNewTriangles.Append(DuplicateResult.NewTriangles);
 
 		// Populate LoopPairs
-		for (FEdgeLoop& BaseLoop : InitialLoops.Loops)
+		LoopPairs.SetNum(InitialLoops.Loops.Num());
+		for (int LoopIndex = 0; LoopIndex < InitialLoops.Loops.Num(); ++LoopIndex)
 		{
-			LoopPairs.Add(FDynamicMeshEditor::FLoopPairSet());
-			FDynamicMeshEditor::FLoopPairSet& LoopPair = LoopPairs.Last();
-
-			// Which loops we choose as the outer/inner will determine whether the
-			// sides are stitched inside out or not. The original Tids are the ones
-			// that are offset. In a positive offset, we want old as the "outer" and
-			// new as "inner". In negative offset, we want the reverse to allow our
-			// stitching code to still have the sides face outward.
-			TArray<int32>* OriginalVertsOut = &LoopPair.OuterVertices;
-			TArray<int32>* OriginalEdgesOut = &LoopPair.OuterEdges;
-			TArray<int32>* NewVertsOut = &LoopPair.InnerVertices;
-			TArray<int32>* NewEdgesOut = &LoopPair.InnerEdges;
-			
-			if (bIsPositiveOffset)
+			FEdgeLoop& BaseLoop = InitialLoops.Loops[LoopIndex];
+			FDynamicMeshEditor::FLoopPairSet& LoopPair = LoopPairs[LoopIndex];
+
+			// The original OffsetTids are the ones that are offset, so InnerVertices/Edges
+			// should be the boundaries of those.
+			LoopPair.InnerVertices = BaseLoop.Vertices;
+			LoopPair.InnerEdges = BaseLoop.Edges;
+
+			// However depending on whether we extruded down or up, we may need to reverse
+			// the loops to get them to be stitched right side out.
+			if (!bIsPositiveOffset)
 			{
-				Swap(OriginalVertsOut, NewVertsOut);
-				Swap(OriginalEdgesOut, NewEdgesOut);
+				Algo::Reverse(LoopPair.InnerVertices);
+
+				// Reversing the edges is slightly different because the last edge is between the first
+				// and last vertex, and that needs to stay in the same place when vertices are reversed.
+				int32 LastEid = LoopPair.InnerEdges.Pop();
+				Algo::Reverse(LoopPair.InnerEdges);
+				LoopPair.InnerEdges.Add(LastEid);
+
+				int32 LastEdgeGroupID = LoopsEdgeGroups[LoopIndex].Pop();
+				Algo::Reverse(LoopsEdgeGroups[LoopIndex]);
+				LoopsEdgeGroups[LoopIndex].Add(LastEdgeGroupID);
 			}
 
-			*OriginalVertsOut = BaseLoop.Vertices;
-			*OriginalEdgesOut = BaseLoop.Edges;
-			
-			for (int32 Vid : BaseLoop.Vertices)
+			// Now assemble the paired loop
+			for (int32 Vid : LoopPair.InnerVertices)
 			{
-				NewVertsOut->Add(IndexMap.GetNewVertex(Vid));
+				LoopPair.OuterVertices.Add(IndexMap.GetNewVertex(Vid));
 			}
-
-			FEdgeLoop OtherLoop;
-			bOK = ensure(OtherLoop.InitializeFromVertices(Mesh, *NewVertsOut, false));
-			*NewEdgesOut = OtherLoop.Edges;
+			FEdgeLoop::VertexLoopToEdgeLoop(Mesh, LoopPair.OuterVertices, LoopPair.OuterEdges);
 		}
 	}
 	else
@@ -194,6 +198,38 @@ bool FOffsetMeshRegion::ApplyOffset(FOffsetInfo& Region, FMeshNormals* UseNormal
 		return false;
 	}
 
+	// Store the vid-independent offset loop before we break bowties
+	typedef TPair<int32, TPair<int8, int8>> TriVertPair;
+	TArray<TArray<TriVertPair>> OffsetStitchSides;
+	OffsetStitchSides.SetNum(LoopPairs.Num());
+	for (int32 i = 0; i < LoopPairs.Num(); ++i)
+	{
+		bOK = bOK && FDynamicMeshEditor::ConvertLoopToTriVidPairSequence(*Mesh, LoopPairs[i].InnerVertices, LoopPairs[i].InnerEdges, OffsetStitchSides[i]);
+	}
+
+	if (bOK == false)
+	{
+		return false;
+	}
+	
+	// Split bowties in the chosen region
+	FDynamicMeshEditResult Result;
+	Editor.SplitBowtiesAtTriangles(Region.OffsetTids, Result);
+	bool bSomeLoopsBroken = Result.NewVertices.Num() > 0;
+
+	// If we broke bowties, the loops in the offset region have changed, and our OffsetLoops no longer
+	// match BaseLoops.
+	if (bSomeLoopsBroken)
+	{
+		FMeshRegionBoundaryLoops UpdatedOffsetLoops(Mesh, Region.OffsetTids, false);
+		bOK = UpdatedOffsetLoops.Compute();
+		if (!bOK)
+		{
+			return false;
+		}
+		Region.OffsetLoops = UpdatedOffsetLoops.Loops;
+	}
+
 	FMeshVertexSelection SelectionV(Mesh);
 	SelectionV.SelectTriangleVertices(Region.OffsetTids);
 	TArray<int32> SelectedVids = SelectionV.AsArray();
@@ -283,24 +319,28 @@ bool FOffsetMeshRegion::ApplyOffset(FOffsetInfo& Region, FMeshNormals* UseNormal
 	// Stitch the loops
 
 	bool bSuccess = true;
-	int NumInitialLoops = InitialLoops.GetLoopCount();
+	int NumInitialLoops = LoopPairs.Num();
 	Region.BaseLoops.SetNum(NumInitialLoops);
-	Region.OffsetLoops.SetNum(NumInitialLoops);
+	if (!bSomeLoopsBroken)
+	{
+		Region.OffsetLoops.SetNum(NumInitialLoops);
+	}
 	Region.StitchTriangles.SetNum(NumInitialLoops);
 	Region.StitchPolygonIDs.SetNum(NumInitialLoops);
-	int32 LoopIndex = 0;
-	for (int32 i = 0; i < LoopPairs.Num(); ++i)
+
+	for (int32 LoopIndex = 0; LoopIndex < LoopPairs.Num(); ++LoopIndex)
 	{
-		FDynamicMeshEditor::FLoopPairSet& LoopPair = LoopPairs[i];
-		const TArray<int32>& EdgeGroups = LoopsEdgeGroups[i];
+		FDynamicMeshEditor::FLoopPairSet& LoopPair = LoopPairs[LoopIndex];
+		const TArray<int32>& EdgeGroups = LoopsEdgeGroups[LoopIndex];
 
 		TArray<int32>& BaseLoopV = LoopPair.OuterVertices;
 		TArray<int32>& OffsetLoopV = LoopPair.InnerVertices;
-		int NumLoopV = BaseLoopV.Num();
+
+		TArray<TriVertPair>& OffsetLoopTriVertPairs = OffsetStitchSides[LoopIndex];
 
 		// stitch the loops
 		FDynamicMeshEditResult StitchResult;
-		bool bStitchSuccess = Editor.StitchVertexLoopsMinimal(OffsetLoopV, BaseLoopV, StitchResult);
+		bool bStitchSuccess = Editor.StitchVertexLoopToTriVidPairSequence(OffsetLoopTriVertPairs, LoopPair.OuterVertices, StitchResult);
 		if (!bStitchSuccess)
 		{
 			bSuccess = false;
@@ -362,8 +402,10 @@ bool FOffsetMeshRegion::ApplyOffset(FOffsetInfo& Region, FMeshNormals* UseNormal
 		}
 
 		Region.BaseLoops[LoopIndex].InitializeFromVertices(Mesh, BaseLoopV);
-		Region.OffsetLoops[LoopIndex].InitializeFromVertices(Mesh, OffsetLoopV);
-		LoopIndex++;
+		if (!bSomeLoopsBroken)
+		{
+			Region.OffsetLoops[LoopIndex].InitializeFromVertices(Mesh, OffsetLoopV);
+		}
 	}
 
 	if (Region.bIsSolid)
@@ -385,6 +427,11 @@ bool FOffsetMeshRegion::ApplyOffset(FOffsetInfo& Region, FMeshNormals* UseNormal
 bool FOffsetMeshRegion::EdgesSeparateSameGroupsAndAreColinearAtBorder(FDynamicMesh3* Mesh, 
 	int32 Eid1, int32 Eid2, bool bCheckColinearityAtBorder)
 {
+	if (!Mesh->IsEdge(Eid1) || !Mesh->IsEdge(Eid2))
+	{
+		return ensure(false);
+	}
+
 	FIndex2i Tris1 = Mesh->GetEdgeT(Eid1);
 	FIndex2i Groups1(Mesh->GetTriangleGroup(Tris1.A),
 		Tris1.B == IndexConstants::InvalidID ? IndexConstants::InvalidID : Mesh->GetTriangleGroup(Tris1.B));
diff --git a/Engine/Plugins/Runtime/GeometryProcessing/Source/DynamicMesh/Public/DynamicMeshEditor.h b/Engine/Plugins/Runtime/GeometryProcessing/Source/DynamicMesh/Public/DynamicMeshEditor.h
--- a/Engine/Plugins/Runtime/GeometryProcessing/Source/DynamicMesh/Public/DynamicMeshEditor.h
+++ b/Engine/Plugins/Runtime/GeometryProcessing/Source/DynamicMesh/Public/DynamicMeshEditor.h
@@ -99,6 +99,45 @@ public:
 	bool StitchVertexLoopsMinimal(const TArray<int>& VertexLoop1, const TArray<int>& VertexLoop2, FDynamicMeshEditResult& ResultOut);
 
 
+	/**
+	 * Given a sequence of edges in the form given by ConvertLoopToTriVidPairSequence and a loop of vertices, 
+	 * stitch the two together. Edges are identified as (TriangleID, (0-2 subindex of vert1, 0-2 subindex of vert2)), 
+	 * and should match the number of vertex ids in the corresponding vertex loop (each i corresponds to the edge 
+	 * between vert i and i+1 in the vertex loop).
+	 * This works much like StitchVertexLoopsMinimal, but it can be used after one of the loops undergoes
+	 * vertex splits (to remove bowties) as long as the "loop" (which may now be broken) can still be
+	 * identified by the composing triangles. In particular, this is used in extrude and inset where the
+	 * base loop needs to remain as a sequence of Vids (since some of these may be free vertices, unattached
+	 * to a triangle) but the offset/inset loop is guaranteed to be attached to triangles yet may have its
+	 * Vids change due to bowtie splits.
+	 * 
+	 * If an entry of TriVidPairs1 resolves to verts [a,b] and the corresponding pair of verts in VertexLoop2
+	 * are [c,d], then tris added are [b,a,d] and [a,c,d].
+	 * @param TriVidPairs1 First loop, given as output of ConvertLoopToTriVidPairSequence()
+	 * @param VertexLoop2 Second loop, given as vertex IDs
+	 * @param ResultOut lists of newly created triangles/vertices/etc
+	 * @return true if operation succeeded. If a failure occurs, any added triangles are removed via RemoveTriangles
+	 */
+	bool StitchVertexLoopToTriVidPairSequence(
+		const TArray<TPair<int32, TPair<int8, int8>>>& TriVidPairs1,
+		const TArray<int>& VertexLoop2, FDynamicMeshEditResult& ResultOut);
+
+	/**
+	 * Converts a loop to a sequence of edge identifiers that are both Vid and Eid independent. The identifiers are
+	 * of the form (TriangleID, ([0,2] vert sub index, [0,2] vert sub index)), and they are used in some cases 
+	 * (extrude, inset) where we want to perform vertex splits along a region boundary (to remove bowties) but 
+	 * need to maintain a record of the original loop path. We don't use (TriangleID, Eid sub index) because that
+	 * makes it harder to keep track of individual edge directionality.
+	 * 
+	 * @param VidLoop Vertex IDs of loop to convert.
+	 * @param EdgeLoop Edge IDs of loop to convert, must match VidLoop.
+	 * @param TriVertPairsOut Output
+	 * @return false if not successful (usually due to a mismatch between VidLoop and EdgeLoop).
+	 */
+	static bool ConvertLoopToTriVidPairSequence(const FDynamicMesh3& Mesh,
+		const TArray<int32>& VidLoop, const TArray<int32>& EdgeLoop,
+		TArray<TPair<int32, TPair<int8, int8>>>& TriVertPairsOut);
+
 	/**
 	 * Stitch together two loops of vertices where vertices are only sparsely corresponded
 	 * @param VertexIDs1 first array of sequential vertices
@@ -191,6 +230,11 @@ public:
 	 */
 	void SplitBowties(int VertexID, FDynamicMeshEditResult& ResultOut);
 
+	/**
+	 * Splits bowties attached to any of the given triangles, and updates (does not reset!) ResultOut with any added vertices
+	 */
+	void SplitBowtiesAtTriangles(const TArray<int32>& TriangleIDs, FDynamicMeshEditResult& ResultOut);
+
 	/**
 	 * In ReinsertSubmesh, a problem can arise where the mesh we are 
 	 * inserting has duplicate triangles of the base mesh.
diff --git a/Engine/Plugins/Runtime/GeometryProcessing/Source/DynamicMesh/Public/Operations/InsetMeshRegion.h b/Engine/Plugins/Runtime/GeometryProcessing/Source/DynamicMesh/Public/Operations/InsetMeshRegion.h
--- a/Engine/Plugins/Runtime/GeometryProcessing/Source/DynamicMesh/Public/Operations/InsetMeshRegion.h
+++ b/Engine/Plugins/Runtime/GeometryProcessing/Source/DynamicMesh/Public/Operations/InsetMeshRegion.h
@@ -81,7 +81,13 @@ public:
 		TArray<int32> InitialTriangles;
 		/** Initial loops on the mesh */
 		TArray<FEdgeLoop> BaseLoops;
-		/** Inset loops on the mesh */
+		/** 
+		 * Inset loops on the mesh. If the region had bowtie vertices (in terms of selection, not
+		 * even necessarily in terms of original mesh topology), then the number of loops here may
+		 * no longer match number of loops in BaseLoops, StitchTriangles, etc (for example, imagine
+		 * a circular hole tangent inside a circular region; when the tangent bowtie is split, this
+		 * becomes one C-shaped loop instead of two nested ones)
+		 */
 		TArray<FEdgeLoop> InsetLoops;
 
 		/** Lists of triangle-strip "tubes" that connect each loop-pair */
diff --git a/Engine/Plugins/Runtime/GeometryProcessing/Source/DynamicMesh/Public/Operations/OffsetMeshRegion.h b/Engine/Plugins/Runtime/GeometryProcessing/Source/DynamicMesh/Public/Operations/OffsetMeshRegion.h
--- a/Engine/Plugins/Runtime/GeometryProcessing/Source/DynamicMesh/Public/Operations/OffsetMeshRegion.h
+++ b/Engine/Plugins/Runtime/GeometryProcessing/Source/DynamicMesh/Public/Operations/OffsetMeshRegion.h
@@ -115,7 +115,13 @@ public:
 
 		/** Initial loops on the mesh */
 		TArray<FEdgeLoop> BaseLoops;
-		/** Offset loops on the mesh */
+		/** 
+		 * Offset loops on the mesh. If the region had bowtie vertices (in terms of selection, not
+		 * even necessarily in terms of original mesh topology), then the number of loops here may
+		 * no longer match number of loops in BaseLoops, StitchTriangles, etc (for example, imagine
+		 * a circular hole tangent inside a circular region; when the tangent bowtie is split, this
+		 * becomes one C-shaped loop instead of two nested ones)
+		 */
 		TArray<FEdgeLoop> OffsetLoops;
 
 		/** Lists of triangle-strip "tubes" that connect each loop-pair */
diff --git a/Engine/Plugins/Runtime/MeshModelingToolset/Source/ModelingOperators/Private/DeformationOps/ExtrudeOp.cpp b/Engine/Plugins/Runtime/MeshModelingToolset/Source/ModelingOperators/Private/DeformationOps/ExtrudeOp.cpp
--- a/Engine/Plugins/Runtime/MeshModelingToolset/Source/ModelingOperators/Private/DeformationOps/ExtrudeOp.cpp
+++ b/Engine/Plugins/Runtime/MeshModelingToolset/Source/ModelingOperators/Private/DeformationOps/ExtrudeOp.cpp
@@ -99,8 +99,21 @@ bool FExtrudeOp::BooleanExtrude(FProgressCancel* Progress)
 
 	// We need a custom function here so that we can set the side groups based on the original mesh, not the submesh
 	Extruder.LoopEdgesShouldHaveSameGroup = [&Submesh, this](int32 Eid1, int32 Eid2) {
-		int32 EidToUse1 = Submesh->MapEdgeToBaseMesh(Eid1);
-		int32 EidToUse2 = Submesh->MapEdgeToBaseMesh(Eid2);
+
+		// Unfortunately, we can't just use Submesh->MapEdgeToBaseMesh() for the eids here because
+		// that does the mapping via the vertices, and as the extruder extrudes separate full
+		// regions, it ends up splitting bowties, which breaks that correspondence. 
+		// However, the triangles of the next regions to be extruded don't change, so we can get
+		// the corresponding edges via tid and edge sub index (in the triangle edge triplet).
+		const FDynamicMesh3& ExtrusionMesh = Submesh->GetSubmesh();
+		int32 Tid1 = ExtrusionMesh.GetEdgeT(Eid1).A;
+		int32 SubIdx1 = IndexUtil::FindTriIndex(Eid1, ExtrusionMesh.GetTriEdges(Tid1));
+		int32 EidToUse1 = ResultMesh->GetTriEdge(Submesh->MapTriangleToBaseMesh(Tid1), SubIdx1);
+
+		int32 Tid2 = ExtrusionMesh.GetEdgeT(Eid2).A;
+		int32 SubIdx2 = IndexUtil::FindTriIndex(Eid2, ExtrusionMesh.GetTriEdges(Tid2));
+		int32 EidToUse2 = ResultMesh->GetTriEdge(Submesh->MapTriangleToBaseMesh(Tid2), SubIdx2);
+
 		return FOffsetMeshRegion::EdgesSeparateSameGroupsAndAreColinearAtBorder(
 			ResultMesh.Get(), EidToUse1, EidToUse2, bUseColinearityForSettingBorderGroups);
 	};
