Geometry Processing:  Implement intrinsic mesh triangle and edge split using normal coordinates

diff --git a/Engine/Plugins/Runtime/GeometryProcessing/Source/DynamicMesh/Private/Operations/IntrinsicCorrespondenceUtils.cpp b/Engine/Plugins/Runtime/GeometryProcessing/Source/DynamicMesh/Private/Operations/IntrinsicCorrespondenceUtils.cpp
--- a/Engine/Plugins/Runtime/GeometryProcessing/Source/DynamicMesh/Private/Operations/IntrinsicCorrespondenceUtils.cpp
+++ b/Engine/Plugins/Runtime/GeometryProcessing/Source/DynamicMesh/Private/Operations/IntrinsicCorrespondenceUtils.cpp
@@ -33,9 +33,7 @@ IntrinsicCorrespondenceUtils::FSurfacePoint::FSurfacePoint(int32 TriID, const FV
 {
 	PositionType = EPositionType::Triangle;
 	Position.TriPosition.TriID = TriID;
-	Position.TriPosition.BarycentricCoords[0] = BaryCentrics[0];
-	Position.TriPosition.BarycentricCoords[1] = BaryCentrics[1];
-	Position.TriPosition.BarycentricCoords[2] = BaryCentrics[2];
+	Position.TriPosition.BarycentricCoords = BaryCentrics;
 }
 
 FVector3d IntrinsicCorrespondenceUtils::AsR3Position(const IntrinsicCorrespondenceUtils::FSurfacePoint& SurfacePoint, const FDynamicMesh3& Mesh, bool& bIsValid)
@@ -76,7 +74,7 @@ FVector3d IntrinsicCorrespondenceUtils::AsR3Position(const IntrinsicCorresponden
 				Result = Mesh.GetTriBaryPoint( SurfaceTID,
 											   Position.TriPosition.BarycentricCoords[0],
 											   Position.TriPosition.BarycentricCoords[1],
-											   Position.TriPosition.BarycentricCoords[2] );
+											   Position.TriPosition.BarycentricCoords[2]);
 			}
 
 		}
@@ -162,9 +160,11 @@ IntrinsicCorrespondenceUtils::FNormalCoordinates::FNormalCoordinates(const FDyna
 	const int32 MaxVID = SurfMesh.MaxVertexID();
 	const int32 MaxTID = SurfMesh.MaxTriangleID();
 
-	// the intrinsic mesh on construction will have the same triangles as the surface mesh,
-	// so no surface mesh edges will cross intrinsic edges.
-	NormalCoord.SetNumZeroed(MaxEID);
+	// on construction the intrinsic mesh will have the same triangles as the surface mesh,
+	// with the intrinsic edges identical to the surface edges. 
+	// the value -1 indicates the intrinsic mesh edge is a surface mesh edge segment.
+	NormalCoord.Resize(MaxEID, -1);
+	
 
 	// populate Ref Vert Degree. 
 	RefVertDegree.SetNumZeroed(MaxVID);
@@ -179,11 +179,8 @@ IntrinsicCorrespondenceUtils::FNormalCoordinates::FNormalCoordinates(const FDyna
 	}
 
 	// init the roundabout data as invalid
-	RoundaboutOrder.SetNumUninitialized(MaxTID);
-	for (FIndex3i& PerTri : RoundaboutOrder)
-	{
-		PerTri = FIndex3i(FDynamicMesh3::InvalidID, FDynamicMesh3::InvalidID, FDynamicMesh3::InvalidID);
-	}
+	RoundaboutOrder.Resize(MaxTID, FIndex3i::Invalid());
+	
 
 	// correctly populate the roundabout data.  Note bow-tie vertices should prohibited at a higher level
 	for (int32 CVID : SurfMesh.VertexIndicesItr())
@@ -201,7 +198,11 @@ IntrinsicCorrespondenceUtils::FNormalCoordinates::FNormalCoordinates(const FDyna
 	}
 
 	// Edge Order is a copy of the roundabout before edge flips.  Edge order describes the surface mesh, and shouldn't be updated after construction
-	EdgeOrder = RoundaboutOrder;
+	EdgeOrder.SetNumUninitialized(MaxTID);
+	for (int32 i = 0; i < MaxTID; ++i)
+	{
+		EdgeOrder[i] = RoundaboutOrder[i];
+	}
 }
 
 bool  IntrinsicCorrespondenceUtils::FNormalCoordinates::OnFlipEdge(const int32 T0ID, const FIndex3i& BeforeT0EIDs, const int32 T0OppVID,
@@ -258,11 +259,11 @@ bool  IntrinsicCorrespondenceUtils::FNormalCoordinates::OnFlipEdge(const int32 T
 		const int32 Eki_j = NumCornerEmanatingRefEdges(BeforeT0EIDs, T0_j);
 
 		// if the edge being flipped aligns with a surface edge, then it will cross that surface edge after the  flip
-		const int32 Kronecker = (NormalCoord[FlippedEID] == 0) ? 1 : 0;
+		const int32 Kronecker = IsSurfaceEdgeSegment(FlippedEID) ? 1 : 0;
 
 		// total number of surfaces edges this edge crosses after the flip
-		const int32 N_kl = TMathUtil<int32>::Max(0, Cji_l + Cij_k + (TMathUtil<int32>::Abs(Cil_j - Cki_j) + TMathUtil<int32>::Abs(Clj_i - Cjk_i) - Eji_l - Eij_k) / 2
-			+ Elj_i + Ejk_i + Eil_j + Eki_j + Kronecker);
+		const int32 N_kl = Cji_l + Cij_k + (TMathUtil<int32>::Abs(Cil_j - Cki_j) + TMathUtil<int32>::Abs(Clj_i - Cjk_i) - Eji_l - Eij_k) / 2
+						 + Elj_i + Ejk_i + Eil_j + Eki_j + Kronecker;
 
 		NormalCoord[FlippedEID] = N_kl;
 	}
@@ -294,18 +295,29 @@ bool  IntrinsicCorrespondenceUtils::FNormalCoordinates::OnFlipEdge(const int32 T
 		const int32 R_lj = T1Rab[T1_l];
 
 		// --compute the roundabouts for the two half-edges associated with the edge that does flip. (i.e  R_kl, R_lk)
-
-		const int32 Kronecker_ki = (NormalCoord[Edge_ki] == 0) ? 1 : 0;
-		const int32 Kronecker_lj = (NormalCoord[Edge_lj] == 0) ? 1 : 0;
-
-		const int32 Eil_k = NumCornerEmanatingRefEdges(AfterT1EIDs, 1/*k corner*/); // AfterT1 = (l, k, i)
-		const int32 Ejk_l = NumCornerEmanatingRefEdges(AfterT0EIDs, 1/*l corner*/); // AfterT0 = (k, l, j)
-
-		// Roundabout values for the two directions of the flipped edge ( i.e. as seen by the new T0 and T1)
-		const int32 R_kl = (R_ki + Eil_k + Kronecker_ki) % RefVertDegree[VID_k];
-		const int32 R_lk = (R_lj + Ejk_l + Kronecker_lj) % RefVertDegree[VID_l];
-
-
+		// note roundabouts are only computed for the original vertex set (i.e. those also in the surface mesh)
+		const int32 R_kl = [&] 
+							{
+								if (!SurfaceMesh->IsVertex(VID_k))
+								{
+									return -1;
+								}
+								const int32 Kronecker_ki = IsSurfaceEdgeSegment(Edge_ki) ? 1 : 0;
+								const int32 Eil_k = NumCornerEmanatingRefEdges(AfterT1EIDs, 1/*k corner*/); // AfterT1 = (l, k, i)
+								return  (R_ki + Eil_k + Kronecker_ki) % RefVertDegree[VID_k];
+							}();
+
+		const int32 R_lk = [&]
+							{
+								if (!SurfaceMesh->IsVertex(VID_l))
+								{
+									return -1;
+								}
+								const int32 Kronecker_lj = IsSurfaceEdgeSegment(Edge_lj) ? 1 : 0;
+								const int32 Ejk_l = NumCornerEmanatingRefEdges(AfterT0EIDs, 1/*l corner*/); // AfterT0 = (k, l, j)
+								return (R_lj + Ejk_l + Kronecker_lj) % RefVertDegree[VID_l];
+							}();
+		
 		// update the roundabouts with new edge order and values for the flipped edge
 		RoundaboutOrder[T0ID] = FIndex3i(R_kl, R_lj, R_jk);
 		RoundaboutOrder[T1ID] = FIndex3i(R_lk, R_ki, R_il);
diff --git a/Engine/Plugins/Runtime/GeometryProcessing/Source/DynamicMesh/Private/Operations/IntrinsicTriangulationMesh.cpp b/Engine/Plugins/Runtime/GeometryProcessing/Source/DynamicMesh/Private/Operations/IntrinsicTriangulationMesh.cpp
--- a/Engine/Plugins/Runtime/GeometryProcessing/Source/DynamicMesh/Private/Operations/IntrinsicTriangulationMesh.cpp
+++ b/Engine/Plugins/Runtime/GeometryProcessing/Source/DynamicMesh/Private/Operations/IntrinsicTriangulationMesh.cpp
@@ -9,6 +9,8 @@
 #include "Containers/Queue.h"
 #include "MathUtil.h"
 #include "Async/ParallelFor.h"
+#include "VectorTypes.h"
+#include "MatrixTypes.h"
 
 using namespace UE::Geometry;
 
@@ -126,7 +128,7 @@
 	{
 		const double Area = TriangleArea(L0, L1, L2);
 		const double Y = 2. * Area / L0;
-		double X = FMath::Sqrt(FMath::Max(0., L2 * L2 - Y * Y));
+		const double X = FMath::Sqrt(FMath::Max(0., (L2 - Y) *  (L2 + Y)));
 		FVector2d ResultPos(X, Y);
 		// angle between L0 and L2 edges can be computed form 2 * L0 * L2 * Cos(theta) = L0^2 + L2^2 - L1^2
 		if (L0 * L0 + L2 * L2 < L1 * L1)
@@ -418,9 +420,8 @@
 			// trace from BC point in triangle
 			const int32 SurfaceTID = SurfacePosition.TriPosition.TriID;
 
-			const FVector3d BaryPoint( SurfacePosition.TriPosition.BarycentricCoords[0],
-									   SurfacePosition.TriPosition.BarycentricCoords[1],
-				                       SurfacePosition.TriPosition.BarycentricCoords[2] );
+			const FVector3d BaryPoint = SurfacePosition.TriPosition.BarycentricCoords;
+									
 			const int32 RefSurfaceEID = SignpostData.TIDToReferenceEID[SurfaceTID];
 
 			FMeshGeodesicSurfaceTracer::FMeshSurfaceDirection DirectionOnSurface(RefSurfaceEID, TracePolarAngle);
@@ -442,23 +443,391 @@
 namespace FNormalCoordSurfaceTraceImpl
 {
 	using namespace IntrinsicCorrespondenceUtils;
+	
+
+	// @return ID of a vertex-adjacent edge that will allow all adjacent edges to be visited in CCW order.
+	template <typename MeshType>
+	int32 IdentifyInitialAdjacentEdge(const MeshType& Mesh, const int32 VID)
+	{
+		if (!Mesh.IsVertex(VID))
+		{
+			return MeshType::InvalidID;
+		}
+
+		if (Mesh.IsBoundaryVertex(VID))
+		{
+			// find the clockwise-most edge
+			for (int32 NbrEID : Mesh.VtxEdgesItr(VID))
+			{
+				if (Mesh.IsBoundaryEdge(NbrEID))
+				{
+					const int32 NbrTID  = Mesh.GetEdgeT(NbrEID).A;
+					const int32 IndexOf = Mesh.GetTriEdges(NbrTID).IndexOf(NbrEID);
+					if (Mesh.GetTriangle(NbrTID)[IndexOf] == VID)
+					{
+						return NbrEID;
+					}
+				}
+			}
+
+			return MeshType::InvalidID; // shouldn't reach 
+		}
+		else
+		{
+			for (int32 NbrEID : Mesh.VtxEdgesItr(VID))
+			{
+				// get the first
+				return NbrEID;
+			}
+			return MeshType::InvalidID; // shouldn't reach here
+		}
+	}
+
+	
 	/**
-	* low-level code that uses normal coordinates to continues tracing a surface edge across an FIntrinsicEdgeFlipMesh until it terminates at a vertex.
-	*
-	* In order to compute the full crossing history for a surface curve, this function could be called directly after 
-	* the first intrinsic edge crossing at which time the 'Crossings' result array will already hold the start vertex and the first edge crossing.
+	* Low-level method for use with normal-coordinate equipped intrinsic meshes.
 	* 
+	* Given a surface edge that intersects (but does not terminate at) an intrinsic vertex 'IntrinsicVID' 
+	* this function identifies the adjacent intrinsic mesh elements (i.e. edges or faces ) also containing that surface edge.
 	* 
-	* This function will populate the rest of the sequence of edge crossings for the surface mesh edge
-	* as a list of the edges crossed, and not the location on the individual edges.  
-	* To compute the those locations the resulting triangle strip should be unfolded and each location can be solved for
-	* using TraceEdgeOverHost below.
-	*/
-	void  ContinueTraceSurfaceEdge( const FIntrinsicEdgeFlipMesh& IntrinsicMesh, const int32 StartTID, const int32 StartEID, const int32 StartP,
-		                            TArray<FNormalCoordinates::FEdgeAndCrossingIdx>& Crossings )
+	* This is used when following a surface mesh edge across an intrinsic mesh (using normal coordinates). 
+	* Note: this situation is the result of an edge split of an intrinsic edge that was initially aligned with a surface edge.
+	*/ 
+	template <typename IntrinsicMeshType>
+	bool GetAdjElementsContainingSurfaceEdge(const IntrinsicMeshType& IntrinsicMesh, const int32 IntrinsicVID, TArray<TTuple<int32, int32>>& IntrisicTIDs, TArray<int32>& IntrisicEIDs)
 	{
+		checkSlow( IsEdgePoint(IntrinsicMesh.GetVertexSurfacePoint(IntrinsicVID)) );
+
 		const FNormalCoordinates& NCoords = IntrinsicMesh.GetNormalCoordinates();
 
+		for (const int32 EID : IntrinsicMesh.VtxEdgesItr(IntrinsicVID))
+		{
+			if (NCoords.IsSurfaceEdgeSegment(EID))
+			{
+				IntrisicEIDs.Add(EID);
+			}
+
+			// triangle that has EID as an outgoing edge
+			const int32 TID = [&]
+								{
+									const FIndex2i AdjTris = IntrinsicMesh.GetEdgeT(EID);
+									for (int32 i = 0; i < 2; ++i)
+									{
+										const int32 TID = AdjTris[i];
+										if (TID == IntrinsicMeshType::InvalidID)
+										{
+											continue;
+										}
+
+										const FIndex3i TriEIDs    = IntrinsicMesh.GetTriEdges(TID);
+										const int32 IndexOf       = TriEIDs.IndexOf(EID);
+										const bool bIsEIDOutgoing = (IntrinsicMesh.GetTriangle(TID)[IndexOf] == IntrinsicVID);
+										if (bIsEIDOutgoing)
+										{
+											return TID;
+										}
+									}
+									return IntrinsicMeshType::InvalidID;
+								}();
+
+			if (TID == IntrinsicMeshType::InvalidID)
+			{
+				continue;
+			}
+
+			const FIndex3i TriEIDs = IntrinsicMesh.GetTriEdges(TID);
+			const int32 IndexOf = TriEIDs.IndexOf(EID);
+
+			// permuted triangle is {a, b, c} where a is VID.
+			// given a surface edge crosses 'a' it will traverse the face of this triangle iff it exits face bc.
+
+			const FIndex3i PermutedEIDs = IntrinsicMesh.Permute(IndexOf, TriEIDs);
+			const int32 Nbc = NCoords.NumEdgeCrossing(PermutedEIDs[1]);           // number of surface edges that exits face bc
+
+			if (Nbc == 0) // a surface edge that crosses 'a' doesn't exit this triangle face bc
+			{
+				continue;
+			}
+
+			const int32 Eca_b = NCoords.NumCornerEmanatingRefEdges(PermutedEIDs, 1); // number of surface edges that exit corner b and exit ca
+			if (Eca_b > 0) // a surface edge that crosses 'a' is blocked from exiting this triangle by other surface edges
+			{
+				continue;
+			}
+			const int32 Eab_c = NCoords.NumCornerEmanatingRefEdges(PermutedEIDs, 2); // number of surface edges that exit corner c and exit ab
+			if (Eab_c > 0) // a surface edge that crosses 'a' is blocked from exiting this triangle by other surface edges
+			{
+				continue;
+			}
+
+			const int32 Cb = NCoords.NumCornerCrossingRefEdges(PermutedEIDs, 1); // number of surface edges that cross corner b
+			const int32 Cc = NCoords.NumCornerCrossingRefEdges(PermutedEIDs, 2); // number of surface edges that cross corner c
+
+			if (Nbc > Cb + Cc) // a surface edge that crosses 'a' must exit face bc
+			{
+				IntrisicTIDs.Add(TTuple<int32, int32>(TID, IndexOf));
+			}
+
+		}
+
+		return (IntrisicTIDs.Num() > 0 || IntrisicEIDs.Num() > 0);
+	}
+
+	
+	/**
+	* Low-level method for use with normal-coordinate equipped intrinsic meshes.
+	* 
+	* Given an intrinsic mesh vertex that is situated on a surface mesh edge and information that specifies the incoming section of the surface edge,
+	* this function computes the crossing point where the outgoing section of the surface edge exits the intrinsic vertex 1-ring.
+	* 
+	* if bIncomingIsEdge == true, then the incoming section of the surface edge is along the intrinsic edge EdgeID = IncomingID, 
+	* otherwise it crossed the intrinsic triangle TriID = IncomingID. 
+	* 
+	* on return the IncomingID and bIncomingIsEdge have been updated with the surface mesh elements traversed by the path to next crossing,
+	* and the crossing itself is encoded in the returned FEdgeAndCrossing struct.  for convenience VID is also updated if the next crossing is an intrinsic vertx.
+	*/
+	template <typename IntrinsicMeshType>   
+	FNormalCoordinates::FEdgeAndCrossingIdx NextCrossingFromEdgePoint(const IntrinsicMeshType& IntrinsicMesh, int32& VID, int32& IncomingID, bool& bIncomingIsEdge)
+	{
+
+		checkSlow(IsEdgePoint(IntrinsicMesh.GetVertexSurfacePoint(VID)));
+	
+
+		// four cases: the outer product of   "enters {on intrinsic edge, across intrinsic tri}"   with "exits {on intrinsic edge, across intrinsic tri}"
+
+		// find intrinsic mesh elements adjacent to VID that support the surface edge
+		TArray<TTuple<int32, int32>> TIDs;   TArray<int32> EIDs;
+		GetAdjElementsContainingSurfaceEdge(IntrinsicMesh, VID, TIDs, EIDs);
+
+		// only one edge crosses the implicit vert, so there can only be 2 sides to that edge ..
+		checkSlow(TIDs.Num() + EIDs.Num() == 2);
+
+		
+
+		auto EncodeExitsOnEdge = [&](const int32 ExitEID, FNormalCoordinates::FEdgeAndCrossingIdx& Xing )
+							{
+								const typename IntrinsicMeshType::FEdge ExitEdge = IntrinsicMesh.GetEdge(ExitEID);
+
+								const int32 AdjTID       = ExitEdge.Tri[0];
+								const int32 IndexOfe     = IntrinsicMesh.GetTriEdges(AdjTID).IndexOf(ExitEID);
+								const int32 IndexOfNextV = (IntrinsicMesh.GetTriangle(AdjTID)[IndexOfe] == VID) ? (IndexOfe + 1) % 3 : IndexOfe;
+
+								Xing.TID = ExitEdge.Tri[0];
+								Xing.EID = IndexOfNextV;
+								Xing.CIdx = 0;
+								return Xing;
+							};
+
+		auto EncodeExitsAcrossTri = [&](const TTuple<int32, int32>& TridAndIdx, FNormalCoordinates::FEdgeAndCrossingIdx& Xing)
+								{
+									const FNormalCoordinates& NCoords = IntrinsicMesh.GetNormalCoordinates();
+									const int32 ExitTID               = TridAndIdx.Get<0>();
+									const int32 OutgoingIndex         = TridAndIdx.Get<1>();
+
+									// when counting the surface edge that cross the triangle face opposite to VID, what number is the single edge that eminates from VID?
+									const FIndex3i TriEIDs = IntrinsicMesh.GetTriEdges(ExitTID);
+
+									// permuted triangle is {a, b, c} where a is VID.
+									// given a surface edge crosses 'a' it will traverse the face of this triangle iff it exits face bc.
+
+									const FIndex3i PermutedEIDs = IntrinsicMesh.Permute(OutgoingIndex, TriEIDs);
+									const int32 Cc = NCoords.NumCornerCrossingRefEdges(PermutedEIDs, 2); // number of surface edges that cross corner c
+
+									// counting from corner c towards b, this is the 'CC +1'-th surface edge to cross the intrinsic edge 
+									const int32 CIdx = Cc + 1;
+
+									Xing.TID = ExitTID;
+									Xing.EID = PermutedEIDs[1];
+									Xing.CIdx = CIdx;
+									return Xing;
+								};
+
+		FNormalCoordinates::FEdgeAndCrossingIdx Xing({-1, -1, -1});
+
+		if (bIncomingIsEdge) // entered on an intrinsic edge
+		{
+			const int32 EnterEID = IncomingID;
+
+			const bool bExitsOnEdge    = (EIDs.Num() == 2);
+			const bool bExitsAccrosTri = (TIDs.Num() > 0);
+
+			if (bExitsOnEdge)
+			{
+				checkSlow(!bExitsAccrosTri);
+
+				const int32 Id = (EIDs[0] == EnterEID) ? 1 : 0;
+				checkSlow(EIDs[1 - Id] == EnterEID);
+
+				const int32 ExitEID = EIDs[Id];
+				EncodeExitsOnEdge(ExitEID, Xing);
+				
+				VID             = IntrinsicMesh.GetTriangle(Xing.TID)[Xing.EID];
+				IncomingID      = ExitEID;
+				bIncomingIsEdge = true;
+			}
+			else if (bExitsAccrosTri)
+			{
+				checkSlow(!bExitsOnEdge);
+
+				checkSlow(TIDs.Num() == 1); // can only traverse a single tri ( we already know it entered by an edge)
+				EncodeExitsAcrossTri(TIDs[0], Xing);
+
+				IncomingID      = TIDs[0].Get<0>();
+				bIncomingIsEdge = false;
+			}
+			else
+			{
+				checkSlow(0);
+			}
+		}
+		else // entered across and intrinsic triangle
+		{
+			const int32 EnterTID       = IncomingID;
+			const bool bExitsOnEdge    = (EIDs.Num() == 1);
+			const bool bExitsAcrossTri = (TIDs.Num() == 2);
+
+			if (bExitsOnEdge)
+			{
+				const int32 ExitEID = EIDs[0];
+				EncodeExitsOnEdge(ExitEID, Xing);
+
+				VID             = IntrinsicMesh.GetTriangle(Xing.TID)[Xing.EID];
+				IncomingID      = ExitEID;
+				bIncomingIsEdge = true;
+			}
+			else if (bExitsAcrossTri)
+			{
+				const int32 Id = (TIDs[0].Get<0>() == EnterTID) ? 1 : 0;
+				EncodeExitsAcrossTri(TIDs[Id], Xing);
+
+				IncomingID      = TIDs[0].Get<0>();
+				bIncomingIsEdge = false;
+			}
+			else
+			{
+				checkSlow(0);
+			}
+		}
+
+		return Xing;
+	}
+
+	template <typename IntrinsicMeshType>
+	void ContinueFromEdgePointCrossing(const IntrinsicMeshType& IntrinsicMesh, int32 VID, int32 IncomingID, bool bIncomingIsEdge, TArray<FNormalCoordinates::FEdgeAndCrossingIdx>& Crossings)
+	{
+		int32 CIdx = 0;
+		while (CIdx == 0 && IsEdgePoint(IntrinsicMesh.GetVertexSurfacePoint(VID)))
+		{
+			FNormalCoordinates::FEdgeAndCrossingIdx Xing = NextCrossingFromEdgePoint(IntrinsicMesh, VID, IncomingID, bIncomingIsEdge);
+			Crossings.Add(Xing);
+			CIdx = Xing.CIdx;
+		}
+	}
+
+	/**
+	* Low-level method for use with normal-coordinate equipped intrinsic meshes.
+	* Traces the p-th surface edge crossing of the intrinsic edge edgeID across the adjacent intrinsic triangle TriID
+	* here 'P' is counted is the CCW direction (relative to TriID).
+	* 
+	* on returns TTuple<in32, int32>  = {EdgeID,P}  
+	* on return if 'P' == 0, the path exits a vertex otherwise it exits the updated EdgeID (with index P relative to the next triangle it enters)
+	* 
+	* This is adapted from Algorithm 2 of  Gillespi et al, 2020
+	*/ 
+	template <typename IntrinsicMeshType>
+	TTuple<int32, int32> GetCrossingExit(const IntrinsicMeshType& IntrinsicMesh, const FNormalCoordinates& NCoords, const int32 TriID, const int32 EdgeIDin, const int32 Pin)
+	{
+
+		TTuple<int32, int32> ExitCoord(EdgeIDin, Pin);
+		int32& EdgeID = ExitCoord.Get<0>();
+		int32& P = ExitCoord.Get<1>();
+
+		const FIndex3i TriEIDs = IntrinsicMesh.GetTriEdges(TriID);
+		const int32 IndexOf = TriEIDs.IndexOf(EdgeID);
+
+		const int32 Edge_ji = EdgeID;
+		const int32 Edge_il = TriEIDs[(IndexOf + 1) % 3];
+		const int32 Edge_lj = TriEIDs[(IndexOf + 2) % 3];
+
+		const int32 N_ji = NCoords.NumEdgeCrossing(Edge_ji);
+		const int32 N_il = NCoords.NumEdgeCrossing(Edge_il);
+		const int32 N_lj = NCoords.NumEdgeCrossing(Edge_lj);
+
+		
+
+		if (N_ji > N_lj + N_il)  // case 1
+		{
+			// some of the N_ji edges that cross edge IJ must connect with the vertex at L 
+			if (P <= N_il)
+			{
+				// exits side il
+				EdgeID = Edge_il;
+				P      = Pin;
+			}
+			else if ((N_il < P) && (P <= N_ji - N_lj))
+			{
+				// this path terminates at vertex
+				// this is a slight abuse.  The convention, when CrossingIndex = 0, we specify the vertex by the edge that originates there.
+				EdgeID = TriEIDs.IndexOf(Edge_lj); // GetTriange()[EdgeID] = vert.  note changed from Edge_lj where  i = GetTriangle()[TriEIDs.IndexOf(EdgeID_lj)] will be the vertex.
+				P = 0;
+			}
+			else
+			{
+				EdgeID = Edge_lj;
+				P      = Pin - (N_ji - N_lj);
+			}
+
+		}
+		else if (N_lj > N_ji + N_il)  // case 2
+		{
+			EdgeID = Edge_lj;
+			P      = Pin + (N_lj - N_ji);
+
+		}
+		else if (N_il > N_ji + N_lj) // case 3
+		{
+			EdgeID = Edge_il;
+			P      = Pin;
+
+		}
+		else  // case 4
+		{
+			const int32 Cij_l = (N_lj + N_il - N_ji) / 2;
+			if (P <= N_il - Cij_l)
+			{
+				EdgeID = Edge_il;
+				P      = Pin;
+
+			}
+			else
+			{
+				const int32 Clj_i = (N_il + N_ji - N_lj) / 2;
+				EdgeID = Edge_lj;
+				P      = Pin - Clj_i + Cij_l;
+			}
+		}
+		return ExitCoord;
+	}
+	
+	/**
+	* Low-level method: code that uses normal coordinates to continues tracing a surface edge across an intrinsic mesh until it encounters a vertex.
+	*  
+	* 
+	* This function requires that the fist edge crossing has already been computed will populate the a sequence of intrinsic edge crossings for the surface mesh edge
+	* until the path encounters a vertex.  The actual crossings are recorded as a list of the edges crossed, and not the location on the individual edges.  
+	*/
+	template <typename IntrinsicMeshType>
+	void  ContinueTraceSurfaceEdgeAcrossFaces( const IntrinsicMeshType& IntrinsicMesh, TArray<FNormalCoordinates::FEdgeAndCrossingIdx>& Crossings )
+	{
+		const FNormalCoordinates::FEdgeAndCrossingIdx& LastXing = Crossings.Last();
+		checkSlow(LastXing.CIdx != 0); // continuing a face crossing
+
+		const int32 StartTID = LastXing.TID;
+		const int32 StartEID = LastXing.EID;
+		const int32 StartP   = LastXing.CIdx;
+
+		const FNormalCoordinates& NCoords = IntrinsicMesh.GetNormalCoordinates();
+
 		if (!IntrinsicMesh.IsTriangle(StartTID))
 		{
 			return;
@@ -467,7 +836,7 @@
 		const FIndex3i StartTriEIDs = IntrinsicMesh.GetTriEdges(StartTID);
 		checkSlow(StartTriEIDs.IndexOf(StartEID) != -1);
 
-		const int32 NumXStartEID = NCoords.NormalCoord[StartEID];
+		const int32 NumXStartEID = NCoords.NumEdgeCrossing(StartEID);
 		if (StartP > NumXStartEID || StartP < 1) // note the actual permitted range of P is smaller.. 
 		{
 			checkSlow(0); // this shouldn't happen
@@ -476,89 +845,23 @@
 
 		// This is adapted from Algorithm 2 of  Gillespi et al, 2020
 		auto GetNextCrossing = [&IntrinsicMesh, &NCoords](int32& TriID, int32& EdgeID, int32& P)
-		{
-			// get next tri
-			TriID = [&]
-				{
-					const FIndex2i EdgeT = IntrinsicMesh.GetEdgeT(EdgeID);
-					return (EdgeT.A == TriID) ? EdgeT.B : EdgeT.A;
-				}();
-			checkSlow(TriID != -1);
+								{
+									// get next tri
+									TriID = [&]
+										{
+											const FIndex2i EdgeT = IntrinsicMesh.GetEdgeT(EdgeID);
+											return (EdgeT.A == TriID) ? EdgeT.B : EdgeT.A;
+										}();
+									checkSlow(TriID != -1);
 
-			const FIndex3i TriEIDs = IntrinsicMesh.GetTriEdges(TriID);
-			const int32 IndexOf    = TriEIDs.IndexOf(EdgeID);
+									const TTuple<int32, int32> ExitEIDandP = GetCrossingExit(IntrinsicMesh, NCoords, TriID, EdgeID, P);
+									EdgeID = ExitEIDandP.Get<0>();
+									P = ExitEIDandP.Get<1>();
+								};
 
-			const int32 Edge_ji = EdgeID;
-			const int32 Edge_il = TriEIDs[(IndexOf + 1) % 3];
-			const int32 Edge_lj = TriEIDs[(IndexOf + 2) % 3];
-
-			const int32 N_ji = NCoords.NormalCoord[Edge_ji];
-			const int32 N_il = NCoords.NormalCoord[Edge_il];
-			const int32 N_lj = NCoords.NormalCoord[Edge_lj];
-
-			const int32 Pin = P;
-
-			if (N_ji > N_lj + N_il)  // case 1
-			{
-				// some of the N_ji edges that cross edge IJ must connect with the vertex at L 
-				if (P <= N_il)
-				{
-					// exits side il
-					EdgeID = Edge_il;
-					P      = Pin;
-				}
-				else if ((N_il < P) && (P <= N_ji - N_lj))
-				{
-					// this path terminates at vertex
-					// this is a slight abuse.  The convention, when CrossingIndex = 0, we specify the vertex by the edge that originates there.
-					EdgeID = Edge_lj; // i = GetTriangle()[TriEIDs.IndexOf(EdgeID_lj)] will be the vertex.
-					P      = 0;
-				}
-				else
-				{
-					EdgeID = Edge_lj;
-					P      = Pin - (N_ji - N_lj);
-				}
-
-			}
-			else if (N_lj > N_ji + N_il)  // case 2
-			{
-				EdgeID = Edge_lj;
-				P      = Pin + (N_lj - N_ji);
-
-			}
-			else if (N_il > N_ji + N_lj) // case 3
-			{
-				EdgeID = Edge_il;
-				P      = Pin;
-
-			}
-			else  // case 4
-			{
-				const int32 Cij_l = (N_lj + N_il - N_ji) / 2;
-				if (P <= N_il - Cij_l)
-				{
-					EdgeID = Edge_il;
-					P      = Pin;
-
-				}
-				else
-				{
-					const int32 Clj_i = (N_il + N_ji - N_lj) / 2;
-					EdgeID = Edge_lj;
-					P      = Pin - Clj_i + Cij_l;
-				}
-			}
-		};
-
 		int32 P = StartP;
 		int32 EID = StartEID;
 		int32 TID = StartTID;
-		// record first crossing
-		auto& FirstXing = Crossings.AddZeroed_GetRef();
-		FirstXing.TID  = TID;
-		FirstXing.EID  = EID;
-		FirstXing.CIdx = P;
 		do
 		{
 			GetNextCrossing(TID, EID, P);
@@ -572,9 +875,82 @@
 
 	}
 
-	// trace the surface edge - note a surface edge along some surface mesh boundaries can't be specified by this API, but those edges will
-	// be the same in the intrinsic mesh since they can't flip.
-	TArray<FNormalCoordinates::FEdgeAndCrossingIdx> TraceSurfaceEdge( const FIntrinsicEdgeFlipMesh& IntrinsicMesh,
+	/**
+	* Low-level method for use with normal-coordinate equipped intrinsic meshes.
+	* 
+	* Given a partial trace of a surface edge, continue the trace across the intrinsic mesh until the path terminates at one end
+	* of the surface edge.
+	* 
+	* Crossings.Last() must be an edge crossing, and not a vertex point.
+	*/ 
+	template <typename IntrinsicMeshType>
+	void ContinueTraceSurfaceEdge(const IntrinsicMeshType& IntrinsicMesh, TArray<FNormalCoordinates::FEdgeAndCrossingIdx>& Crossings)
+	{
+		if (Crossings.Num() == 0)
+		{
+			return;
+		}
+
+		auto IsIntrinsicVertex = [](const FNormalCoordinates::FEdgeAndCrossingIdx& Xing)->bool
+									{
+										return (Xing.CIdx == 0);
+									};
+
+		checkSlow(!IsIntrinsicVertex(Crossings.Last()));
+
+		auto IsSurfaceVertex = [&IntrinsicMesh, &IsIntrinsicVertex](const FNormalCoordinates::FEdgeAndCrossingIdx& Xing)->bool
+								{
+									if (!IsIntrinsicVertex(Xing))
+									{
+										return false;
+									}
+
+									// convert to intrinsic VID
+									const int32 IndexOfV = Xing.EID;
+									const int32 IntrinsicVID = IntrinsicMesh.GetTriangle(Xing.TID)[IndexOfV];
+
+									return IsVertexPoint(IntrinsicMesh.GetVertexSurfacePoint(IntrinsicVID));
+								};
+		
+		
+		while (!IsSurfaceVertex(Crossings.Last()))
+		{
+			const FNormalCoordinates::FEdgeAndCrossingIdx& LastXing = Crossings.Last();
+			if (!IsIntrinsicVertex(LastXing))
+			{
+				// trace enters intrinsic triangle edge - continue across intrinsic triangle faces.
+				ContinueTraceSurfaceEdgeAcrossFaces(IntrinsicMesh, Crossings);
+				
+			}
+			else
+			{
+				// trace hit an intrinsic vertex, coming from an intrinsic triangle face. 
+
+				// convert to intrinsic VID
+				const int32 IndexOfV = LastXing.EID;
+				const int32 IntrinsicVID = IntrinsicMesh.GetTriangle(LastXing.TID)[IndexOfV];
+
+				checkSlow(IsEdgePoint(IntrinsicMesh.GetVertexSurfacePoint(IntrinsicVID)));
+
+				// advance along the surface edge until it either enters an intrinsic triangle or 
+				// follows on more intrinsic edges to terminate at a surface vertex. 
+				const bool bEnterTypeEdge = false; // we entered across a triangle face
+				const int32 EnterEID = LastXing.TID;
+				ContinueFromEdgePointCrossing(IntrinsicMesh, IntrinsicVID, EnterEID, bEnterTypeEdge, Crossings);
+			}
+		}
+	}
+
+	/**
+	* Low-level method for use with normal-coordinate equipped intrinsic meshes.
+	* 
+	* Traces a surface edge across an intrinsic mesh, the result is an array of intrinsic edges crossed and an integer coordinates
+	* that specify the intrinsic edge specific index of the crossing.   
+	* 
+	* The surface edge to trace is specified by the triangle adjacent to the edge, and the local index of the edge within that triangle
+	*/ 
+	template <typename IntrinsicMeshType>
+	TArray<FNormalCoordinates::FEdgeAndCrossingIdx> TraceSurfaceEdge( const IntrinsicMeshType& IntrinsicMesh,
 		                                                              const int32 SurfaceTID, const int32 IndexOf )
 	{
 		const FNormalCoordinates& NCoords = IntrinsicMesh.GetNormalCoordinates();
@@ -592,7 +968,7 @@
 		// the surface edge we are tracing
 		const int32 TraceEID = SurfaceMesh.GetTriEdge(SurfaceTID, IndexOf);
 
-		checkSlow(!SurfaceMesh.IsBoundaryEdge(TraceEID));
+		
 		const int32 RefEID = NCoords.VIDToReferenceEID[StartVID];
 		const int32 OrderOfTraceEID = NCoords.GetEdgeOrder(StartVID, TraceEID);
 		const int32 ValenceOfStartVID = NCoords.RefVertDegree[StartVID];
@@ -613,81 +989,72 @@
 		{
 
 			// where to start when visiting the intrinsic triangles that are adjacent to the startVID
-			const int32 VistorStartEID =[&] 
-				{	
-					if (SurfaceMesh.IsBoundaryVertex(StartVID))
-					{
-						// due to the way the intrinsic mesh is constructed we know the boundary edges will have the same EID in both meshes.
-						// recall, these can't flip, so they are fixed in the intrinsic mesh.
-						return RefEID;
-					}
-					else
-					{
-						for (int32 NbrEID : IntrinsicMesh.VtxEdgesItr(StartVID))
-						{
-							// get the first
-							return NbrEID;
-						}
-						return -1; // shouldn't reach here
-					}
-				}();
+			const int32 VistorStartEID = IdentifyInitialAdjacentEdge(IntrinsicMesh, StartVID);
 
-			// when visiting each adjacent triangle (in CCW order) test if the edge we trace is inside this triangle.
+			// when visiting each adjacent intrinsic triangle (in CCW order) test if the surface edge we trace is inside this triangle.
 			auto EdgeFinder = [&](int32 IntrinsicTID, int32 IntrinsicEID, int32 IdxOf)->bool
-			{
+								{
 
 
-				const FIndex3i IntrinsicTriEIDs = IntrinsicMesh.GetTriEdges(IntrinsicTID);
-				const int32 ThisRoundabout      = NCoords.RoundaboutOrder[IntrinsicTID][IdxOf];
-				const bool bEquivalentToSurface = (NCoords.NormalCoord[IntrinsicEID] == 0) && (ThisRoundabout == OrderOfTraceEID);
-				int32 NextRoundabout = -1;
+									const FIndex3i IntrinsicTriEIDs = IntrinsicMesh.GetTriEdges(IntrinsicTID);
+									const int32 ThisRoundabout      = NCoords.RoundaboutOrder[IntrinsicTID][IdxOf];
+									const bool bOnSurfaceEdge       = NCoords.IsSurfaceEdgeSegment(IntrinsicEID);
+									const bool bEquivalentToSurface = bOnSurfaceEdge  && (ThisRoundabout == OrderOfTraceEID);
+									int32 NextRoundabout = -1;
 
-				bool bShouldBreak = false;
+									bool bShouldBreak = false;
 
-				if (bEquivalentToSurface) // test if the current intrinsic edge is the same as the surface mesh trace edge
-				{
-					bShouldBreak = true;
-				}
-				else // test if the edge starts in this intrinsic triangle
-				{
-					const int32 NextIntrinsicEID      = IntrinsicTriEIDs[(IdxOf + 2) % 3];
-					const FIndex2i NextIntrinsicEdgeT = IntrinsicMesh.GetEdgeT(NextIntrinsicEID);
-					const int32 NextIntrinsicTID      = (NextIntrinsicEdgeT.A == IntrinsicTID) ? NextIntrinsicEdgeT.B : NextIntrinsicEdgeT.A;
+									if (bEquivalentToSurface) // test if the current intrinsic edge is the same as the surface mesh trace edge
+									{
+										bShouldBreak = true;
+									}
+									else // test if the edge starts in this intrinsic triangle
+									{
+										const int32 NextIntrinsicEID      = IntrinsicTriEIDs[(IdxOf + 2) % 3];
+										const FIndex2i NextIntrinsicEdgeT = IntrinsicMesh.GetEdgeT(NextIntrinsicEID);
+										const int32 NextIntrinsicTID      = (NextIntrinsicEdgeT.A == IntrinsicTID) ? NextIntrinsicEdgeT.B : NextIntrinsicEdgeT.A;
 
-					if (NextIntrinsicTID != -1)
-					{
-						const int32 NextIndexOf = IntrinsicMesh.GetTriEdges(NextIntrinsicTID).IndexOf(NextIntrinsicEID);
-						NextRoundabout          = NCoords.RoundaboutOrder[NextIntrinsicTID][NextIndexOf];
-						if (NextRoundabout < ThisRoundabout)
-						{
-							// we have crossed the zero roundabout cut
-							NextRoundabout += ValenceOfStartVID;
-						}
-						checkSlow(IntrinsicMesh.GetTriangle(NextIntrinsicTID)[NextIndexOf] == StartVID);
-						//checkSlow(NextRoundabout != 0 || ThisRoundabout != 0)
-						if ((NextRoundabout > OrderOfTraceEID) && (OrderOfTraceEID >= ThisRoundabout))
-						{
-							bShouldBreak = true;
-						}
-					}
-					else // mesh boundary case and we made it to the last triangle w/o finding this edge. it must be in this one.
-					{
-						bShouldBreak = true;
-					}
-				}
+										if (NextIntrinsicTID != -1)
+										{
+											const int32 NextIndexOf = IntrinsicMesh.GetTriEdges(NextIntrinsicTID).IndexOf(NextIntrinsicEID);
+											NextRoundabout          = NCoords.RoundaboutOrder[NextIntrinsicTID][NextIndexOf];
+											if (NextRoundabout < ThisRoundabout) // Order = {NextRO |zero cut | ThisRO}
+											{
+												// we have crossed the zero roundabout cut ( e.g. NextRo = 2oclock and ThisRo = 11oclock)
+												if ( NextRoundabout > OrderOfTraceEID) // Order = {NextRO, Trace | zero cut | ThisRO}
+												{
+													bShouldBreak = true;
+												}
+												else 
+												{ 
+													NextRoundabout += ValenceOfStartVID; // Order = {NextRO, | zero cut |, Trace(?), ThisRO}
+												}
+											}
+											checkSlow(IntrinsicMesh.GetTriangle(NextIntrinsicTID)[NextIndexOf] == StartVID);
+											//checkSlow(NextRoundabout != 0 || ThisRoundabout != 0)
+											if ((NextRoundabout > OrderOfTraceEID) && (OrderOfTraceEID >= ThisRoundabout))
+											{
+												bShouldBreak = true;
+											}
+										}
+										else // mesh boundary case and we made it to the last triangle w/o finding this edge. it must be in this one.
+										{
+											bShouldBreak = true;
+										}
+									}
 
-				if (bShouldBreak)
-				{
-					FinderInfo.bIsAlsoSurfaceEdge = bEquivalentToSurface;
-					FinderInfo.TID                = IntrinsicTID;
-					FinderInfo.EID                = IntrinsicEID;
-					FinderInfo.IdxOf              = IdxOf;
-					FinderInfo.FirstRoundabout    = ThisRoundabout;
-					FinderInfo.SecondRoundabout   = NextRoundabout;
-				}
+									if (bShouldBreak)
+									{
+										FinderInfo.bIsAlsoSurfaceEdge = bEquivalentToSurface;
+										FinderInfo.TID                = IntrinsicTID;
+										FinderInfo.EID                = IntrinsicEID;
+										FinderInfo.IdxOf              = IdxOf;
+										FinderInfo.FirstRoundabout    = ThisRoundabout;
+										FinderInfo.SecondRoundabout   = NextRoundabout;
+									}
 
-				return bShouldBreak;
-			};
+									return bShouldBreak;
+								};
 			VisitVertexAdjacentElements(IntrinsicMesh, StartVID, VistorStartEID, EdgeFinder);
 
 			checkSlow(FinderInfo.TID != -1); // should have found something!
@@ -700,43 +1067,93 @@
 		StartXing.TID   = FinderInfo.TID;
 		StartXing.EID   = FinderInfo.IdxOf; // TriVIDs(IdxOf) = StartVID
 		StartXing.CIdx  = 0;
-		if (FinderInfo.bIsAlsoSurfaceEdge)
+
+		
+		if (FinderInfo.bIsAlsoSurfaceEdge) 
 		{
-			// only need to add end vertex location
+			// the surface edge we are tracing is initially coincident with an intrinsic edge. 
+			
+			// add end intrinsic vertex location
 			auto& EndXing = Crossings.AddZeroed_GetRef();
 			EndXing.TID   = FinderInfo.TID;
 			EndXing.EID   = (FinderInfo.IdxOf + 1) % 3;  // TriVIDs( (IdxOf +1)%3) = EndVID
 			EndXing.CIdx  = 0;
+
+			const int32 IndexOfV        = (FinderInfo.IdxOf + 1) % 3;
+			const int32 EndIntrinsicVID = IntrinsicMesh.GetTriangle(FinderInfo.TID)[IndexOfV];
+			
+			// the surface edge will continue if the end intrinsic vertex is not a surface vertex.
+			
+			// advance along the surface edge until it either enters an intrinsic triangle or 
+			// follows on more intrinsic edges to terminate at a surface vertex. 
+			if (IsEdgePoint(IntrinsicMesh.GetVertexSurfacePoint(EndIntrinsicVID)))
+			{
+				const bool bEnterTypeEdge  = true;
+				const int32 EnterEID       = FinderInfo.EID;
+				ContinueFromEdgePointCrossing(IntrinsicMesh,  EndIntrinsicVID, EnterEID, bEnterTypeEdge, Crossings);
+				
+				const bool bVertexTerminated = (Crossings.Last().CIdx == 0);
+				if (bVertexTerminated)
+				{
+					// Made it to the end of the surface edge since the trace must have hit an intrinsic vertex that is also a surface vertex. 
+					return MoveTemp(Crossings);
+				}
+				else
+				{
+					// continue across intrinsic faces
+					ContinueTraceSurfaceEdge(IntrinsicMesh, Crossings);
+				}
+			}
+			else
+			{	
+				// Made it to the end of the surface edge since the trace must have hit an intrinsic vertex that is also a surface vertex. 
+				return MoveTemp(Crossings);
+			}
+	
 		}
 		else
 		{
 			// edge trace starts at StartVID and exits the opposite edge of intrinsic Tri TID (but it isn't an edge of the intrinsic tri).
 			const FIndex3i TriEIDs  = IntrinsicMesh.GetTriEdges(FinderInfo.TID);
 			const int32 OppEID      = TriEIDs[(FinderInfo.IdxOf + 1) % 3];
-			const int32 NormalCoord = NCoords.NormalCoord[FinderInfo.EID];
+
 			const int32 CrossingIdx = [&]
-				{
-					if (NormalCoord == 0)
-					{
-						return OrderOfTraceEID - FinderInfo.FirstRoundabout;
-					}
-					else
-					{
-						return OrderOfTraceEID - FinderInfo.FirstRoundabout + 1 + NormalCoord;
-					}
-				}();
+										{
+											const int32 AdvanceFromFirstRO = ( (ValenceOfStartVID + OrderOfTraceEID) - FinderInfo.FirstRoundabout) % ValenceOfStartVID;
+											if (NCoords.IsSurfaceEdgeSegment(FinderInfo.EID))
+											{
+												return AdvanceFromFirstRO;
+											}
+											else
+											{
+												const int32 NumCrossings = NCoords.NumEdgeCrossing(FinderInfo.EID);
+												return AdvanceFromFirstRO + 1 + NumCrossings;
+											}
+										}();
 
 			checkSlow(CrossingIdx > 0); // would be zero if the edge was also a surface edge, but that case is handled above
+			
+			// the first intrinsic edge crossing.  Need this to jump-start the continuation.
+			auto& Xing = Crossings.AddZeroed_GetRef();
+			Xing.TID = FinderInfo.TID;
+			Xing.EID = OppEID;
+			Xing.CIdx = CrossingIdx;
 
-			ContinueTraceSurfaceEdge(IntrinsicMesh, FinderInfo.TID, OppEID, CrossingIdx, Crossings);
+			ContinueTraceSurfaceEdge(IntrinsicMesh, Crossings);
+		
 		}
 
 		return MoveTemp(Crossings);
 	}
 
-
-
-	TArray<FNormalCoordinates::FEdgeAndCrossingIdx> TraceSurfaceEdge( const FIntrinsicEdgeFlipMesh& IntrinsicMesh,
+	/**
+	* Low-level method for use with normal-coordinate equipped intrinsic meshes.
+	* 
+	* Traces a surface edge across an intrinsic mesh, the result is an array of intrinsic edges crossed each with a corresponding 
+	* integer. The integer is used to disambiguate when a single intrinsic edge is crossed by multiple surface edges. 
+	*/ 
+	template <typename IntrinsicMeshType>
+	TArray<FNormalCoordinates::FEdgeAndCrossingIdx> TraceSurfaceEdge( const IntrinsicMeshType& IntrinsicMesh,
 																	  const int32 SurfaceEID, const bool bReverse )
 	{
 		const FNormalCoordinates& NCoords = IntrinsicMesh.GetNormalCoordinates();
@@ -749,33 +1166,50 @@
 
 		const FIndex2i EdgeV = SurfaceMesh.GetEdgeV(SurfaceEID);
 		const int32 StartVID = (bReverse) ? EdgeV.B : EdgeV.A;
+	
 
-		// Special case of tracing a surface edge on the mesh boundary.  This will be the same as an intrinsic edge.
-		if (SurfaceMesh.IsBoundaryEdge(SurfaceEID))
+		// if this surface edge corresponds to an intrinsic edge, then the trace is trivial
 		{
-			// boundary edges don't flip.  Will be the same as the intrinsic edge.
+			// note this test relies on the fact that intrinsic vertices that sit on surface mesh vertices 
+			// will have the same vertex IDs by construction. 
+			int32 EquivalentIntrinsicEID = -1;
+			const int32 EndVID = (bReverse) ? EdgeV.A : EdgeV.B;
+			for (int32 IntrinsicEID : IntrinsicMesh.VtxEdgesItr(StartVID))
+			{
+				const FIndex2i IntrinsicEdgeV = IntrinsicMesh.GetEdgeV(IntrinsicEID);
+				if (IntrinsicEdgeV.IndexOf(EndVID) != -1 && NCoords.IsSurfaceEdgeSegment(IntrinsicEID))
+				{
+					EquivalentIntrinsicEID = IntrinsicEID;
+					break;
+				}
+			}
 
-			const int32 IntrinsicEID      = SurfaceEID;
-			const int32 IntrinsicTID      = IntrinsicMesh.GetEdgeT(IntrinsicEID).A;
-			const int32 IndexOf           = IntrinsicMesh.GetTriEdges(IntrinsicTID).IndexOf(IntrinsicEID);
-			const bool bRerverseIntrinsic = !(IntrinsicMesh.GetTriangle(IntrinsicTID)[IndexOf] == StartVID);
+			if (EquivalentIntrinsicEID != -1)
+			{
+				checkSlow(IntrinsicMesh.IsEdge(EquivalentIntrinsicEID));
+				const int32 IntrinsicEID      = EquivalentIntrinsicEID;
+				const int32 IntrinsicTID      = IntrinsicMesh.GetEdgeT(IntrinsicEID).A;
+				const int32 IndexOf           = IntrinsicMesh.GetTriEdges(IntrinsicTID).IndexOf(IntrinsicEID);
+				const bool bRerverseIntrinsic = !(IntrinsicMesh.GetTriangle(IntrinsicTID)[IndexOf] == StartVID);
 
 
-			TArray< FNormalCoordinates::FEdgeAndCrossingIdx > Crossings;
-			Crossings.SetNumUninitialized(2);
-			auto& XingStart = Crossings[0];
-			XingStart.TID   = IntrinsicTID;
-			XingStart.EID   = (bRerverseIntrinsic) ? (IndexOf + 1) % 3 : IndexOf;
-			XingStart.CIdx  = 0;
+				TArray< FNormalCoordinates::FEdgeAndCrossingIdx > Crossings;
+				Crossings.SetNumUninitialized(2);
+				auto& XingStart = Crossings[0];
+				XingStart.TID   = IntrinsicTID;
+				XingStart.EID   = (bRerverseIntrinsic) ? (IndexOf + 1) % 3 : IndexOf;
+				XingStart.CIdx  = 0;
 
-			auto& XingEnd = Crossings[1];
-			XingEnd.TID   = IntrinsicTID;
-			XingEnd.EID   = (bRerverseIntrinsic) ? IndexOf : (IndexOf + 1) % 3;
-			XingEnd.CIdx  = 0;
+				auto& XingEnd = Crossings[1];
+				XingEnd.TID   = IntrinsicTID;
+				XingEnd.EID   = (bRerverseIntrinsic) ? IndexOf : (IndexOf + 1) % 3;
+				XingEnd.CIdx  = 0;
 
-			return MoveTemp(Crossings);
+				return MoveTemp(Crossings);
+			}
 		}
-		else  // the edge isn't a surface edge:  encode its direction by identifying it as an edge of a triangle and do the trace
+
+		// the edge isn't equivalent to an intrinsic edge:  encode its direction by identifying it as an edge of a triangle and do the trace
 		{
 
 			const FIndex2i EdgeT = SurfaceMesh.GetEdgeT(SurfaceEID);
@@ -783,24 +1217,365 @@
 			int32 IndexOf        = SurfaceMesh.GetTriEdges(TID).IndexOf(SurfaceEID);
 			if (SurfaceMesh.GetTriangle(TID)[IndexOf] != StartVID)
 			{
-				TID = EdgeT.B;
-				IndexOf = SurfaceMesh.GetTriEdges(TID).IndexOf(SurfaceEID);
-				checkSlow(SurfaceMesh.GetTriangle(TID)[IndexOf] == StartVID);
+				if (EdgeT.B != -1)
+				{ 
+					TID = EdgeT.B;
+					IndexOf = SurfaceMesh.GetTriEdges(TID).IndexOf(SurfaceEID);
+					checkSlow(SurfaceMesh.GetTriangle(TID)[IndexOf] == StartVID);
+
+					return TraceSurfaceEdge(IntrinsicMesh, TID, IndexOf);
+				}
+				else
+				{
+					TArray<FNormalCoordinates::FEdgeAndCrossingIdx> TraceResult = TraceSurfaceEdge(IntrinsicMesh, TID, IndexOf);
+					Algo::Reverse(TraceResult);
+					return MoveTemp(TraceResult);
+				}
 			}
-			return TraceSurfaceEdge(IntrinsicMesh, TID, IndexOf);
+			else
+			{ 
+				return TraceSurfaceEdge(IntrinsicMesh, TID, IndexOf);
+			}
 		}
 	}
 
 
 	/**
+	* Low-level utility for use with normal-coordinate equipped intrinsic meshes.
+	* 
+	* Utility to trace an edge defined  on the TraceMesh across the HostMesh, given a list of host mesh edges intersected by the trace mesh.
+	* this assumes the host mesh and trace mesh come from an intrinsic mesh, surface mesh pair.
+	* 
+	* From the list of edges being crossed, this creates a 2d triangle strip of the triangles being traversed, and translates Trace{Start,End}SurfacePosition
+	* to this space so it can compute the actual intersection locations of the trace with the triangle edges (stored line-based barycentric coord, alpha).
+	* 
+	* Note: this assumes, that StartSurfacePoint and EndSurfacePoints are actually points on the edge, and HostXing lists the Ids of the edges
+	* being crossed in order.
+	*/
+	template<typename HostMeshType, typename TraceMeshType, typename TraceVIDToSurfacePointFtor>
+	void ConvertEdgesCrossed( const FSurfacePoint& StartSurfacePoint, const FSurfacePoint& EndSurfacePoint,
+							  const TArray<int32>& HostEdgesCrossed, const HostMeshType& HostMesh,
+		                      const TraceMeshType& TraceMesh, double CoalesceThreshold,
+							  const TraceVIDToSurfacePointFtor& VIDToSurfacePoint,
+							  TArray<FSurfacePoint>& EdgeTrace)
+	{
+	
+		if (HostEdgesCrossed.Num() == 0)
+		{
+			// no edge crossings to convert.
+			return;
+		}
+
+		// utility to keep from adding duplicate VIDs.  Duplicate VIDs could result when Coalescing..
+		auto AddVIDToTrace = [&EdgeTrace](const int32 HostVID)
+							{
+								if (EdgeTrace.Num() == 0)
+								{
+									EdgeTrace.Emplace(HostVID);
+								}
+								else
+								{
+									const FSurfacePoint& LastPoint = EdgeTrace.Last();
+									const bool bSameAsLast         = ( IsVertexPoint(LastPoint)
+																	&& LastPoint.Position.VertexPosition.VID == HostVID );
+									if (!bSameAsLast)
+									{
+										EdgeTrace.Emplace(HostVID);
+									}
+								}
+							};
+		
+		
+		// knowing the host mesh edges this trace edge crosses, we make a 2d triangle strip by unfolding the 3d triangles the edge crosses
+		// and solve a 2x2 problem for each 2d edge crossed.
+
+		// struct holds bare-bones 2d triangle strip and ability to map vertexIDs from the src mesh
+		struct
+		{
+			TMap<int32, int32> ToStripIndexMap;    // maps from mesh VID to triangle strip VID.
+			TArray<FVector2d> StripVertexBuffer;
+		} TriangleStrip;
+		TriangleStrip.StripVertexBuffer.Reserve(3 + HostEdgesCrossed.Num());
+
+
+		// Functor to add a triangle to the triangle strip. This assumes that at least one shared edge of this triangle already exits in the strip
+		auto AddTriangleToStrip = [&HostMesh, &TriangleStrip](const int32 HostTID)
+				{
+
+					FIndex3i TriVIDs = HostMesh.GetTriangle(HostTID);
+
+					// two of the tri vertices are already in the buffer.  Identify the new one.
+					int32 IndexOf = -1;
+					for (int32 i = 0; i < 3; ++i)
+					{
+						const int32 VID       = TriVIDs[i];
+						const int32* StripVID = TriangleStrip.ToStripIndexMap.Find(VID);
+						if (!StripVID)
+						{
+							IndexOf = i;
+							break;
+						}
+					}
+
+					// no need to do anything if all vertices had previously been added, 
+					if (IndexOf == -1)
+					{
+						return;
+					}
+
+
+					const FVector3d Verts[3] = { HostMesh.GetVertex(TriVIDs[0]),  HostMesh.GetVertex(TriVIDs[1]),  HostMesh.GetVertex(TriVIDs[2]) };
+					// with new vert last (i.e. V2).
+					const FIndex3i Reordered((IndexOf + 1) % 3, (IndexOf + 2) % 3, IndexOf);
+					// the spanning vectors 
+					const FVector3d E1 = Verts[Reordered[1]] - Verts[Reordered[0]];
+					const FVector3d E2 = Verts[Reordered[2]] - Verts[Reordered[0]];
+
+					// coordinates of V2 relative to the direction of E1, and its orthogonal complement.
+					const double E1DotE2       = E2.Dot(E1);
+					const double E1LengthSqr   = FMath::Max(E1.SizeSquared(), TMathUtilConstants<double>::ZeroTolerance);
+					const double E1Length      = FMath::Sqrt(E1LengthSqr);
+					const double E1Dist        = E2.Dot(E1) / E1Length;
+					const double OrthE1DistSqr = FMath::Max(0., E2.SizeSquared() - E1DotE2 * E1DotE2 / E1LengthSqr);
+					const double OrthE1Dist    = FMath::Sqrt(OrthE1DistSqr);
+
+					// 2d version of E1
+					const int32* StripTriV0 = TriangleStrip.ToStripIndexMap.Find(TriVIDs[Reordered[0]]);
+					const int32* StripTriV1 = TriangleStrip.ToStripIndexMap.Find(TriVIDs[Reordered[1]]);
+					checkSlow(StripTriV0); checkSlow(StripTriV1);
+
+					const FVector2d& StripTriVert0 = TriangleStrip.StripVertexBuffer[*StripTriV0];
+					const FVector2d& StripTriVert1 = TriangleStrip.StripVertexBuffer[*StripTriV1];
+					const FVector2d StripE1        = (StripTriVert1 - StripTriVert0);
+
+					const FVector2d StripE1Perp(-StripE1[1], StripE1[0]); // Rotate StripE1 90 CCW.
+
+					// new vertex 2d position
+					const FVector2d StripTriVert2 = StripTriVert0 + (StripE1)*E1DotE2 / E1LengthSqr + (StripE1Perp / E1Length) * OrthE1Dist;
+					TriangleStrip.StripVertexBuffer.Add(StripTriVert2);
+					const int32 StripVID          = TriangleStrip.StripVertexBuffer.Num() - 1;
+					const int32 SurfaceVID        = TriVIDs[IndexOf];
+					TriangleStrip.ToStripIndexMap.Add(SurfaceVID, StripVID);
+				};
+
+		
+
+		// find host triangle that contains both the StartVID and the first host edge we cross.
+		const int32 FirstHostTID = [&]
+			{
+				if (IsVertexPoint(StartSurfacePoint))
+				{ 
+					const int32 HostEID = HostEdgesCrossed[0];
+					const FIndex2i HostEdgeT = HostMesh.GetEdgeT(HostEID);
+					const FIndex3i TriAVIDs = HostMesh.GetTriangle(HostEdgeT.A);
+					const FIndex3i TriBVIDs = HostMesh.GetTriangle(HostEdgeT.B);
+
+					const int32 HostStartVID = StartSurfacePoint.Position.VertexPosition.VID;
+					
+					return (TriAVIDs.IndexOf(HostStartVID) != -1) ? HostEdgeT.A : HostEdgeT.B;
+				}
+				else if (IsEdgePoint(StartSurfacePoint))
+				{
+					const int32 StartEID = StartSurfacePoint.Position.EdgePosition.EdgeID;
+					const int32 ExitEID  = HostEdgesCrossed[0];
+					const FIndex2i AdjTris = HostMesh.GetEdgeT(StartEID);
+					const int32 NumAdj = (AdjTris[1] == -1) ? 1 : 2;
+					for (int i = 0; i < NumAdj; ++i)
+					{
+						int32 TID = AdjTris[i];
+						if (HostMesh.GetTriEdges(TID).IndexOf(ExitEID) != -1)
+						{
+							return TID;
+						}
+					}
+					return -1;
+				}
+				else
+				{
+					checkSlow(IsFacePoint(StartSurfacePoint));
+					return StartSurfacePoint.Position.TriPosition.TriID;
+				}
+			}();
+
+		// jump-start the process of making the triangle strip by adding the first 2 verts
+		// that define the edge prior (in the ccw sense) to the first edge crossing.
+		{
+			const FIndex3i TriVIDs = HostMesh.GetTriangle(FirstHostTID);
+			const FIndex3i TriEIDs = HostMesh.GetTriEdges(FirstHostTID);
+			const int32 FirstEIDXed = HostEdgesCrossed[0];
+			const int32 IndexOfe = TriEIDs.IndexOf(FirstEIDXed);
+			const int32 FirstVID = TriVIDs[(IndexOfe + 2)%3];
+			const int32 SecondVID = TriVIDs[(IndexOfe)];
+
+			const FVector3d Vert0 = HostMesh.GetVertex(FirstVID);
+			const FVector3d Vert1 = HostMesh.GetVertex(SecondVID);
+
+			const double LSqr = FMath::Max((Vert0 - Vert1).SizeSquared(), TMathUtilConstants<double>::ZeroTolerance);
+			const double L    = FMath::Sqrt(LSqr);
+
+			const FVector2d StripVert0(0., 0.);
+			const FVector2d StripVert1(L, 0.);
+
+			TriangleStrip.StripVertexBuffer.Add(StripVert0);
+			TriangleStrip.ToStripIndexMap.Add(FirstVID, 0);
+
+			TriangleStrip.StripVertexBuffer.Add(StripVert1);
+			TriangleStrip.ToStripIndexMap.Add(SecondVID, 1);
+		}
+
+		// unfold the triangle strip, add the first triangle and then all subsequent ones
+			
+		AddTriangleToStrip(FirstHostTID);
+		{
+			int32 LastHostTID = FirstHostTID;
+			for (int32 HostEID : HostEdgesCrossed)
+			{
+				const FIndex2i EdgeT = HostMesh.GetEdgeT(HostEID);
+				LastHostTID          = (EdgeT.A == LastHostTID) ? EdgeT.B : EdgeT.A;
+				AddTriangleToStrip(LastHostTID);
+			}
+		}
+
+		// translate a trace point on the surface of the host mesh, to Fvector2d relative to the 2d triangle strip 
+		auto To2dStripPosition = [&](const FSurfacePoint& TraceSurfacePoint)->FVector2d
+								{
+									if (IsVertexPoint(TraceSurfacePoint))
+									{
+										// end point is a vertex in the host mesh
+										const auto& HostVertexPosition = TraceSurfacePoint.Position.VertexPosition;
+										const int32 HostVID            = HostVertexPosition.VID;
+										const int32* StripVID          = TriangleStrip.ToStripIndexMap.Find(HostVID);
+										checkSlow(StripVID);
+										return TriangleStrip.StripVertexBuffer[*StripVID];
+									}
+									if (IsEdgePoint(TraceSurfacePoint))
+									{
+										// end point lies on an edge of the host mesh. identify the edge vertices in the 2d-triangle strip
+										// and reconstruct the point from alpha
+
+										const auto& HostEdgePosition = TraceSurfacePoint.Position.EdgePosition;
+										const double Alpha		     = HostEdgePosition.Alpha;
+										const FIndex2i HostEdgeV     = HostMesh.GetEdgeV(HostEdgePosition.EdgeID);
+										const FIndex2i StripEdgeV( *TriangleStrip.ToStripIndexMap.Find(HostEdgeV.A),
+																	*TriangleStrip.ToStripIndexMap.Find(HostEdgeV.B) );
+										const FVector2d StripEdgePos[2] = { TriangleStrip.StripVertexBuffer[StripEdgeV.A],
+																			TriangleStrip.StripVertexBuffer[StripEdgeV.B] };
+										return Alpha * StripEdgePos[0] + (1. - Alpha) * StripEdgePos[1];
+									}
+									else
+									{
+										// end point lies in the face of the last triangle. identify the tri vertices in the 2d-triangle strip and 
+										// reconstruct the point from the barycentric coordinates.
+
+										checkSlow(IsFacePoint(TraceSurfacePoint));
+										const auto& HostTriPosition = TraceSurfacePoint.Position.TriPosition;
+										const FVector3d Barycentric = HostTriPosition.BarycentricCoords;
+										const FIndex3i HostTri = HostMesh.GetTriangle(HostTriPosition.TriID);
+										const FIndex3i StripTri( *TriangleStrip.ToStripIndexMap.Find(HostTri[0]),
+																	*TriangleStrip.ToStripIndexMap.Find(HostTri[1]),
+																	*TriangleStrip.ToStripIndexMap.Find(HostTri[2]) );
+										const FVector2d StripTriPos[3] = { TriangleStrip.StripVertexBuffer[StripTri[0]],
+																			TriangleStrip.StripVertexBuffer[StripTri[1]],
+																			TriangleStrip.StripVertexBuffer[StripTri[2]] };
+
+										return Barycentric[0] * StripTriPos[0] + Barycentric[1] * StripTriPos[1] + Barycentric[2] * StripTriPos[2];
+									}
+								};
+			
+			
+
+		// translate the end of the trace edge to a position in the unfolded triangle strip
+		
+	
+		const FVector2d TraceStartVertex = To2dStripPosition(StartSurfacePoint);
+		const FVector2d TraceEndVertex   = To2dStripPosition(EndSurfacePoint);										
+		const FVector2d TraceVector      = TraceEndVertex - TraceStartVertex;
+
+		// loop over the two-d versions of the host edges the path crosses and find the intersection.
+		for (int32 HostEID : HostEdgesCrossed)
+		{
+			const FIndex2i EdgeV = HostMesh.GetEdgeV(HostEID);
+			const int32* StripA  = TriangleStrip.ToStripIndexMap.Find(EdgeV.A);
+			const int32* StripB  = TriangleStrip.ToStripIndexMap.Find(EdgeV.B);
+			checkSlow(StripA); checkSlow(StripB);
+			const FIndex2i StripEdgeV(*StripA, *StripB);
+
+			const FVector2d StripVertA = TriangleStrip.StripVertexBuffer[StripEdgeV.A];
+			const FVector2d StripVertB = TriangleStrip.StripVertexBuffer[StripEdgeV.B];
+
+			// solve Alpha*StripVertA + (1-Alpha)StripVertB = Gamma Start + (1-Gamma)End.
+			// i.e. 
+			// Alpha * (StripVertA - StripVertB) + Gamma * (End - Start) =  End - StripVertB.
+			//     here End - Start = TraceVector
+			// write this as 2x2 matrix problem M.x = b solving for unknown vector x=(alpha, gamma).
+
+			// matrix
+			double m[2][2];
+			m[0][0] = (StripVertA - StripVertB)[0];  m[0][1] = TraceVector[0];
+			m[1][0] = (StripVertA - StripVertB)[1];  m[1][1] = TraceVector[1];
+
+			// b-vector
+			const FVector2d b = TraceEndVertex - StripVertB;
+
+			const double Det = m[0][0] * m[1][1] - m[0][1] * m[1][0];
+			// inverse: not yet scaled by 1/det
+			double invm[2][2];
+			invm[0][0] =  m[1][1];   invm[0][1] = -m[0][1];
+			invm[1][0] = -m[1][0];   invm[1][1] =  m[0][0];
+
+			// solve for alpha only ( don't care about gamma ) 
+			double alpha = invm[0][0] * b[0] + invm[0][1] * b[1];
+
+			if (FMath::Abs(Det) < TMathUtilConstants<double>::ZeroTolerance)
+			{
+				// this surface edge and the intrinsic edge that intersects it are nearly parallel.  
+				// Assume the crossing happens at the farther vertex.
+				const double dA = StripVertA.SquaredLength();
+				const double dB = StripVertB.SquaredLength();
+
+				alpha = (dA > dB) ? 1. : 0.;
+			}
+			else
+			{
+				alpha = FMath::Clamp(alpha / Det, 0., 1.);
+			}
+
+			// may want to convert this edge crossing to a vertex crossing, if it is close enough.
+			int32 CoalesceVID = -1;
+			bool bSnapToVert = false;
+
+			if (alpha < CoalesceThreshold)
+			{
+				CoalesceVID = EdgeV.B;
+				bSnapToVert = true;
+			}
+			else if ((1. - alpha) < CoalesceThreshold)
+			{
+				CoalesceVID = EdgeV.A;
+				bSnapToVert = true;
+			}
+
+			if (bSnapToVert)
+			{
+				AddVIDToTrace(CoalesceVID);
+			}
+			else
+			{
+				EdgeTrace.Emplace(HostEID, alpha);
+			}
+		}
+
+	}
+
+	/**
 	* Utility to trace an edge defined (by TraceEID) on the TraceMesh across the HostMesh, given a list of host mesh edges intersected by the trace mesh.
-	* note this assumes that the TraceMesh and HostMesh share the same vertex set
+	* this assumes the host mesh and trace mesh come from an intrinsic mesh, surface mesh pair.
 	*/
-	template<typename HostMeshType, typename TraceMeshType>
-	TArray<FIntrinsicEdgeFlipMesh::FSurfacePoint> TraceEdgeOverHost( const int32 TraceEID, const TArray<int32>& HostXings, const HostMeshType& HostMesh,
-		                                                             const TraceMeshType& TraceMesh, double CoalesceThreshold, bool bReverse )
+	template<typename HostMeshType, typename TraceMeshType, typename TraceVIDToSurfacePointFtor>
+	TArray<FSurfacePoint> TraceEdgeOverHost( const int32 TraceEID, const TArray<int32>& HostEdgesCrossed, const HostMeshType& HostMesh,
+											 const TraceMeshType& TraceMesh, double CoalesceThreshold, bool bReverse, 
+											 const TraceVIDToSurfacePointFtor& VIDToSurfacePoint)
 	{
-		using FSurfacePoint = FIntrinsicEdgeFlipMesh::FSurfacePoint;
 		TArray<FSurfacePoint> EdgeTrace;
 
 		if (!TraceMesh.IsEdge(TraceEID))
@@ -809,263 +1584,44 @@
 			return MoveTemp(EdgeTrace);
 		}
 
-		// utility to keep from adding duplicate VIDs.  Duplicate VIDs could result when Coalescing..
-		auto AddVIDToTrace = [&EdgeTrace](const int32 VID)
+
+		EdgeTrace.Reserve(HostEdgesCrossed.Num()+2);
+
+		// find start (.A) and end (.B) trace mesh vids for this edge
+		const FIndex2i OrderedTraceEdgeV = [&]
+											{
+												// by convention the intrinsic edge direction will be defined by the first adj triangle
+												const int32 AdjTID     = TraceMesh.GetEdgeT(TraceEID).A;
+												const int32 IndexOf    = TraceMesh.GetTriEdges(AdjTID).IndexOf(TraceEID);
+												const FIndex3i TriVIDs = TraceMesh.GetTriangle(AdjTID);
+												return FIndex2i( TriVIDs[IndexOf], TriVIDs[(IndexOf + 1) % 3]);
+											}();
+		
+		// translate to points on the surface of the host mesh
+		const FSurfacePoint TraceStartSurfacePoint = VIDToSurfacePoint(OrderedTraceEdgeV.A);
+		const FSurfacePoint TraceEndSurfacePoint   = VIDToSurfacePoint(OrderedTraceEdgeV.B);
+
+		
+		// add the first point in the trace.
+		EdgeTrace.Add(TraceStartSurfacePoint);
+
+		// compute (and add) the crossing locations relative to the edges of the host mesh
+		ConvertEdgesCrossed(TraceStartSurfacePoint, TraceEndSurfacePoint, HostEdgesCrossed, HostMesh, TraceMesh, CoalesceThreshold, VIDToSurfacePoint, EdgeTrace);
+		
+		// maybe add end point, if it doesn't duplicate the last one already
 		{
-			if (EdgeTrace.Num() == 0)
+			const FSurfacePoint& LastPoint = EdgeTrace.Last();
+			const bool bSameAsLast = IsVertexPoint(TraceEndSurfacePoint) 
+			                      && IsVertexPoint(LastPoint) 
+								  && (LastPoint.Position.VertexPosition.VID == TraceEndSurfacePoint.Position.VertexPosition.VID);
+			if (!bSameAsLast)
 			{
-				EdgeTrace.Emplace(VID);
+				EdgeTrace.Add(TraceEndSurfacePoint);
 			}
-			else
-			{
-				const FSurfacePoint& LastPoint = EdgeTrace.Last();
-				const bool bSameAsLast         = ( LastPoint.PositionType == FSurfacePoint::EPositionType::Vertex 
-				                                && LastPoint.Position.VertexPosition.VID == VID );
-				if (!bSameAsLast)
-				{
-					EdgeTrace.Emplace(VID);
-				}
-			}
-		};
-
-
-		// an edge  w/o crossings is the same as the host mesh edge. the trace in this case requires only the start and edge vertex 
-		if (HostXings.Num() == 0)
-		{
-			const int32 AdjTID     = TraceMesh.GetEdgeT(TraceEID).A;
-			const int32 IndexOf    = TraceMesh.GetTriEdges(AdjTID).IndexOf(TraceEID);
-			const FIndex3i TriVIDs = TraceMesh.GetTriangle(AdjTID);
-
-			const int32 StartVID = TriVIDs[IndexOf];
-			const int32 EndVID   = TriVIDs[(IndexOf + 1) % 3];
-
-
-			EdgeTrace.Emplace(StartVID);
-			EdgeTrace.Emplace(EndVID);
-		}
-		else
-		{
-			// knowing the host mesh edges this trace edge crosses, we make a 2d triangle strip by unfolding the 3d triangles the edge crosses
-			// and solve a 2x2 problem for each 2d edge crossed.
-
-			// struct holds bare-bones 2d triangle strip and ability to map vertexIDs from the src mesh
-			struct
-			{
-				TMap<int32, int32> ToStripIndexMap;    // maps from mesh VID to triangle strip VID.
-				TArray<FVector2d> StripVertexBuffer;
-			} TriangleStrip;
-			TriangleStrip.StripVertexBuffer.Reserve(3 + HostXings.Num());
-
-
-			// Functor to add a triangle to the triangle strip. This assumes that at least one shared edge of this triangle already exits in the strip
-			auto AddTriangleToStrip = [&HostMesh, &TriangleStrip](const int32 HostTID)
-			{
-
-				FIndex3i TriVIDs = HostMesh.GetTriangle(HostTID);
-
-				// two of the tri vertices are already in the buffer.  Identify the new one.
-				int32 IndexOf = -1;
-				for (int32 i = 0; i < 3; ++i)
-				{
-					const int32 VID       = TriVIDs[i];
-					const int32* StripVID = TriangleStrip.ToStripIndexMap.Find(VID);
-					if (!StripVID)
-					{
-						IndexOf = i;
-						break;
-					}
-				}
-
-				// no need to do anything if all vertices had previously been added, 
-				if (IndexOf == -1)
-				{
-					return;
-				}
-
-
-				const FVector3d Verts[3] = { HostMesh.GetVertex(TriVIDs[0]),  HostMesh.GetVertex(TriVIDs[1]),  HostMesh.GetVertex(TriVIDs[2]) };
-				// with new vert last (i.e. V2).
-				const FIndex3i Reordered((IndexOf + 1) % 3, (IndexOf + 2) % 3, IndexOf);
-				// the spanning vectors 
-				const FVector3d E1 = Verts[Reordered[1]] - Verts[Reordered[0]];
-				const FVector3d E2 = Verts[Reordered[2]] - Verts[Reordered[0]];
-
-				// coordinates of V2 relative to the direction of E1, and its orthogonal complement.
-				const double E1DotE2       = E2.Dot(E1);
-				const double E1LengthSqr   = FMath::Max(E1.SizeSquared(), TMathUtilConstants<double>::ZeroTolerance);
-				const double E1Length      = FMath::Sqrt(E1LengthSqr);
-				const double E1Dist        = E2.Dot(E1) / E1Length;
-				const double OrthE1DistSqr = FMath::Max(0., E2.SizeSquared() - E1DotE2 * E1DotE2 / E1LengthSqr);
-				const double OrthE1Dist    = FMath::Sqrt(OrthE1DistSqr);
-
-				// 2d version of E1
-				const int32* StripTriV0 = TriangleStrip.ToStripIndexMap.Find(TriVIDs[Reordered[0]]);
-				const int32* StripTriV1 = TriangleStrip.ToStripIndexMap.Find(TriVIDs[Reordered[1]]);
-				checkSlow(StripTriV0); checkSlow(StripTriV1);
-
-				const FVector2d& StripTriVert0 = TriangleStrip.StripVertexBuffer[*StripTriV0];
-				const FVector2d& StripTriVert1 = TriangleStrip.StripVertexBuffer[*StripTriV1];
-				const FVector2d StripE1        = (StripTriVert1 - StripTriVert0);
-
-				const FVector2d StripE1Perp(-StripE1[1], StripE1[0]); // Rotate StripE1 90 CCW.
-
-																	  // new vertex 2d position
-				const FVector2d StripTriVert2 = StripTriVert0 + (StripE1)*E1DotE2 / E1LengthSqr + (StripE1Perp / E1Length) * OrthE1Dist;
-				TriangleStrip.StripVertexBuffer.Add(StripTriVert2);
-				const int32 StripVID          = TriangleStrip.StripVertexBuffer.Num() - 1;
-				const int32 SurfaceVID        = TriVIDs[IndexOf];
-				TriangleStrip.ToStripIndexMap.Add(SurfaceVID, StripVID);
-			};
-
-			// find the VID where we start the edge trace and add it to the trace.
-			const int32 StartVID = [&]
-				{
-					// by convention the intrinsic edge direction will be defined by the first adj triangle
-					const int32 AdjTID     = TraceMesh.GetEdgeT(TraceEID).A;
-					const int32 IndexOf    = TraceMesh.GetTriEdges(AdjTID).IndexOf(TraceEID);
-					const FIndex3i TriVIDs = TraceMesh.GetTriangle(AdjTID);
-					return TriVIDs[IndexOf];
-				}();
-			EdgeTrace.Emplace(StartVID);
-
-			// find host triangle that contains both the StartVID and the first host edge we cross.
-			int32 HostTID = [&]
-				{
-					const int32 HostEID      = HostXings[0];
-					const FIndex2i HostEdgeT = HostMesh.GetEdgeT(HostEID);
-					const FIndex3i TriAVIDs  = HostMesh.GetTriangle(HostEdgeT.A);
-					const FIndex3i TriBVIDs  = HostMesh.GetTriangle(HostEdgeT.B);
-
-					return (TriAVIDs.IndexOf(StartVID) != -1) ? HostEdgeT.A : HostEdgeT.B;
-				}();
-
-			// jump-start the process of making the triangle strip by adding the first 2 verts.
-			{
-				const FIndex3i TriVIDs = HostMesh.GetTriangle(HostTID);
-				const int32 IndexOf    = TriVIDs.IndexOf(StartVID);
-				checkSlow(IndexOf != -1);
-				const int32 NextVID    = TriVIDs[(IndexOf + 1) % 3];
-
-				const FVector3d Vert0 = HostMesh.GetVertex(StartVID);
-				const FVector3d Vert1 = HostMesh.GetVertex(NextVID);
-
-				const double LSqr = FMath::Max((Vert0 - Vert1).SizeSquared(), TMathUtilConstants<double>::ZeroTolerance);
-				const double L    = FMath::Sqrt(LSqr);
-
-				const FVector2d StripVert0(0., 0.);
-				const FVector2d StripVert1(L, 0.);
-
-				TriangleStrip.StripVertexBuffer.Add(StripVert0);
-				TriangleStrip.ToStripIndexMap.Add(StartVID, 0);
-
-				TriangleStrip.StripVertexBuffer.Add(StripVert1);
-				TriangleStrip.ToStripIndexMap.Add(NextVID, 1);
-			}
-
-			// unfold the triangle strip, add the first triangle and then all subsequent ones.
-			AddTriangleToStrip(HostTID);
-			for (int32 HostEID : HostXings)
-			{
-				const FIndex2i EdgeT = HostMesh.GetEdgeT(HostEID);
-				HostTID              = (EdgeT.A == HostTID) ? EdgeT.B : EdgeT.A;
-				AddTriangleToStrip(HostTID);
-			}
-
-			// get the 2d location of the end of the trace edge (i.e. at EndVID ). Note: the start location is (0,0) in 2d 
-			const FVector2d StripEndVertex = [&]
-				{
-					const FIndex2i TraceEdgeV = TraceMesh.GetEdgeV(TraceEID);
-					const int32 EndVID        = (TraceEdgeV.A == StartVID) ? TraceEdgeV.B : TraceEdgeV.A;
-					const int32* StripEndVID  = TriangleStrip.ToStripIndexMap.Find(EndVID);
-					checkSlow(StripEndVID);
-					return TriangleStrip.StripVertexBuffer[*StripEndVID];
-				}();
-
-			// loop over the two-d versions of the host edges the path crosses and find the intersection.
-			for (int32 HostEID : HostXings)
-			{
-				const FIndex2i EdgeV = HostMesh.GetEdgeV(HostEID);
-				const int32* StripA  = TriangleStrip.ToStripIndexMap.Find(EdgeV.A);
-				const int32* StripB  = TriangleStrip.ToStripIndexMap.Find(EdgeV.B);
-				checkSlow(StripA); checkSlow(StripB);
-				const FIndex2i StripEdgeV(*StripA, *StripB);
-
-				const FVector2d StripVertA = TriangleStrip.StripVertexBuffer[StripEdgeV.A];
-				const FVector2d StripVertB = TriangleStrip.StripVertexBuffer[StripEdgeV.B];
-
-				// solve Alpha*StripVertA + (1-Alpha)StripVertB = Gamma Start + (1-Gamma)End.
-				// i.e. 
-				// Alpha * (StripVertA - StripVertB) + Gamma * (End - Start) =  End - StripVertB.
-
-				// write this as 2x2 matrix problem M.x = b solving for unknown vector x=(alpha, gamma).
-
-				// matrix
-				double m[2][2];
-				m[0][0] = (StripVertA - StripVertB)[0];  m[0][1] = StripEndVertex[0];
-				m[1][0] = (StripVertA - StripVertB)[1];  m[1][1] = StripEndVertex[1];
-
-				// b-vector
-				const FVector2d b = StripEndVertex - StripVertB;
-
-				const double Det = m[0][0] * m[1][1] - m[0][1] * m[1][0];
-				// inverse: not yet scaled by 1/det
-				double invm[2][2];
-				invm[0][0] =  m[1][1];   invm[0][1] = -m[0][1];
-				invm[1][0] = -m[1][0];   invm[1][1] =  m[0][0];
-
-				// solve for alpha only ( don't care about gamma ) 
-				double alpha = invm[0][0] * b[0] + invm[0][1] * b[1];
-
-				if (FMath::Abs(Det) < TMathUtilConstants<double>::ZeroTolerance)
-				{
-					// this surface edge and the intrinsic edge that intersects it are nearly parallel.  
-					// Assume the crossing happens at the farther vertex.
-					const double dA = StripVertA.SquaredLength();
-					const double dB = StripVertB.SquaredLength();
-
-					alpha = (dA > dB) ? 1. : 0.;
-				}
-				else
-				{
-					alpha = FMath::Clamp(alpha / Det, 0., 1.);
-				}
-
-				// may want to convert this edge crossing to a vertex crossing, if it is close enough.
-				int32 CoalesceVID = -1;
-				bool bSnapToVert = false;
-
-				if (alpha < CoalesceThreshold)
-				{
-					CoalesceVID = EdgeV.B;
-					bSnapToVert = true;
-				}
-				else if ((1. - alpha) < CoalesceThreshold)
-				{
-					CoalesceVID = EdgeV.A;
-					bSnapToVert = true;
-				}
-
-				if (bSnapToVert)
-				{
-					AddVIDToTrace(CoalesceVID);
-				}
-				else
-				{
-					EdgeTrace.Emplace(HostEID, alpha);
-				}
-			}
-
-			// Add the end vertex to the EdgeTrace
-			{
-				const FIndex2i TraceEdgeV = TraceMesh.GetEdgeV(TraceEID);
-				const int32 EndVID        = (TraceEdgeV.A == StartVID) ? TraceEdgeV.B : TraceEdgeV.A;
-				AddVIDToTrace(EndVID);
-			}
-
 		}
 
 		// correct trace results order so it is either EdgeV order or reversed as requested
-		const bool bHasEdgeVOrder = (TraceMesh.GetEdgeV(TraceEID).A == EdgeTrace[0].Position.VertexPosition.VID);
+		const bool bHasEdgeVOrder = (TraceMesh.GetEdgeV(TraceEID).A == OrderedTraceEdgeV.A);
 		const bool bNeedToReverse = (bHasEdgeVOrder && bReverse) || (!bHasEdgeVOrder && !bReverse);
 
 		if (bNeedToReverse)
@@ -1077,6 +1633,83 @@
 	}
 };
 
+namespace FNormalCoordIntrinsicTraceImpl
+{
+	using namespace IntrinsicCorrespondenceUtils;
+
+	template <typename IntrinsicMeshType>
+	TArray<FSurfacePoint> TraceEdge(const IntrinsicMeshType& IntrinsicMesh, int32 IntrinsicEID, double CoalesceThreshold, bool bReverse)
+	{
+		using FEdgeAndCrossingIdx = FNormalCoordinates::FEdgeAndCrossingIdx;
+		
+		TArray<FSurfacePoint> Result;
+
+		if (!IntrinsicMesh.IsEdge(IntrinsicEID))
+		{
+			return MoveTemp(Result);
+		}
+
+		const FNormalCoordinates& NormalCoordinates = IntrinsicMesh.GetNormalCoordinates();
+
+		TArray<int32> SurfXings;
+		const int32 NumSurfXings = NormalCoordinates.NumEdgeCrossing(IntrinsicEID);
+		if (NumSurfXings > 0)
+		{
+			SurfXings.Reserve(NumSurfXings);
+		}
+
+		const FDynamicMesh3& SurfaceMesh = *IntrinsicMesh.GetExtrinsicMesh();
+
+		const FIndex2i IntrinsicEdgeT = IntrinsicMesh.GetEdgeT(IntrinsicEID);
+
+		// need to create a list of surface mesh edges that cross this edge.  To do this we need to follow
+		// each curve that crosses this intrinsic edge to the vertices where it terminates ( thus identifying the surface edge )
+		for (int32 p = 1; p < NumSurfXings + 1; ++p)
+		{
+
+			const FIndex2i SurfaceEdgeV = [&]
+										{
+
+											TArray<FEdgeAndCrossingIdx> Crossings;
+											FIndex2i Verts;
+											// follow surface curve (i.e. surface edge) forward to the end and identify the vertex
+											{
+												Crossings.Add(FEdgeAndCrossingIdx({IntrinsicEdgeT.A, IntrinsicEID, p}));
+												FNormalCoordSurfaceTraceImpl::ContinueTraceSurfaceEdgeAcrossFaces(IntrinsicMesh, Crossings);
+												const FEdgeAndCrossingIdx& LastXing = Crossings.Last();
+
+												const FIndex3i TriEIDs = IntrinsicMesh.GetTriEdges(LastXing.TID);
+												Verts.A = IntrinsicMesh.GetTriangle(LastXing.TID)[LastXing.EID];
+											}
+											Crossings.Reset();
+
+											// follow surface curve (i.e. surface edge) backward to other end and identify the vertex
+											{
+												Crossings.Add(FEdgeAndCrossingIdx({IntrinsicEdgeT.B, IntrinsicEID, NumSurfXings + 1 - p}));
+												FNormalCoordSurfaceTraceImpl::ContinueTraceSurfaceEdgeAcrossFaces(IntrinsicMesh, Crossings);
+												const FEdgeAndCrossingIdx& LastXing = Crossings.Last();
+
+												const FIndex3i TriEIDs = IntrinsicMesh.GetTriEdges(LastXing.TID);
+												Verts.B = IntrinsicMesh.GetTriangle(LastXing.TID)[LastXing.EID];
+											}
+											return Verts;
+										}();
+
+			// identify the surface edge from its endpoints. 
+			const int32 SurfaceEID = SurfaceMesh.FindEdge(SurfaceEdgeV.A, SurfaceEdgeV.B);
+			checkSlow(SurfaceEID != IndexConstants::InvalidID);
+
+			SurfXings.Add(SurfaceEID);
+		}
+
+		// do the actual trace - this identifies the surface mesh triangles crossed by the intrinsic edge and unfolds them into a triangle strip where the trace is performed
+		Result = FNormalCoordSurfaceTraceImpl::TraceEdgeOverHost(IntrinsicEID, SurfXings, SurfaceMesh, IntrinsicMesh, CoalesceThreshold, bReverse, 
+																[&IntrinsicMesh](int32 VID) { return IntrinsicMesh.GetVertexSurfacePoint(VID); } );
+
+		return MoveTemp(Result);
+	}
+}
+
 int32 UE::Geometry::FlipToDelaunay(FIntrinsicTriangulation& IntrinsicMesh, TSet<int32>& Uncorrected, const int32 MaxFlipCount)
 {
 	return FlipToDelaunayImpl(IntrinsicMesh, Uncorrected, MaxFlipCount);
@@ -1087,6 +1720,16 @@
 	return FlipToDelaunayImpl(IntrinsicMesh, Uncorrected, MaxFlipCount);
 }
 
+int32 UE::Geometry::FlipToDelaunay(FSimpleIntrinsicMesh& IntrinsicMesh, TSet<int32>& Uncorrected, const int32 MaxFlipCount)
+{
+	return FlipToDelaunayImpl(IntrinsicMesh, Uncorrected, MaxFlipCount);
+}
+ 
+int32 UE::Geometry::FlipToDelaunay(FIntrinsicMesh& IntrinsicMesh, TSet<int32>& Uncorrected, const int32 MaxFlipCount)
+{
+	return FlipToDelaunayImpl(IntrinsicMesh, Uncorrected, MaxFlipCount);
+}
+
 int32 UE::Geometry::FlipToDelaunay(FIntrinsicEdgeFlipMesh& IntrinsicMesh, TSet<int32>& Uncorrected, const int32 MaxFlipCount)
 {
 	return FlipToDelaunayImpl(IntrinsicMesh, Uncorrected, MaxFlipCount);
@@ -1443,6 +2086,414 @@
 	return Result;
 }
 
+/**------------------------------------------------------------------------------
+*  FSimpleIntrinsicMesh Methods
+*-------------------------------------------------------------------------------*/
+
+UE::Geometry::EMeshResult UE::Geometry::FSimpleIntrinsicMesh::PokeTriangleTopology(int32 TriangleID, FPokeTriangleInfo& PokeResult)
+{
+	PokeResult = FPokeTriangleInfo();
+
+	if (!IsTriangle(TriangleID))
+	{
+		return EMeshResult::Failed_NotATriangle;
+	}
+
+	FIndex3i tv = GetTriangle(TriangleID);
+	FIndex3i te = GetTriEdges(TriangleID);
+
+	// create vertex with averaged position.. 
+	FVector3d vPos = (1. / 3.) * (GetVertex(tv[0]) + GetVertex(tv[1]) + GetVertex(tv[2]));
+
+	int32 newVertID = AppendVertex(vPos);
+
+
+	// add in edges to center vtx, do not connect to triangles yet
+	int32 eAN = AddEdgeInternal(tv[0], newVertID, -1, -1);
+	int32 eBN = AddEdgeInternal(tv[1], newVertID, -1, -1);
+	int32 eCN = AddEdgeInternal(tv[2], newVertID, -1, -1);
+	VertexRefCounts.Increment(tv[0]);
+	VertexRefCounts.Increment(tv[1]);
+	VertexRefCounts.Increment(tv[2]);
+	VertexRefCounts.Increment(newVertID, 3);
+
+	// old triangle becomes tri along first edge
+	Triangles[TriangleID] = FIndex3i(tv[0], tv[1], newVertID);
+	TriangleEdges[TriangleID] = FIndex3i(te[0], eBN, eAN);
+
+	// add two triangles
+	int32 t1 = AddTriangleInternal(tv[1], tv[2], newVertID, te[1], eCN, eBN);
+	int32 t2 = AddTriangleInternal(tv[2], tv[0], newVertID, te[2], eAN, eCN);
+
+	// second and third edges of original tri have neighbors
+	ReplaceEdgeTriangle(te[1], TriangleID, t1);
+	ReplaceEdgeTriangle(te[2], TriangleID, t2);
+
+	// set the triangles for the edges we created above
+	SetEdgeTrianglesInternal(eAN, TriangleID, t2);
+	SetEdgeTrianglesInternal(eBN, TriangleID, t1);
+	SetEdgeTrianglesInternal(eCN, t1, t2);
+
+
+	PokeResult.OriginalTriangle = TriangleID;
+	PokeResult.TriVertices = tv;
+	PokeResult.NewVertex = newVertID;
+	PokeResult.NewTriangles = FIndex2i(t1, t2);
+	PokeResult.NewEdges = FIndex3i(eAN, eBN, eCN);
+	PokeResult.BaryCoords = FVector3d(1. / 3., 1. / 3., 1. / 3.);
+
+	return EMeshResult::Ok;
+}
+
+UE::Geometry::EMeshResult  UE::Geometry::FSimpleIntrinsicMesh::SplitEdgeTopology(int32 eab, FEdgeSplitInfo& SplitInfo)
+{
+	SplitInfo = FEdgeSplitInfo();
+
+	if (!IsEdge(eab))
+	{
+		return EMeshResult::Failed_NotAnEdge;
+	}
+
+	// look up primary edge & triangle
+	const FEdge Edge = Edges[eab];
+	const int32 t0 = Edge.Tri[0];
+	if (t0 == InvalidID)
+	{
+		return EMeshResult::Failed_BrokenTopology;
+	}
+	const FIndex3i T0tv = GetTriangle(t0);
+	const FIndex3i T0te = GetTriEdges(t0);
+	const int32 IndexOfe = T0te.IndexOf(eab);
+	const int32 a = T0tv[IndexOfe];
+	const int32 b = T0tv[AddOneModThree[IndexOfe]];
+	const int32 c = T0tv[AddTwoModThree[IndexOfe]];
+
+	// look up edge bc, which needs to be modified
+	const int32 ebc = T0te[AddOneModThree[IndexOfe]];
+
+	// RefCount overflow check. Conservatively leave room for
+	// extra increments from other operations.
+	if (VertexRefCounts.GetRawRefCount(c) > FRefCountVector::INVALID_REF_COUNT - 3)
+	{
+		return EMeshResult::Failed_HitValenceLimit;
+	}
+
+
+	SplitInfo.OriginalEdge      = eab;
+	SplitInfo.OriginalVertices  = FIndex2i(a, b);   // this is the oriented a,b
+	SplitInfo.OriginalTriangles = FIndex2i(t0, InvalidID);
+	SplitInfo.SplitT            = 0.5;
+
+	// quite a bit of code is duplicated between boundary and non-boundary case, but it
+	//  is too hard to follow later if we factor it out...
+	if (IsBoundaryEdge(eab))
+	{
+		// create vertex
+		const FVector3d vNew = 0.5 * (GetVertex(a) + GetVertex(b));
+		const int32 f = AppendVertex(vNew);
+
+
+		// rewrite existing triangle
+		Triangles[t0][AddOneModThree[IndexOfe]] = f;
+
+		// add second triangle
+		const int32 t2 = AddTriangleInternal(f, b, c, InvalidID, InvalidID, InvalidID);
+
+
+		// rewrite edge bc, create edge af
+		ReplaceEdgeTriangle(ebc, t0, t2);
+		const int32 eaf = eab;
+		Edges[eaf].Vert = FIndex2i(FMath::Min(a, f), FMath::Max(a, f));
+		//ReplaceEdgeVertex(eaf, b, f);
+		if (a != b)
+		{
+			VertexEdgeLists.Remove(b, eab);
+		}
+		VertexEdgeLists.Insert(f, eaf);
+
+		// create edges fb and fc
+		const int32 efb = AddEdgeInternal(f, b, t2);
+		const int32 efc = AddEdgeInternal(f, c, t0, t2);
+
+		// update triangle edge-nbrs
+		ReplaceTriangleEdge(t0, ebc, efc);
+		TriangleEdges[t2] = FIndex3i(efb, ebc, efc);
+
+		// update vertex refcounts
+		VertexRefCounts.Increment(c);
+		VertexRefCounts.Increment(f, 2);
+
+		SplitInfo.bIsBoundary   = true;
+		SplitInfo.OtherVertices = FIndex2i(c, InvalidID);
+		SplitInfo.NewVertex     = f;
+		SplitInfo.NewEdges      = FIndex3i(efb, efc, InvalidID);
+		SplitInfo.NewTriangles  = FIndex2i(t2, InvalidID);
+
+		return EMeshResult::Ok;
+
+	}
+	else 		// interior triangle branch
+	{
+		// look up other triangle
+		const int32 t1 = Edges[eab].Tri[1];
+		SplitInfo.OriginalTriangles.B = t1;
+		const FIndex3i T1tv = GetTriangle(t1);
+		const FIndex3i T1te = GetTriEdges(t1);
+		const int32 T1IndexOfe = T1te.IndexOf(eab);
+		checkSlow(T1tv[T1IndexOfe] == b);
+
+		const int32 d = T1tv[AddTwoModThree[T1IndexOfe]];
+		const int32 edb = T1te[AddTwoModThree[T1IndexOfe]];
+
+		// RefCount overflow check. Conservatively leave room for
+		// extra increments from other operations.
+		if (VertexRefCounts.GetRawRefCount(d) > FRefCountVector::INVALID_REF_COUNT - 3)
+		{
+			return EMeshResult::Failed_HitValenceLimit;
+		}
+
+		// create vertex
+		FVector3d vNew = 0.5 * (GetVertex(a) + GetVertex(b));
+		int32 f = AppendVertex(vNew);
+
+
+		// rewrite existing triangles, replacing b with f
+		Triangles[t0][AddOneModThree[IndexOfe]] = f;
+		Triangles[t1][T1IndexOfe] = f;
+
+
+		// add two triangles to close holes we just created
+		int32 t2 = AddTriangleInternal(f, b, c, InvalidID, InvalidID, InvalidID);
+		int32 t3 = AddTriangleInternal(f, d, b, InvalidID, InvalidID, InvalidID);
+
+
+		// update the edges we found above, to point to triangles
+		ReplaceEdgeTriangle(ebc, t0, t2);
+		ReplaceEdgeTriangle(edb, t1, t3);
+
+		// edge eab became eaf
+		int32 eaf = eab; //Edge * eAF = eAB;
+		Edges[eaf].Vert = FIndex2i(FMath::Min(a, f), FMath::Max(a, f));
+
+		// update a/b/f vertex-edges
+		if (a != b)
+		{
+			VertexEdgeLists.Remove(b, eab);
+		}
+		VertexEdgeLists.Insert(f, eaf);
+
+		// create edges connected to f  (also updates vertex-edges)
+		int32 efb = AddEdgeInternal(f, b, t2, t3);
+		int32 efc = AddEdgeInternal(f, c, t0, t2);
+		int32 edf = AddEdgeInternal(d, f, t1, t3);
+
+		// update triangle edge-nbrs
+		ReplaceTriangleEdge(t0, ebc, efc);
+		ReplaceTriangleEdge(t1, edb, edf);
+		TriangleEdges[t2] = FIndex3i(efb, ebc, efc);
+		TriangleEdges[t3] = FIndex3i(edf, edb, efb);
+
+		// update vertex refcounts
+		VertexRefCounts.Increment(c);
+		VertexRefCounts.Increment(d);
+		VertexRefCounts.Increment(f, 4);
+
+		SplitInfo.bIsBoundary   = false;
+		SplitInfo.OtherVertices = FIndex2i(c, d);
+		SplitInfo.NewVertex     = f;
+		SplitInfo.NewEdges      = FIndex3i(efb, efc, edf);
+		SplitInfo.NewTriangles  = FIndex2i(t2, t3);
+
+		return EMeshResult::Ok;
+	}
+}
+
+
+UE::Geometry::EMeshResult UE::Geometry::FSimpleIntrinsicMesh::PokeTriangle(int32 TID, const FVector3d& BaryCoordinates, FPokeTriangleInfo& PokeInfo)
+{
+	if (!IsTriangle(TID))
+	{
+		return EMeshResult::Failed_NotATriangle;
+	}
+
+	// state before the poke
+	const FIndex3i OriginalVIDs            = GetTriangle(TID);
+	const FIndex3i OriginalEdges           = GetTriEdges(TID);
+	const FVector3d OriginalTriEdgeLengths = GetTriEdgeLengths(TID);
+
+	// Add a new vertex and faces to the IntrinsicMesh.   
+	// Note: the r3 position will be wrong initially since this poke will just interpolate the corners of the tri.
+	//       we fix this position as the last step in this function
+	EMeshResult PokeResult = PokeTriangleTopology(TID, PokeInfo);
+	if (PokeResult != EMeshResult::Ok)
+	{
+		return PokeResult;
+	}
+
+
+	FIndex3i NewTris(TID, PokeInfo.NewTriangles[0], PokeInfo.NewTriangles[1]);
+	const int32 NewVID = PokeInfo.NewVertex;
+
+	// Need to update intrinsic information for the 3 triangles that resulted from the poke
+	// 1) update the edge lengths for the triangles
+	// 2) update the internal angles for the triangles
+
+
+	// (1) compute intrinsic edge lengths for the 3 new edges
+	{
+		FVector3d DistancesFromOldCorner;
+		for (int32 v = 0; v < 3; ++v)
+		{
+			FVector3d BaryCoordVertex(0., 0., 0.); BaryCoordVertex[v] = 1.;
+			const double Distance = DistanceBetweenBarycentricPoints( OriginalTriEdgeLengths[0],
+																	  OriginalTriEdgeLengths[1],
+																	  OriginalTriEdgeLengths[2],
+																	  BaryCoordVertex, BaryCoordinates);
+			DistancesFromOldCorner[v] = Distance;
+		}
+
+		{
+			// original tri is and verts (a, b,c) and edges (ab, bc, ca)
+			// the updated tri has verts (a, b, new) and edges (ab, bnew, newa)
+			FIndex3i T0EIDs = GetTriEdges(NewTris[0]);
+			EdgeLengths.InsertAt(DistancesFromOldCorner[1], T0EIDs[1]);
+			EdgeLengths.InsertAt(DistancesFromOldCorner[0], T0EIDs[2]);
+
+			// new tri[0] has verts (b, c, new) and edges (bc, cnew, newb)
+			FIndex3i T1EIDs = GetTriEdges(PokeInfo.NewTriangles[0]);
+			EdgeLengths.InsertAt(DistancesFromOldCorner[2], T1EIDs[1]);
+		}
+	}
+	// (2) update internal angles for the triangles.
+	{
+		InternalAngles.InsertAt(ComputeTriInternalAnglesR(NewTris[2]), NewTris[2]);
+		InternalAngles.InsertAt(ComputeTriInternalAnglesR(NewTris[1]), NewTris[1]);
+		InternalAngles[NewTris[0]] = ComputeTriInternalAnglesR(NewTris[0]);
+	}
+
+	// record the coordinate actually used. 
+	PokeInfo.BaryCoords = BaryCoordinates;
+
+	return EMeshResult::Ok;
+}
+
+
+UE::Geometry::EMeshResult UE::Geometry::FSimpleIntrinsicMesh::SplitEdge(int32 EdgeAB, FEdgeSplitInfo& SplitInfo, double SplitParameterT)
+{
+	SplitParameterT = FMath::Clamp(SplitParameterT, 0., 1.);
+
+	if (!IsEdge(EdgeAB))
+	{
+		return EMeshResult::Failed_NotAnEdge;
+	}
+
+
+	const FEdge OriginalEdge = GetEdge(EdgeAB);
+
+	if (IsBoundaryEdge(EdgeAB))
+	{
+		// state before the split
+		const int32 TID      = OriginalEdge.Tri[0];
+		const int32 IndexOfe = GetTriEdges(TID).IndexOf(EdgeAB);
+
+		// Info about the original T0 tri, reordered to make the split edge the first edge..
+		const FVector3d OriginalT0EdgeLengths = Permute(IndexOfe, GetTriEdgeLengths(TID));    // as (|ab|, |bc|, |ca|)
+
+
+		// Update the connectivity with the edge split
+		EMeshResult Result = SplitEdgeTopology(EdgeAB, SplitInfo);
+
+		if (Result != EMeshResult::Ok)
+		{
+			return Result;
+		}
+		const int32 NewVID = SplitInfo.NewVertex;
+		const int32 NewTID = SplitInfo.NewTriangles[0]; // edges {fb, bc, cf}
+		const int32 NewEdgeFB = SplitInfo.NewEdges[0];
+		const int32 NewEdgeFC = SplitInfo.NewEdges[1];
+
+		const double DistAF = SplitParameterT * OriginalT0EdgeLengths[0];
+		const double DistFB = FMath::Max(0., OriginalT0EdgeLengths[0] - DistAF);
+		// new edge length
+		const double DistFC = DistanceBetweenBarycentricPoints( OriginalT0EdgeLengths[0],
+																OriginalT0EdgeLengths[1],
+																OriginalT0EdgeLengths[2],
+																FVector3d(SplitParameterT, 1. - SplitParameterT, 0.),
+																FVector3d(0., 0., 1.));
+
+		EdgeLengths.InsertAt(DistFC, NewEdgeFC);
+		EdgeLengths.InsertAt(DistFB, NewEdgeFB);
+		EdgeLengths[EdgeAB] = DistAF;
+
+		// update the internal angles
+		InternalAngles.InsertAt(ComputeTriInternalAnglesR(NewTID), NewTID);
+		InternalAngles[TID] = ComputeTriInternalAnglesR(TID);
+	
+		return Result;
+	}
+	else
+	{
+		// state before the split
+		const int32 TID      = OriginalEdge.Tri[0];
+		const int32 IndexOfe = GetTriEdges(TID).IndexOf(EdgeAB);
+
+		// Info about the original T0 tri, reordered to make the split edge the first edge..
+		const FVector3d OriginalT0EdgeLengths = Permute(IndexOfe, GetTriEdgeLengths(TID));    // as (|ab|, |bc|, |ca|)
+		
+
+	    // Info about the original T1 tri, reordered to make the split edge the first edge..
+		const int32 TID1       = OriginalEdge.Tri[1];
+		const int32 IndexOfT1e = GetTriEdges(TID1).IndexOf(EdgeAB);
+
+		const FVector3d OriginalT1EdgeLengths = Permute(IndexOfT1e, GetTriEdgeLengths(TID1));   // as (|ba|, |ad|, |db})
+
+
+		// Update the connectivity with the edge split
+		EMeshResult Result = SplitEdgeTopology(EdgeAB, SplitInfo);
+
+		if (Result != EMeshResult::Ok)
+		{
+			return Result;
+		}
+		const int32 NewVID = SplitInfo.NewVertex;
+		const int32 NewTID0 = SplitInfo.NewTriangles[0]; // edges {fb, bc, cf}
+		const int32 NewTID1 = SplitInfo.NewTriangles[1]; // edges {fd, db, bc}
+		const int32 EdgeAF = EdgeAB;
+		const int32 NewEdgeFB = SplitInfo.NewEdges[0];
+		const int32 NewEdgeFC = SplitInfo.NewEdges[1];
+		const int32 NewEdgeFD = SplitInfo.NewEdges[2];
+
+		const double DistAF = SplitParameterT * OriginalT0EdgeLengths[0];
+		const double DistFB = FMath::Max(0., OriginalT0EdgeLengths[0] - DistAF);
+		// new edge length
+		const double DistFC = DistanceBetweenBarycentricPoints( OriginalT0EdgeLengths[0],
+																OriginalT0EdgeLengths[1],
+																OriginalT0EdgeLengths[2],
+																FVector3d(SplitParameterT, 1. - SplitParameterT, 0.),
+																FVector3d(0., 0., 1.));
+
+		// new edge length
+		const double DistFD = DistanceBetweenBarycentricPoints( OriginalT1EdgeLengths[0],
+																OriginalT1EdgeLengths[1],
+																OriginalT1EdgeLengths[2],
+																FVector3d(1. - SplitParameterT, SplitParameterT, 0.),
+																FVector3d(0., 0., 1.));
+
+		EdgeLengths.InsertAt(DistFD, NewEdgeFD);
+		EdgeLengths.InsertAt(DistFC, NewEdgeFC);
+		EdgeLengths.InsertAt(DistFB, NewEdgeFB);
+		EdgeLengths[EdgeAF] = DistAF;
+
+		// update the internal angles (this uses edge lengths)
+		InternalAngles.InsertAt(ComputeTriInternalAnglesR(NewTID1), NewTID1);
+		InternalAngles.InsertAt(ComputeTriInternalAnglesR(NewTID0), NewTID0);
+		InternalAngles[TID] = ComputeTriInternalAnglesR(TID);
+		InternalAngles[TID1] = ComputeTriInternalAnglesR(TID1);
+
+		return Result;
+	}
+}
+
 
 /**------------------------------------------------------------------------------
 *  FIntrinsicTriangulation Methods
@@ -1669,61 +2720,6 @@
 }
 
 
-UE::Geometry::EMeshResult UE::Geometry::FIntrinsicTriangulation::PokeTriangleTopology(int32 TriangleID, FPokeTriangleInfo& PokeResult)
-{
-	PokeResult = FPokeTriangleInfo();
-
-	if (!IsTriangle(TriangleID))
-	{
-		return EMeshResult::Failed_NotATriangle;
-	}
-
-	FIndex3i tv = GetTriangle(TriangleID);
-	FIndex3i te = GetTriEdges(TriangleID);
-
-	// create vertex with averaged position.. 
-	FVector3d vPos = (1./3.)*(GetVertex(tv[0]) + GetVertex(tv[1]) + GetVertex(tv[2]));
-
-	int32 newVertID = AppendVertex(vPos);
-
-
-	// add in edges to center vtx, do not connect to triangles yet
-	int32 eAN = AddEdgeInternal(tv[0], newVertID, -1, -1);
-	int32 eBN = AddEdgeInternal(tv[1], newVertID, -1, -1);
-	int32 eCN = AddEdgeInternal(tv[2], newVertID, -1, -1);
-	VertexRefCounts.Increment(tv[0]);
-	VertexRefCounts.Increment(tv[1]);
-	VertexRefCounts.Increment(tv[2]);
-	VertexRefCounts.Increment(newVertID, 3);
-
-	// old triangle becomes tri along first edge
-	Triangles[TriangleID] = FIndex3i(tv[0], tv[1], newVertID);
-	TriangleEdges[TriangleID] = FIndex3i(te[0], eBN, eAN);
-
-	// add two triangles
-	int32 t1 = AddTriangleInternal(tv[1], tv[2], newVertID, te[1], eCN, eBN);
-	int32 t2 = AddTriangleInternal(tv[2], tv[0], newVertID, te[2], eAN, eCN);
-
-	// second and third edges of original tri have neighbors
-	ReplaceEdgeTriangle(te[1], TriangleID, t1);
-	ReplaceEdgeTriangle(te[2], TriangleID, t2);
-
-	// set the triangles for the edges we created above
-	SetEdgeTrianglesInternal(eAN, TriangleID, t2);
-	SetEdgeTrianglesInternal(eBN, TriangleID, t1);
-	SetEdgeTrianglesInternal(eCN, t1, t2);
-
-
-	PokeResult.OriginalTriangle = TriangleID;
-	PokeResult.TriVertices = tv;
-	PokeResult.NewVertex = newVertID;
-	PokeResult.NewTriangles = FIndex2i(t1, t2);
-	PokeResult.NewEdges = FIndex3i(eAN, eBN, eCN);
-	PokeResult.BaryCoords = FVector3d(1./3., 1./3., 1./3.);
-	
-	return EMeshResult::Ok;
-}
-
 UE::Geometry::EMeshResult UE::Geometry::FIntrinsicTriangulation::PokeTriangle(int32 TID, const FVector3d& BaryCoordinates, FPokeTriangleInfo& PokeInfo)
 {
 	if (!IsTriangle(TID))
@@ -1733,14 +2729,12 @@
 
 	// state before the poke
 	const FIndex3i OriginalVIDs = GetTriangle(TID);
-	const FIndex3i OriginalEdges = GetTriEdges(TID);
-	const FVector3d OriginalTriEdgeLengths = GetTriEdgeLengths(TID);
 	const FVector3d OriginalEdgeDirs = SignpostData.IntrinsicEdgeAngles[TID];
 	
 	// Add a new vertex and faces to the IntrinsicMesh.   
-	// Note: the r3 position will be wrong initially since this poke will just interpolate the corners of the tri.
-	//       we fix this position as the last step in this function
-	EMeshResult PokeResult = PokeTriangleTopology(TID, PokeInfo);
+	// this operation will update the internal angles and the edge lengths, but the actual R3 positon 
+	// of the vertex will need to be updated as will any signpost data.
+	EMeshResult PokeResult = MyBase::PokeTriangle(TID, BaryCoordinates, PokeInfo);
 	if (PokeResult != EMeshResult::Ok)
 	{
 		return PokeResult;
@@ -1751,67 +2745,34 @@
 	const int32 NewVID = PokeInfo.NewVertex;
 
 	// Need to update intrinsic information for the 3 triangles that resulted from the poke
-	// 1) update the edge lengths for the triangles
-	// 2) update the internal angles for the triangles
-	// 3) update the edge directions for the triangles. 
-	// 4) update the surface point for the new vertex by tracing one of the new edges. 
+	// 1) update Signpost data ( the edge directions for the triangles )
+	// 2) update the surface point for the new vertex by tracing one of the new edges. 
 	//    also update the edge directions leaving the new vertex to be relative to the direction defined on the extrinsic mesh
-
-
-	// (1) compute intrinsic edge lengths for the 3 new edges
-	{
-		FVector3d DistancesFromOldCorner;
-		for (int32 v = 0; v < 3; ++v)
-		{
-			FVector3d BaryCoordVertex(0., 0., 0.); BaryCoordVertex[v] = 1.;
-			const double Distance = DistanceBetweenBarycentricPoints( OriginalTriEdgeLengths[0],
-							  										  OriginalTriEdgeLengths[1],
-																	  OriginalTriEdgeLengths[2],
-																	  BaryCoordVertex, BaryCoordinates);
-			DistancesFromOldCorner[v] = Distance;
-		}
 	
-		{
-			// original tri is and verts (a, b,c) and edges (ab, bc, ca)
-			// the updated tri has verts (a, b, new) and edges (ab, bnew, newa)
-			FIndex3i T0EIDs = GetTriEdges(NewTris[0]);
-			EdgeLengths.InsertAt(DistancesFromOldCorner[1], T0EIDs[1]);
-			EdgeLengths.InsertAt(DistancesFromOldCorner[0], T0EIDs[2]);
-
-			// new tri[0] has verts (b, c, new) and edges (bc, cnew, newb)
-			FIndex3i T1EIDs = GetTriEdges(PokeInfo.NewTriangles[0]);
-			EdgeLengths.InsertAt(DistancesFromOldCorner[2], T1EIDs[1]);
-		}
-	}
-	// (2) update internal angles for the triangles.
-	{
-		InternalAngles.InsertAt(ComputeTriInternalAnglesR(NewTris[2]), NewTris[2]);
-		InternalAngles.InsertAt(ComputeTriInternalAnglesR(NewTris[1]), NewTris[1]);
-		InternalAngles[NewTris[0]] = ComputeTriInternalAnglesR(NewTris[0]);
-	}
-	
-	// (3) update the edge directions 
+	// (1) update the edge directions 
 	FVector3d TriEdgeDir[3]; // one for each tri in order (TID, NewTris[0], NewTris[1] )
-	// edges around the boundary of the original tri
-	TriEdgeDir[0][0] = OriginalEdgeDirs[0]; // AtoB dir
-	TriEdgeDir[1][0] = OriginalEdgeDirs[1]; // BtoC dir
-	TriEdgeDir[2][0] = OriginalEdgeDirs[2]; // CtoA dir
+	{
+		// edges around the boundary of the original tri
+		TriEdgeDir[0][0] = OriginalEdgeDirs[0]; // AtoB dir
+		TriEdgeDir[1][0] = OriginalEdgeDirs[1]; // BtoC dir
+		TriEdgeDir[2][0] = OriginalEdgeDirs[2]; // CtoA dir
 	
-	// edges from the corners of the original tri towards the new vertex at the "center"
-	const double ToRadiansAtA = SignpostData.GeometricVertexInfo[OriginalVIDs[0]].ToRadians;
-	const double ToRadiansAtB = SignpostData.GeometricVertexInfo[OriginalVIDs[1]].ToRadians;
-	const double ToRadiansAtC = SignpostData. GeometricVertexInfo[OriginalVIDs[2]].ToRadians;
-	TriEdgeDir[0][1] = AsZeroToTwoPi( OriginalEdgeDirs[1] + ToRadiansAtB * InternalAngles[NewTris[1]][0]);  // BtoNew dir
-	TriEdgeDir[1][1] = AsZeroToTwoPi( OriginalEdgeDirs[2] + ToRadiansAtC * InternalAngles[NewTris[2]][0]);  // CtoNew dir
-	TriEdgeDir[2][1] = AsZeroToTwoPi( OriginalEdgeDirs[0] + ToRadiansAtA * InternalAngles[NewTris[0]][0]);  // AtoNew dir   
+		// edges from the corners of the original tri towards the new vertex at the "center"
+		const double ToRadiansAtA = SignpostData.GeometricVertexInfo[OriginalVIDs[0]].ToRadians;
+		const double ToRadiansAtB = SignpostData.GeometricVertexInfo[OriginalVIDs[1]].ToRadians;
+		const double ToRadiansAtC = SignpostData. GeometricVertexInfo[OriginalVIDs[2]].ToRadians;
+		TriEdgeDir[0][1] = AsZeroToTwoPi( OriginalEdgeDirs[1] + ToRadiansAtB * InternalAngles[NewTris[1]][0]);  // BtoNew dir
+		TriEdgeDir[1][1] = AsZeroToTwoPi( OriginalEdgeDirs[2] + ToRadiansAtC * InternalAngles[NewTris[2]][0]);  // CtoNew dir
+		TriEdgeDir[2][1] = AsZeroToTwoPi( OriginalEdgeDirs[0] + ToRadiansAtA * InternalAngles[NewTris[0]][0]);  // AtoNew dir   
 
-	// edges from new vertex to a, to b, and to c, using new-to-A as the zero direction.  These will be updated later when we learn the correct angle for new-to-A
-	TriEdgeDir[0][2] = 0.;                                                                // NewToA dir 
-	TriEdgeDir[1][2] = InternalAngles[NewTris[0]][2];                                     // NewToB dir
-	TriEdgeDir[2][2] = InternalAngles[NewTris[0]][2] + InternalAngles[NewTris[1]][2];     // NewToC dir
+		// edges from new vertex to a, to b, and to c, using new-to-A as the zero direction.  These will be updated later when we learn the correct angle for new-to-A
+		TriEdgeDir[0][2] = 0.;                                                                // NewToA dir 
+		TriEdgeDir[1][2] = InternalAngles[NewTris[0]][2];                                     // NewToB dir
+		TriEdgeDir[2][2] = InternalAngles[NewTris[0]][2] + InternalAngles[NewTris[1]][2];     // NewToC dir
+	}
 	SignpostData.GeometricVertexInfo.InsertAt(FSignpost::FGeometricInfo(), NewVID); // we want the default of false, and 1.
 
-	// (4) update the surface point for the new vertex by tracing one of the new edges. 
+	// (2) update the surface point for the new vertex by tracing one of the new edges. 
 	// 
 	// --- fix the r3 position of the new vertex and compute its SurfacePoint Attributes by doing a trace on the Extrinsic Mesh along one of the edges incident on NewVID
 
@@ -1843,171 +2804,6 @@
 }
 
 
-EMeshResult UE::Geometry::FIntrinsicTriangulation::SplitEdgeTopology(int32 eab, FEdgeSplitInfo& SplitInfo)
-{
-	SplitInfo = FEdgeSplitInfo();
-
-	if (!IsEdge(eab))
-	{
-		return EMeshResult::Failed_NotAnEdge;
-	}
-
-	// look up primary edge & triangle
-	const FEdge Edge = Edges[eab];
-	const int32 t0 = Edge.Tri[0];
-	if (t0 == InvalidID)
-	{
-		return EMeshResult::Failed_BrokenTopology;
-	}
-	const FIndex3i T0tv = GetTriangle(t0);
-	const FIndex3i T0te = GetTriEdges(t0);
-	const int32 IndexOfe = T0te.IndexOf(eab);
-	const int32 a = T0tv[IndexOfe];
-	const int32 b = T0tv[AddOneModThree[IndexOfe]];
-	const int32 c = T0tv[AddTwoModThree[IndexOfe]];
-
-	// look up edge bc, which needs to be modified
-	const int32 ebc = T0te[AddOneModThree[IndexOfe]];
-
-	// RefCount overflow check. Conservatively leave room for
-	// extra increments from other operations.
-	if (VertexRefCounts.GetRawRefCount(c) > FRefCountVector::INVALID_REF_COUNT - 3)
-	{
-		return EMeshResult::Failed_HitValenceLimit;
-	}
-	
-
-	SplitInfo.OriginalEdge = eab;
-	SplitInfo.OriginalVertices = FIndex2i(a, b);   // this is the oriented a,b
-	SplitInfo.OriginalTriangles = FIndex2i(t0, InvalidID);
-	SplitInfo.SplitT = 0.5;
-
-	// quite a bit of code is duplicated between boundary and non-boundary case, but it
-	//  is too hard to follow later if we factor it out...
-	if (IsBoundaryEdge(eab))
-	{
-		// create vertex
-		const FVector3d vNew = 0.5 * (GetVertex(a) + GetVertex(b));
-		const int32 f = AppendVertex(vNew);
-		
-
-		// rewrite existing triangle
-		Triangles[t0][AddOneModThree[IndexOfe]] =  f;
-
-		// add second triangle
-		const int32 t2 = AddTriangleInternal(f, b, c, InvalidID, InvalidID, InvalidID);
-	
-
-		// rewrite edge bc, create edge af
-		ReplaceEdgeTriangle(ebc, t0, t2);
-		const int32 eaf = eab;
-		Edges[eaf].Vert = FIndex2i(FMath::Min(a,f), FMath::Max(a,f));
-		//ReplaceEdgeVertex(eaf, b, f);
-		if (a != b)
-		{ 
-			VertexEdgeLists.Remove(b, eab);
-		}
-		VertexEdgeLists.Insert(f, eaf);
-
-		// create edges fb and fc
-		const int32 efb = AddEdgeInternal(f, b, t2);
-		const int32 efc = AddEdgeInternal(f, c, t0, t2);
-
-		// update triangle edge-nbrs
-		ReplaceTriangleEdge(t0, ebc, efc);
-		TriangleEdges[t2] = FIndex3i(efb, ebc, efc);
-
-		// update vertex refcounts
-		VertexRefCounts.Increment(c);
-		VertexRefCounts.Increment(f, 2);
-
-		SplitInfo.bIsBoundary = true;
-		SplitInfo.OtherVertices = FIndex2i(c, InvalidID);
-		SplitInfo.NewVertex = f;
-		SplitInfo.NewEdges = FIndex3i(efb, efc, InvalidID);
-		SplitInfo.NewTriangles = FIndex2i(t2, InvalidID);
-
-		return EMeshResult::Ok;
-
-	}
-	else 		// interior triangle branch
-	{
-		// look up other triangle
-		const int32 t1 = Edges[eab].Tri[1];
-		SplitInfo.OriginalTriangles.B = t1;
-		const FIndex3i T1tv = GetTriangle(t1);
-		const FIndex3i T1te = GetTriEdges(t1);
-		const int32 T1IndexOfe = T1te.IndexOf(eab);
-		checkSlow(T1tv[T1IndexOfe] == b);
-
-		const int32 d = T1tv[AddTwoModThree[T1IndexOfe]];
-		const int32 edb = T1te[AddTwoModThree[T1IndexOfe]];
-
-		// RefCount overflow check. Conservatively leave room for
-		// extra increments from other operations.
-		if (VertexRefCounts.GetRawRefCount(d) > FRefCountVector::INVALID_REF_COUNT - 3)
-		{
-			return EMeshResult::Failed_HitValenceLimit;
-		}
-
-		// create vertex
-		FVector3d vNew = 0.5 * (GetVertex(a) + GetVertex(b));
-		int32 f = AppendVertex(vNew);
-
-
-		// rewrite existing triangles, replacing b with f
-		Triangles[t0][AddOneModThree[IndexOfe]] = f;
-		Triangles[t1][T1IndexOfe] = f;
-	
-
-		// add two triangles to close holes we just created
-		int32 t2 = AddTriangleInternal(f, b, c, InvalidID, InvalidID, InvalidID);
-		int32 t3 = AddTriangleInternal(f, d, b, InvalidID, InvalidID, InvalidID);
-		
-
-		// update the edges we found above, to point to triangles
-		ReplaceEdgeTriangle(ebc, t0, t2);
-		ReplaceEdgeTriangle(edb, t1, t3);
-
-		// edge eab became eaf
-		int32 eaf = eab; //Edge * eAF = eAB;
-		Edges[eaf].Vert = FIndex2i(FMath::Min(a, f), FMath::Max(a, f));
-
-		// update a/b/f vertex-edges
-		if (a != b)
-		{ 
-			VertexEdgeLists.Remove(b, eab);
-		}
-		VertexEdgeLists.Insert(f, eaf);
-
-		// create edges connected to f  (also updates vertex-edges)
-		int32 efb = AddEdgeInternal(f, b, t2, t3);
-		int32 efc = AddEdgeInternal(f, c, t0, t2);
-		int32 edf = AddEdgeInternal(d, f, t1, t3);
-
-		// update triangle edge-nbrs
-		ReplaceTriangleEdge(t0, ebc, efc);
-		ReplaceTriangleEdge(t1, edb, edf);
-		TriangleEdges[t2] = FIndex3i(efb, ebc, efc);
-		TriangleEdges[t3] = FIndex3i(edf, edb, efb);
-
-		// update vertex refcounts
-		VertexRefCounts.Increment(c);
-		VertexRefCounts.Increment(d);
-		VertexRefCounts.Increment(f, 4);
-
-		SplitInfo.bIsBoundary = false;
-		SplitInfo.OtherVertices = FIndex2i(c, d);
-		SplitInfo.NewVertex = f;
-		SplitInfo.NewEdges = FIndex3i(efb, efc, edf);
-		SplitInfo.NewTriangles = FIndex2i(t2, t3);
-
-		return EMeshResult::Ok;
-	}
-
-}
-
-
 UE::Geometry::EMeshResult  UE::Geometry::FIntrinsicTriangulation::SplitEdge(int32 EdgeAB, FEdgeSplitInfo& SplitInfo, double SplitParameterT)
 {
 	SplitParameterT = FMath::Clamp(SplitParameterT, 0., 1.);
@@ -2023,19 +2819,16 @@
 	if (IsBoundaryEdge(EdgeAB))
 	{
 		// state before the split
-		const int32 TID = OriginalEdge.Tri[0];
+		const int32 TID      = OriginalEdge.Tri[0];
 		const int32 IndexOfe = GetTriEdges(TID).IndexOf(EdgeAB);
 
 		// Info about the original T0 tri, reordered to make the split edge the first edge..
-		const FIndex3i OriginalT0VIDs  = Permute(IndexOfe, GetTriangle(TID));  // as (a, b, c)
-		const FIndex3i OriginalT0Edges = Permute(IndexOfe, GetTriEdges(TID));  // as (ab, bc, ca)
-		
+		const FIndex3i OriginalT0VIDs         = Permute(IndexOfe, GetTriangle(TID));  // as (a, b, c)
+		const FVector3d OriginalT0EdgeDirs    = Permute(IndexOfe, SignpostData.IntrinsicEdgeAngles[TID]);  // as ( aTob, bToc, cToa)
 		const FVector3d OriginalT0EdgeLengths = Permute(IndexOfe, GetTriEdgeLengths(TID));    // as (|ab|, |bc|, |ca|)
-		const FVector3d OriginalT0EdgeDirs    = Permute(IndexOfe, SignpostData.IntrinsicEdgeAngles[TID]);  // as ( aTob, bToc, cToa)
-
 		
 		// Update the connectivity with the edge split
-		EMeshResult Result = SplitEdgeTopology(EdgeAB, SplitInfo);
+		EMeshResult Result = MyBase::SplitEdge(EdgeAB, SplitInfo, SplitParameterT);
 
 		if (Result != EMeshResult::Ok)
 		{
@@ -2043,25 +2836,8 @@
 		}
 		const int32 NewVID = SplitInfo.NewVertex;
 		const int32 NewTID = SplitInfo.NewTriangles[0]; // edges {fb, bc, cf}
-		const int32 NewEdgeFB = SplitInfo.NewEdges[0];
-		const int32 NewEdgeFC = SplitInfo.NewEdges[1];
 
 		const double DistAF = SplitParameterT * OriginalT0EdgeLengths[0];
-		const double DistFB = FMath::Max(0., OriginalT0EdgeLengths[0] - DistAF);
-		// new edge length
-		const double DistFC = DistanceBetweenBarycentricPoints( OriginalT0EdgeLengths[0],
-			                                                    OriginalT0EdgeLengths[1],
-																OriginalT0EdgeLengths[2],
-			                                                    FVector3d(SplitParameterT, 1. - SplitParameterT, 0.), 
-																FVector3d(0., 0., 1.));
-		
-		EdgeLengths.InsertAt(DistFC, NewEdgeFC);
-		EdgeLengths.InsertAt(DistFB, NewEdgeFB);
-		EdgeLengths[EdgeAB] = DistAF;
-		
-		// update the internal angles
-		InternalAngles.InsertAt(ComputeTriInternalAnglesR(NewTID), NewTID);
-		InternalAngles[TID] = ComputeTriInternalAnglesR(TID);
 		const FVector3d TIDInternalAngles = Permute(IndexOfe, InternalAngles[TID]);
 		
 		// update the directions.  Say DirFtoB is zero
@@ -2082,8 +2858,8 @@
 
 		// trace from A vertex in the direction toward the new vertex the edge length distance. This updates the surface position entry for the new vert.
 		const int32 AVID = OriginalT0VIDs[0];
-		const double AtoFDir = TriEdgeDir[0][0];
-		const double FToAAngle = UpdateVertexByEdgeTrace(NewVID, AVID, AtoFDir, DistAF);
+		const double AtoFDir    = TriEdgeDir[0][0];
+		const double FToAAngle  = UpdateVertexByEdgeTrace(NewVID, AVID, AtoFDir, DistAF);
 		const double ToLocalDir = (FToAAngle - FtoADir); 
 		// update the Fto{B,C}  directions such that the direction from F to A would agree with FtoAAngle
 		TriEdgeDir[0][1] = AsZeroToTwoPi(TriEdgeDir[0][1] + ToLocalDir);
@@ -2104,24 +2880,19 @@
 
 		// Info about the original T0 tri, reordered to make the split edge the first edge..
 		const FIndex3i OriginalT0VIDs  = Permute(IndexOfe, GetTriangle(TID));  // as (a, b, c)
-		const FIndex3i OriginalT0Edges = Permute(IndexOfe, GetTriEdges(TID));  // as (ab, bc, ca)
+		const FVector3d OriginalT0EdgeDirs    = Permute(IndexOfe, SignpostData.IntrinsicEdgeAngles[TID]);  // as ( aTob, bToc, cToa)
+		const FVector3d OriginalT0EdgeLengths = Permute(IndexOfe, GetTriEdgeLengths(TID));    // as (|ab|, |bc|, |ca|)
 
-		const FVector3d OriginalT0EdgeLengths = Permute(IndexOfe, GetTriEdgeLengths(TID));    // as (|ab|, |bc|, |ca|)
-		const FVector3d OriginalT0EdgeDirs    = Permute(IndexOfe, SignpostData.IntrinsicEdgeAngles[TID]);  // as ( aTob, bToc, cToa)
-		
 		// Info about the original T0 tri, reordered to make the split edge the first edge..
 		const int32 TID1 = OriginalEdge.Tri[1];
 		const int32 IndexOfT1e = GetTriEdges(TID1).IndexOf(EdgeAB);
 
-		const FIndex3i OriginalT1VIDs  = Permute(IndexOfT1e, GetTriangle(TID1));  // as (b, a, d)
-		const FIndex3i OriginalT1Edges = Permute(IndexOfT1e, GetTriEdges(TID1));  // as (ba, ad, db)
-		
-		const FVector3d OriginalT1EdgeLengths = Permute(IndexOfT1e, GetTriEdgeLengths(TID1));   // as (|ba|, |ad|, |db})
-		const FVector3d OriginalT1EdgeDirs    = Permute(IndexOfT1e, SignpostData.IntrinsicEdgeAngles[TID1]); // as (bToa, aTod, dTob) 
+		const FIndex3i OriginalT1VIDs      = Permute(IndexOfT1e, GetTriangle(TID1));  // as (b, a, d)
+		const FVector3d OriginalT1EdgeDirs = Permute(IndexOfT1e, SignpostData.IntrinsicEdgeAngles[TID1]); // as (bToa, aTod, dTob) 
 		
 
-		// Update the connectivity with the edge split
-		EMeshResult Result = SplitEdgeTopology(EdgeAB, SplitInfo);
+		// Update the connectivity and the intrinsic data with the edge split
+		EMeshResult Result = MyBase::SplitEdge(EdgeAB, SplitInfo, SplitParameterT);
 
 		if (Result != EMeshResult::Ok)
 		{
@@ -2130,38 +2901,9 @@
 		const int32 NewVID = SplitInfo.NewVertex;
 		const int32 NewTID0 = SplitInfo.NewTriangles[0]; // edges {fb, bc, cf}
 		const int32 NewTID1 = SplitInfo.NewTriangles[1]; // edges {fd, db, bc}
-		const int32 EdgeAF = EdgeAB;
-		const int32 NewEdgeFB = SplitInfo.NewEdges[0];
-		const int32 NewEdgeFC = SplitInfo.NewEdges[1];
-		const int32 NewEdgeFD = SplitInfo.NewEdges[2];
 
 		const double DistAF = SplitParameterT * OriginalT0EdgeLengths[0];
-		const double DistFB = FMath::Max(0., OriginalT0EdgeLengths[0] - DistAF);
-		// new edge length
-		const double DistFC = DistanceBetweenBarycentricPoints( OriginalT0EdgeLengths[0],
-			                                                    OriginalT0EdgeLengths[1],
-			                                                    OriginalT0EdgeLengths[2],
-			                                                    FVector3d(SplitParameterT, 1. - SplitParameterT, 0.),
-			                                                    FVector3d(0., 0., 1.));
 		
-		// new edge length
-		const double DistFD = DistanceBetweenBarycentricPoints( OriginalT1EdgeLengths[0],
-			                                                    OriginalT1EdgeLengths[1],
-			                                                    OriginalT1EdgeLengths[2],
-			                                                    FVector3d(1. - SplitParameterT, SplitParameterT, 0.),
-			                                                    FVector3d(0., 0., 1.));
-
-		EdgeLengths.InsertAt(DistFD, NewEdgeFD);
-		EdgeLengths.InsertAt(DistFC, NewEdgeFC);
-		EdgeLengths.InsertAt(DistFB, NewEdgeFB);
-		EdgeLengths[EdgeAF] = DistAF;
-
-		// update the internal angles (this uses edge lengths)
-		InternalAngles.InsertAt(ComputeTriInternalAnglesR(NewTID1), NewTID1);
-		InternalAngles.InsertAt(ComputeTriInternalAnglesR(NewTID0), NewTID0);
-		InternalAngles[TID] = ComputeTriInternalAnglesR(TID);
-		InternalAngles[TID1] = ComputeTriInternalAnglesR(TID1);
-
 		// update the directions.  Say DirFtoB is zero
 		const double ToRadiansAtC = SignpostData.GeometricVertexInfo[OriginalT0VIDs[2]].ToRadians;
 		const double ToRadiansAtD = SignpostData.GeometricVertexInfo[OriginalT1VIDs[2]].ToRadians;
@@ -2193,7 +2935,7 @@
 		// trace from A vertex in the direction toward the new vertex the edge length distance. This updates the surface position entry for the new vert.
 		const int32 AVID = OriginalT0VIDs[0];
 		const double AtoFDir = TriEdgeDir[0][0];
-		const double FToAAngle = UpdateVertexByEdgeTrace(NewVID, AVID, AtoFDir, DistAF);
+		const double FToAAngle   = UpdateVertexByEdgeTrace(NewVID, AVID, AtoFDir, DistAF);
 		const double ToLocalDir = (FToAAngle - TriEdgeDir[2][0]);
 
 		// fix angles relative to local dir.
@@ -2252,67 +2994,7 @@
 
 TArray<FIntrinsicEdgeFlipMesh::FSurfacePoint> FIntrinsicEdgeFlipMesh::TraceEdge(int32 IntrinsicEID, double CoalesceThreshold, bool bReverse) const
 {
-	TArray<FIntrinsicEdgeFlipMesh::FSurfacePoint> Result;
-	
-	if (!IsEdge(IntrinsicEID))
-	{
-		return MoveTemp(Result);
-	}
-
-	TArray<int32> SurfXings;
-	const int32 NumSurfXings = NormalCoordinates.NormalCoord[IntrinsicEID];
-	if (NumSurfXings > 0)
-	{
-		SurfXings.Reserve(NumSurfXings);
-	}
-
-	const FDynamicMesh3& SurfaceMesh = *GetExtrinsicMesh();
-	const FIntrinsicEdgeFlipMesh& IntrinsicMesh = *this;
-
-	const FIndex2i IntrinsicEdgeT = GetEdgeT(IntrinsicEID);
-
-	// need to create a list of surface mesh edges that cross this edge.  To do this we need to follow
-	// each curve that crosses this intrinsic edge to the vertices where it terminates ( thus identifying the surface edge )
-	for (int32 p = 1; p < NumSurfXings + 1; ++p)
-	{
-
-		const FIndex2i SurfaceEdgeV = [&]
-			{
-		
-				TArray<FNormalCoordinates::FEdgeAndCrossingIdx> Crossings;
-				FIndex2i Verts;
-				// follow surface curve (i.e. surface edge) forward to the end and identify the vertex
-				{
-					FNormalCoordSurfaceTraceImpl::ContinueTraceSurfaceEdge(IntrinsicMesh,  IntrinsicEdgeT.A, IntrinsicEID, p, Crossings);
-					const FNormalCoordinates::FEdgeAndCrossingIdx& LastXing = Crossings.Last();
-
-					const FIndex3i TriEIDs = GetTriEdges(LastXing.TID);
-					Verts.A                = GetTriangle(LastXing.TID)[TriEIDs.IndexOf(LastXing.EID)];
-				}
-				Crossings.Reset();
-
-				// follow surface curve (i.e. surface edge) backward to other end and identify the vertex
-				{
-					FNormalCoordSurfaceTraceImpl::ContinueTraceSurfaceEdge(IntrinsicMesh, IntrinsicEdgeT.B, IntrinsicEID, NumSurfXings + 1 - p, Crossings);
-					const FNormalCoordinates::FEdgeAndCrossingIdx& LastXing = Crossings.Last();
-
-					const FIndex3i TriEIDs = GetTriEdges(LastXing.TID);
-					Verts.B                = GetTriangle(LastXing.TID)[TriEIDs.IndexOf(LastXing.EID)];
-				}
-				return Verts;
-			}();
-
-		// identify the surface edge from its endpoints. 
-		const int32 SurfaceEID = SurfaceMesh.FindEdge(SurfaceEdgeV.A, SurfaceEdgeV.B);
-		checkSlow(SurfaceEID != InvalidID);
-		
-		SurfXings.Add(SurfaceEID);
-	}
-	
-	// do the actual trace - this identifies the surface mesh triangles crossed by the intrinsic edge and unfolds them into a triangle strip where the trace is performed
-	Result =  FNormalCoordSurfaceTraceImpl::TraceEdgeOverHost(IntrinsicEID, SurfXings, SurfaceMesh, IntrinsicMesh, CoalesceThreshold, bReverse);
-
-	return MoveTemp(Result);
+	return FNormalCoordIntrinsicTraceImpl::TraceEdge(*this, IntrinsicEID, CoalesceThreshold, bReverse);
 }
 
 TArray<FIntrinsicEdgeFlipMesh::FEdgeAndCrossingIdx> FIntrinsicEdgeFlipMesh::GetImplicitEdgeCrossings(const int32 SurfaceEID, const bool bReverse) const
@@ -2320,12 +3002,43 @@
 	return FNormalCoordSurfaceTraceImpl::TraceSurfaceEdge(*this, SurfaceEID, bReverse);
 }
 
-FIntrinsicEdgeFlipMesh::FEdgeCorrespondence::FEdgeCorrespondence(const FIntrinsicEdgeFlipMesh& Mesh)
-	: IntrinsicMesh(&Mesh)
-	, SurfaceMesh(Mesh.GetNormalCoordinates().SurfaceMesh)
+
+
+TArray<FIntrinsicEdgeFlipMesh::FSurfacePoint>
+FIntrinsicEdgeFlipMesh::TraceSurfaceEdge(int32 SurfaceEID, double CoalesceThreshold, bool bReverse) const
 {
+	const FDynamicMesh3& HostMesh                   = *this->GetExtrinsicMesh();
+	const FIntrinsicEdgeFlipMesh& TraceMesh         = *this;
+	TArray<FEdgeAndCrossingIdx> EdgeAndCrossingIdxs = this->GetImplicitEdgeCrossings(SurfaceEID, bReverse);
+
+	// NB: for this intrinsic mesh type, we know that the intrinsic verts are the same as the surface verts ( since splits and pokes aren't allowed)
+	TArray<int32> HostEdgeCrossings;
+	HostEdgeCrossings.Reserve(EdgeAndCrossingIdxs.Num() - 2); // EdgeAndCrossingsIdx include the start and end vertex.  don't need them.
+	for (FEdgeAndCrossingIdx EdgeAndCrossing : EdgeAndCrossingIdxs)
+	{
+		if (EdgeAndCrossing.CIdx != 0) // skip the start and end vertex
+		{
+			HostEdgeCrossings.Add(EdgeAndCrossing.EID);
+		}
+	}
+
+	return FNormalCoordSurfaceTraceImpl::TraceEdgeOverHost(SurfaceEID, HostEdgeCrossings, HostMesh, TraceMesh, CoalesceThreshold, bReverse, 
+														  [](int32 VID) {return FSurfacePoint(VID);});
+}
+
+/**------------------------------------------------------------------------------
+*  TEdgeCorrespondence 
+*-------------------------------------------------------------------------------*/
+
+template <typename IntrinsicMeshType>
+void UE::Geometry::TEdgeCorrespondence<IntrinsicMeshType>::Setup(const IntrinsicMeshType& Mesh)
+{
+	SurfaceEdgesCrossed.Reset();
+	IntrinsicMesh = &Mesh;
+	SurfaceMesh = Mesh.GetNormalCoordinates().SurfaceMesh;
+
 	const int32 IntrinsicMaxEID = IntrinsicMesh->MaxEdgeID();
-	const int32 SurfaceMaxEID   = SurfaceMesh->MaxEdgeID();
+	const int32 SurfaceMaxEID = SurfaceMesh->MaxEdgeID();
 
 	SurfaceEdgesCrossed.SetNum(IntrinsicMaxEID);
 	const IntrinsicCorrespondenceUtils::FNormalCoordinates& NormalCoords = IntrinsicMesh->GetNormalCoordinates();
@@ -2338,7 +3051,7 @@
 			continue;
 		}
 		// number of times a surface edge crosses this intrinsic edge
-		int32 NumXings = NormalCoords.NormalCoord[IntrinsicEID];
+		const int32 NumXings = NormalCoords.NumEdgeCrossing(IntrinsicEID);
 
 		if (NumXings > 0)
 		{
@@ -2362,11 +3075,11 @@
 		for (int32 i = 0; i < IntrinsicEdgeXings.Num(); ++i)
 		{
 			const FEdgeAndCrossingIdx& EdgeXing = IntrinsicEdgeXings[i];
-			bool bIsEndVertex                   = (EdgeXing.CIdx == 0);
+			bool bIsEndVertex = (EdgeXing.CIdx == 0);
 
 			if (!bIsEndVertex)
 			{
-				const int32 IntrinsicEID      = EdgeXing.EID;
+				const int32 IntrinsicEID = EdgeXing.EID;
 				const FIndex2i IntrinsicEdgeT = IntrinsicMesh->GetEdgeT(IntrinsicEID);
 				checkSlow(IntrinsicEdgeT.A == EdgeXing.TID || IntrinsicEdgeT.B == EdgeXing.TID);
 
@@ -2383,33 +3096,1471 @@
 	}
 }
 
-TArray<FIntrinsicEdgeFlipMesh::FSurfacePoint>
-FIntrinsicEdgeFlipMesh::TraceSurfaceEdge(int32 SurfaceEID, double CoalesceThreshold, bool bReverse) const
+template <typename IntrinsicMeshType>
+TArray<UE::Geometry::IntrinsicCorrespondenceUtils::FSurfacePoint> TEdgeCorrespondence<IntrinsicMeshType>::TraceEdge(int32 IntrinsicEID, double CoalesceThreshold, bool bReverse) const
 {
-	const FDynamicMesh3& HostMesh                   = *this->GetExtrinsicMesh();
-	const FIntrinsicEdgeFlipMesh& TraceMesh         = *this;
-	TArray<FEdgeAndCrossingIdx> EdgeAndCrossingIdxs = this->GetImplicitEdgeCrossings(SurfaceEID, bReverse);
+	const FDynamicMesh3& HostMesh = *SurfaceMesh;
+	const IntrinsicMeshType& TraceMesh = *IntrinsicMesh;
+	const TArray<int32>& HostEdgeCrossings = SurfaceEdgesCrossed[IntrinsicEID];
 
-	TArray<int32> HostEdgeCrossings;
-	HostEdgeCrossings.Reserve(EdgeAndCrossingIdxs.Num() - 2); // EdgeAndCrossingsIdx include the start and end vertex.  don't need them.
-	for (FEdgeAndCrossingIdx EdgeAndCrossing : EdgeAndCrossingIdxs)
+	return FNormalCoordSurfaceTraceImpl::TraceEdgeOverHost(IntrinsicEID, HostEdgeCrossings, HostMesh, TraceMesh, CoalesceThreshold, bReverse, 
+	                                                       [&TraceMesh](int32 VID) { return TraceMesh.GetVertexSurfacePoint(VID); });
+}
+
+template struct UE::Geometry::TEdgeCorrespondence<UE::Geometry::FIntrinsicEdgeFlipMesh>;
+template struct UE::Geometry::TEdgeCorrespondence<UE::Geometry::FIntrinsicMesh>;
+
+/**------------------------------------------------------------------------------
+*  FIntrinsicMesh Helpers
+*-------------------------------------------------------------------------------*/
+
+
+namespace FIntrinsicMeshImplUtils
+{
+	using namespace UE::Geometry;
+
+	/** 
+	* Test if a given SurfacePoint is on the specified surface triangle of the surface mesh, this includes the triangle boundaries:
+	* i.e. will return true, if the point is on the face of the triangle, one of the edges, or on a vertex of the triangle.
+	*/
+	bool IsOnSurfaceTriangle(const IntrinsicCorrespondenceUtils::FSurfacePoint& SurfacePoint, const int32 SurfaceTID, const FDynamicMesh3& SurfaceMesh)
 	{
-		if (EdgeAndCrossing.CIdx != 0) // skip the start and end vertex
+		bool bFoundTID = false;
+		if (!SurfaceMesh.IsTriangle(SurfaceTID))
 		{
-			HostEdgeCrossings.Add(EdgeAndCrossing.EID);
+			return bFoundTID;
+		}
+		
+		if (IsFacePoint(SurfacePoint))
+		{
+			const int32 TID = SurfacePoint.Position.TriPosition.TriID;
+			if (TID == SurfaceTID)
+			{
+				bFoundTID = true;
+			}
+		}
+		else if (IsEdgePoint(SurfacePoint))
+		{
+			const int32 SurfaceEID      = SurfacePoint.Position.EdgePosition.EdgeID;
+			const FIndex2i SurfaceEdgeT = SurfaceMesh.GetEdgeT(SurfaceEID);
+			if (SurfaceEdgeT.IndexOf(SurfaceTID) != -1)
+			{
+				bFoundTID = true;
+			}
+		}
+		else
+		{
+			checkSlow(IsVertexPoint(SurfacePoint));
+			const int32 SurfaceVID = SurfacePoint.Position.VertexPosition.VID;
+			for (int TID : SurfaceMesh.VtxTrianglesItr(SurfaceVID))
+			{
+				if (TID == SurfaceTID)
+				{
+					bFoundTID = true;
+					break;
+				}
+			}
+
+		}
+
+		return bFoundTID;
+	}
+
+	/** 
+	* Collect the surface triangles that are adjacent to the specified surface point 
+	* this may be one, two, or many - depending on location of the point (eg a surface point that is on an edge may have two adjacent tris)
+	*/
+	TArray<int32> GetAdjacentTriangles(const IntrinsicCorrespondenceUtils::FSurfacePoint& SurfacePoint, const FDynamicMesh3& SurfaceMesh)
+	{
+
+		TArray<int32> AdjTIDs;
+		if (IsEdgePoint(SurfacePoint))
+		{
+			const int32 SurfaceEID      = SurfacePoint.Position.EdgePosition.EdgeID;
+			const FIndex2i SurfaceEdgeT = SurfaceMesh.GetEdgeT(SurfaceEID);
+			AdjTIDs.Add(SurfaceEdgeT.A);
+			if (SurfaceEdgeT.B != -1)
+			{
+				AdjTIDs.Add(SurfaceEdgeT.B);
+			}
+		}
+		else if (IsFacePoint(SurfacePoint))
+		{
+			const int32 SurfaceTID = SurfacePoint.Position.TriPosition.TriID;
+			AdjTIDs.Add(SurfaceTID);
+		}
+		else
+		{
+			checkSlow(IsVertexPoint(SurfacePoint));
+			const int32 SurfaceVID = SurfacePoint.Position.VertexPosition.VID;
+			for (int SurfaceTID : SurfaceMesh.VtxTrianglesItr(SurfaceVID))
+			{
+				AdjTIDs.Add(SurfaceTID);
+			}
+		}
+		return MoveTemp(AdjTIDs);
+	}
+
+	/**  
+	* Convert a surface point to barycentric coordinates relative to the specified surface triangle. 
+	*     NB: this assumes (but does not check) that the surface point is on or adjacent to the surface triangle. 
+	*/
+	FVector3d AsBarycenteric(const IntrinsicCorrespondenceUtils::FSurfacePoint& SurfacePoint, const int32 SurfaceTID, const FDynamicMesh3& SurfaceMesh)
+	{
+		FVector3d BC(0., 0., 0.);
+
+
+		if (IsEdgePoint(SurfacePoint))  // translate an edge point to barycentric coordinates
+		{
+			const int32 SurfaceEID = SurfacePoint.Position.EdgePosition.EdgeID;
+			const double Alpha     = SurfacePoint.Position.EdgePosition.Alpha; // Pos(Edge.A) (Alpha) + (1-Alpha) Pos(Edge.B)
+
+			checkSlow(SurfaceMesh.GetEdgeT(SurfaceEID).IndexOf(SurfaceTID) != -1);
+
+			const FIndex2i SurfaceEdgeV = SurfaceMesh.GetEdgeV(SurfaceEID);
+			const FIndex3i SurfaceTri   = SurfaceMesh.GetTriangle(SurfaceTID);
+
+			for (int32 i = 0; i < 3; ++i)
+			{
+				if (SurfaceTri[i] == SurfaceEdgeV.A)
+				{
+					BC[i] = Alpha;
+				}
+				else if (SurfaceTri[i] == SurfaceEdgeV.B)
+				{
+					BC[i] = (1. - Alpha);
+				}
+			}
+		}
+		else if (IsFacePoint(SurfacePoint)) // translate a face point to barycentric coordinates
+		{
+			checkSlow(SurfacePoint.Position.TriPosition.TriID == SurfaceTID);
+			BC = SurfacePoint.Position.TriPosition.BarycentricCoords;
+		}
+		else // translate a vertex point to barycentric coordinates
+		{
+			const int32 SurfaceVID    = SurfacePoint.Position.VertexPosition.VID;
+			const FIndex3i SurfaceTri = SurfaceMesh.GetTriangle(SurfaceTID);
+			bool bValid = false;
+			for (int32 i = 0; i < 3; ++i)
+			{
+				if (SurfaceTri[i] == SurfaceVID)
+				{
+					BC[i] = 1.;
+					bValid = true;
+				}
+			}
+			checkSlow(bValid);
+		}
+		return BC;
+	}
+
+	bool AreOnSameSurfaceEdge(const IntrinsicCorrespondenceUtils::FSurfacePoint& PointA, const IntrinsicCorrespondenceUtils::FSurfacePoint&  PointB, const FDynamicMesh3& Mesh)
+	{
+		if (IsFacePoint(PointA) || IsFacePoint(PointB))
+		{
+			return false;
+		}
+
+		if (IsVertexPoint(PointA) && IsVertexPoint(PointB))
+		{
+			return Mesh.FindEdge(PointA.Position.VertexPosition.VID, PointB.Position.VertexPosition.VID) != FDynamicMesh3::InvalidID;
+		}
+
+		auto TestVertAndEdge = [&Mesh](const IntrinsicCorrespondenceUtils::FSurfacePoint& VertexPoint, const IntrinsicCorrespondenceUtils::FSurfacePoint& EdgePoint)->bool
+										{
+											const int32 VID = VertexPoint.Position.VertexPosition.VID;
+											const int32 EID = EdgePoint.Position.EdgePosition.EdgeID;
+											const FDynamicMesh3::FEdge& Edge = Mesh.GetEdgeRef(EID);
+											return (Edge.Vert.IndexOf(VID) != -1);
+										};
+		if (IsVertexPoint(PointA))
+		{
+			// B must be an edge point.
+			return TestVertAndEdge(PointA, PointB);
+		}
+		else
+		{
+			// B must be vert and A is edge
+			return TestVertAndEdge(PointB, PointA);
 		}
 	}
 
-	return FNormalCoordSurfaceTraceImpl::TraceEdgeOverHost(SurfaceEID, HostEdgeCrossings, HostMesh, TraceMesh, CoalesceThreshold, bReverse);
+	/** array given matrix-like access with row-major layout for a non-square matrix */
+	struct FNonSqrMat
+	{
+		FNonSqrMat(int32 NumCols, int32 NumRows)
+			: N(NumCols)
+			, M(NumRows)
+		{
+			Data.AddZeroed(N * M);
+		}
+
+		double& operator()(int32 i, int32 j)
+		{
+			return Data[RowMajorOffset(i, j)];
+		}
+		const double& operator()(int32 i, int32 j) const
+		{
+			return Data[RowMajorOffset(i, j)];
+		}
+		int32 RowMajorOffset(int32 i, int32 j) const
+		{
+			return i * N + j;
+		}
+
+		int32 N;
+		int32 M;
+		TArray<double> Data;
+	};
+
+	/**
+	* Given an underspecified matrix equation (more unknowns than equations) of the form  M * Soln = BVector, 
+	* where M is a Nx3 matrix, Soln is an N-dimensional array, and BVector is a 3-dimensional array
+	* 
+	* this computes the solution that minimizes the quantity |Soln|
+	* 
+	* The matrix M is specified by providing the three rows of the matrix (Row0, Row1, Row2).
+	* 
+	* @return false on failure, otherwise the vector Soln will hold the result.
+	*/
+	bool ComputeMinNormSolution(const TArray<double>& Row0, const TArray<double>& Row1, const TArray<double>& Row2, const FVector3d& BVector, TArray<double>& Soln)
+	{
+		checkSlow(Row0.Num() == Row1.Num() && Row2.Num() == Row1.Num());
+
+		const int32 NCols = Row0.Num();
+
+		checkSlow(NCols > 2);
+
+
+		Soln.SetNum(NCols);
+
+		if (NCols == 3)
+		{
+			const FMatrix3d Mat3d(Row0[0], Row0[1], Row0[2],
+								  Row1[0], Row1[1], Row1[2],
+								  Row2[0], Row2[1], Row2[2]);
+
+			const double Det = Mat3d.Determinant();
+			if (TMathUtil<double>::Abs(Det) < TMathUtil<double>::Epsilon)
+			{
+				return false;
+			}
+			else
+			{
+				const FMatrix3d InvMat3d    = Mat3d.Inverse();
+				const FVector3d SolVector3d = InvMat3d * BVector;
+
+				Soln[0] = SolVector3d[0];
+				Soln[1] = SolVector3d[1];
+				Soln[2] = SolVector3d[2];
+			}
+		}
+		else
+		{
+			FNonSqrMat Mat(NCols, 3);
+			for (int j = 0; j < NCols; ++j)
+			{
+				Mat(0, j) = Row0[j];
+				Mat(1, j) = Row1[j];
+				Mat(2, j) = Row2[j];
+			}
+
+			// form MMt = Mat * Transpose(Mat)
+
+			FMatrix3d MMt(0.);
+			for (int32 i = 0; i < 3; ++i)
+			{
+				FVector3d& Row = (i==0) ? MMt.Row0 : (i==1) ? MMt.Row1 : MMt.Row2 ;
+				for (int32 j = 0; j < 3; ++j)
+				{
+					for (int32 k = 0; k < NCols; ++k)
+					{
+						Row[j] += Mat(i, k) * Mat(j, k);
+					}
+				}
+			}
+
+			// solve MMt y = b, and compute x = Mt * y
+			const double Det = MMt.Determinant();
+			if (TMathUtil<double>::Abs(Det) < TMathUtil<double>::Epsilon)
+			{
+				return false;
+			}
+			else
+			{
+				const FMatrix3d InvMMt    = MMt.Inverse();
+				const FVector3d YVector3d = InvMMt * BVector;
+
+				//Soln = Transpose(Mat) * Yvector;
+				for (int32 j = 0; j < NCols; ++j)
+				{
+					Soln[j] = Mat(0, j) * YVector3d[0] + Mat(1, j) * YVector3d[1] + Mat(2, j) * YVector3d[2];
+				}
+			}
+		}
+
+		return true;
+	}
+
+} // end namespace FIntrinsicMeshImplUtils
+
+
+/**------------------------------------------------------------------------------
+*  FIntrinsicMesh Methods
+*-------------------------------------------------------------------------------*/
+
+UE::Geometry::FIntrinsicMesh::FIntrinsicMesh(const FDynamicMesh3& SurfaceMesh)
+	: MyBase(SurfaceMesh)
+	, NormalCoordinates(SurfaceMesh)
+{
+	const int32 MaxVertexID = SurfaceMesh.MaxVertexID();
+
+	// add surface position.
+
+	IntrinsicVertexPositions.SetNum(MaxVertexID);
+
+	// initialize: identify with vertex in Extrinsic Mesh
+	for (int32 VID = 0; VID < MaxVertexID; ++VID)
+	{
+		if (SurfaceMesh.IsVertex(VID))
+		{
+			FSurfacePoint VertexSurfacePoint(VID);
+			IntrinsicVertexPositions[VID] = VertexSurfacePoint;
+		}
+	}
+
 }
 
-TArray<FIntrinsicEdgeFlipMesh::FSurfacePoint>
-FIntrinsicEdgeFlipMesh::FEdgeCorrespondence::TraceEdge(int32 IntrinsicEID, double CoalesceThreshold, bool bReverse) const
+
+TArray<UE::Geometry::FIntrinsicMesh::FSurfacePoint> FIntrinsicMesh::TraceEdge(int32 IntrinsicEID, double CoalesceThreshold, bool bReverse) const
 {
-	const FDynamicMesh3& HostMesh           = *SurfaceMesh;
-	const FIntrinsicEdgeFlipMesh& TraceMesh = *IntrinsicMesh;
-	const TArray<int32>& HostEdgeCrossings  = SurfaceEdgesCrossed[IntrinsicEID];
-
-	return FNormalCoordSurfaceTraceImpl::TraceEdgeOverHost(IntrinsicEID, HostEdgeCrossings, HostMesh, TraceMesh, CoalesceThreshold, bReverse);
+	return FNormalCoordIntrinsicTraceImpl::TraceEdge(*this, IntrinsicEID, CoalesceThreshold, bReverse);
 }
 
+TArray<UE::Geometry::FIntrinsicMesh::FEdgeAndCrossingIdx> FIntrinsicMesh::GetImplicitEdgeCrossings(const int32 SurfaceEID, const bool bReverse) const
+{
+	return FNormalCoordSurfaceTraceImpl::TraceSurfaceEdge(*this, SurfaceEID, bReverse);
+}
+
+
+TArray<UE::Geometry::FIntrinsicMesh::FSurfacePoint> FIntrinsicMesh::TraceSurfaceEdge(int32 SurfaceEID, double CoalesceThreshold, bool bReverse) const
+{
+	const FDynamicMesh3& SurfaceMesh = *this->GetExtrinsicMesh();
+	const FIntrinsicMesh& IntrinsicMesh = *this;
+	TArray<FEdgeAndCrossingIdx> EdgeAndCrossingIdxs = this->GetImplicitEdgeCrossings(SurfaceEID, bReverse);
+
+	const int32 NumEdgeAndXIdx = EdgeAndCrossingIdxs.Num();
+
+	TArray<UE::Geometry::FIntrinsicMesh::FSurfacePoint>  ResultTraceArray;
+
+	if (NumEdgeAndXIdx == 0)
+	{
+		return ResultTraceArray;
+	}
+
+	// find start (.A) and end (.B) trace mesh vids for this edge
+	const FIndex2i OrderedTraceEdgeV = [&]
+										{
+											const FIndex2i EdgeV = SurfaceMesh.GetEdgeV(SurfaceEID);
+											if (bReverse)
+											{
+												return FIndex2i(EdgeV.B, EdgeV.A);
+											}
+											else
+											{
+												return EdgeV;
+											}
+										}();
+
+	// the start and end are vertices in both SurfaceMesh and IntrinsicMesh, furthermore when vertices exist on both meshes they have the same VID 
+	const FSurfacePoint TraceStartSurfacePoint(OrderedTraceEdgeV.A);
+	const FSurfacePoint TraceEndSurfacePoint(OrderedTraceEdgeV.B);
+
+
+	ResultTraceArray.Add(TraceStartSurfacePoint);
+
+	// due to (possible) intrinsic mesh edge splits, the trace of the surface edge across the intrinsic mesh has to be broken
+	// into sections that follow split intrinsic mesh edges, and those that cross the faces of intrinsic mesh triangles.
+	{
+		FSurfacePoint LocalStartSurfacePoint = TraceStartSurfacePoint;
+		TArray<int32> IntrinsicEdgesCrossed;
+		for (int32 i = 1; i < NumEdgeAndXIdx; ++i)
+		{
+			const FEdgeAndCrossingIdx& EdgeIDandXIdx = EdgeAndCrossingIdxs[i];
+			if (EdgeIDandXIdx.CIdx != 0)
+			{
+				IntrinsicEdgesCrossed.Add(EdgeIDandXIdx.EID);
+			}
+			else if (IntrinsicEdgesCrossed.Num() > 0)
+			{
+				// current point is the end point of local face crossing.
+				// convert to surface point.
+				const FSurfacePoint LocalEndSurfacePoint( IntrinsicMesh.GetTriangle(EdgeIDandXIdx.TID)[EdgeIDandXIdx.EID]);
+				
+				// compute crossings
+				FNormalCoordSurfaceTraceImpl::ConvertEdgesCrossed( LocalStartSurfacePoint, LocalEndSurfacePoint, IntrinsicEdgesCrossed, 
+					                                               IntrinsicMesh, SurfaceMesh, CoalesceThreshold,
+															       [](int32 vid){return FSurfacePoint(vid);}, ResultTraceArray);
+				
+				// add surface point 
+				ResultTraceArray.Add(LocalEndSurfacePoint);
+
+				// empty the temp crossings.
+				IntrinsicEdgesCrossed.Reset();
+
+				LocalStartSurfacePoint = LocalEndSurfacePoint;
+			}
+			else
+			{
+				// convert to surface point.
+				const FSurfacePoint LocalEndSurfacePoint(IntrinsicMesh.GetTriangle(EdgeIDandXIdx.TID)[EdgeIDandXIdx.EID]);
+				// add surface point 
+				ResultTraceArray.Add(LocalEndSurfacePoint);
+				LocalStartSurfacePoint = LocalEndSurfacePoint;
+			}
+		}
+
+		// should have processed all edge crossings.
+		checkSlow(IntrinsicEdgesCrossed.Num() == 0);
+		// ended on the last vertex.
+		checkSlow(ResultTraceArray.Last().Position.VertexPosition.VID == OrderedTraceEdgeV.B); 
+	}
+	
+	return MoveTemp(ResultTraceArray);
+}
+
+UE::Geometry::EMeshResult UE::Geometry::FIntrinsicMesh::FlipEdge(int32 EID, FEdgeFlipInfo& EdgeFlipInfo)
+{
+	using namespace FIntrinsicMeshImplUtils;
+
+	if (!IsEdge(EID))
+	{
+		return EMeshResult::Failed_NotAnEdge;
+	}
+
+	if (IsBoundaryEdge(EID))
+	{
+		return EMeshResult::Failed_IsBoundaryEdge;
+	}
+
+	// state before flip, needed when updating the normal coords after the flip
+	const FIndex2i EdgeT  = GetEdgeT(EID);
+	const FIndex3i TAEIDs = GetTriEdges(EdgeT.A);
+	const FIndex3i TBEIDs = GetTriEdges(EdgeT.B);
+	const FIndex2i OppVs  = GetEdgeOpposingV(EID);
+
+	EMeshResult FlipResult = MyBase::FlipEdge(EID, EdgeFlipInfo);
+
+	if (FlipResult == EMeshResult::Ok)
+	{
+		// update the normal coords
+		NormalCoordinates.OnFlipEdge(EdgeT.A, TAEIDs, OppVs.A, EdgeT.B, TBEIDs, OppVs.B, EID);
+
+		// if the flip produced an intrinsic edge that is a segment of a surface edge, update the 
+		// coorindate to -1 to help tracing.
+		if (NormalCoordinates.NormalCoord[EID] == 0 && OppVs.A != OppVs.B)
+		{
+			FSurfacePoint SurfacePoints[2] = {GetVertexSurfacePoint(OppVs.A), GetVertexSurfacePoint(OppVs.B)};
+			if (AreOnSameSurfaceEdge(SurfacePoints[0], SurfacePoints[1], *NormalCoordinates.SurfaceMesh))
+			{
+				NormalCoordinates.NormalCoord[EID] = -1;
+			}
+		}
+	}
+
+	return FlipResult;
+}
+
+
+UE::Geometry::EMeshResult UE::Geometry::FIntrinsicMesh::PokeTriangle(int32 IntrinsicTID, const FVector3d& BaryCoordinates, FPokeTriangleInfo& PokeInfo)
+{
+	// The intrinsic poke  creates 3 new intrinsic edges and "normal coordinates"(the number of crossing surface edges) need to be created for each.
+	// 
+	// also, "roundabout" data is recored for any new intrinsic edge that is adjacent to an intrinsic vertex that is also a surface vertex.
+	// this roundabout data indicates the next surface edge when traveling ccw about the vertex from the new edge.
+	// 
+	// lastly the surface position of the new intrinsic vertex is computed and recored.
+	//
+	// note: to compute the surface position of the new intrinsic vertex, the intersection (convex polygon) between the original intrinsic triangle and the surface triangle 
+	// that supports the new vertex is first identified 
+
+	using namespace FIntrinsicMeshImplUtils;
+
+	if (!IsTriangle(IntrinsicTID))
+	{
+		return EMeshResult::Failed_NotATriangle;
+	}
+
+	// state before the poke
+	const FIndex3i OriginalVIDs            = GetTriangle(IntrinsicTID);
+	const FIndex3i OriginalEdges           = GetTriEdges(IntrinsicTID);
+	const FVector3d OriginalTriEdgeLengths = GetTriEdgeLengths(IntrinsicTID);
+
+	const int32 IndexOf = VectorUtil::Min3Index(OriginalTriEdgeLengths);
+
+	// fail if the smallest side of the triangle is too small
+	//[todo] consider permuting the triangle so the zero side is the longest.
+	if (OriginalTriEdgeLengths[IndexOf] < TMathUtilConstants<double>::ZeroTolerance)
+	{
+		return EMeshResult::Failed_Unsupported;
+	}
+
+	
+	// unwrap the intrinsic triangle to a 2d plane.
+	// edge0 runs along the positive x-axis
+	const FVector2d IntrinsicTri2D[3] = { FVector2d(0.,0.),
+								          FVector2d(OriginalTriEdgeLengths[0], 0.),
+								          ComputeOpposingVert2d(OriginalTriEdgeLengths[0], OriginalTriEdgeLengths[1], OriginalTriEdgeLengths[2]) };
+
+	auto AsR2Position = [&IntrinsicTri2D](const FVector3d& BCoords)
+						{
+							return BCoords[0] * IntrinsicTri2D[0] + BCoords[1] * IntrinsicTri2D[1] + BCoords[2] * IntrinsicTri2D[2];
+						};
+
+	// position of the new vertex relative to the flattened intrinsic triangle.								
+	const FVector2d PokedPos = AsR2Position(BaryCoordinates);
+
+	// Utility: converts distance along a directed edge of an the intrinsic triangle to barycentric coords.
+	auto EdgeDistanceToBaryCoords = [&OriginalTriEdgeLengths](int32 TriSide, double Distance)
+									{
+										const double EdgeLength = OriginalTriEdgeLengths[TriSide];
+										const double Alpha = (EdgeLength > TMathUtilConstants<double>::ZeroTolerance) ? Distance / EdgeLength : 0.5;
+
+										FVector3d Result(0., 0., 0.);
+										Result[TriSide]           = (1. - Alpha);
+										Result[(TriSide + 1) % 3] = Alpha;
+										return Result;
+									};
+
+
+	// -- find the surface edges that cross the intrinsic triangle.
+
+	// for each side of the intrinsic triangle, generate a list of places where the surfaces edge cross the intrinsic edge
+	// storing the location as a FPointCorrespondence
+	
+	// point described both relative to the intrinsic triangle and to the surface mesh.
+	struct FPointCorrespondence
+	{
+		FVector2d LocalPos;       // position relative to flattened intrinsic triangle
+		FVector3d IntrinsicBC;    // barycentric coordinates relative to intrinsic triangle ( duplicates info in LocalPos, but easier to work with)
+		FSurfacePoint SurfacePos; // position relative to surface mesh
+	};
+
+	// for each directed side the intrinsic triangle, recored all the surface edge crossings and their relative locations along the side.
+	TMap<int32, FPointCorrespondence> SurfaceEdgeXings[3];
+	{
+		auto GetSurfaceEdgeCrossingLocations = [&](const int32 Side)
+												{
+													const int32 IntrinsicEID = OriginalEdges[Side];
+													const FDynamicMesh3& SurfaceMesh = *NormalCoordinates.SurfaceMesh;
+
+													const double CoalesceThreshold = 0.;
+													const bool bReverse = (GetEdgeT(IntrinsicEID).A != IntrinsicTID);
+
+													// sequence of surface points corresponding to surface edges intersecting this intrinsic edge. 
+													const TArray<FSurfacePoint> EdgeAsSurfacePoints = TraceEdge(IntrinsicEID, CoalesceThreshold, bReverse);
+
+
+													// record the surface point and the distance along this triangle side to the intersection 
+													// key with the surface edge id.
+													TMap<int32, FPointCorrespondence> CrossingLocations;
+													const int32 NumSPoints = EdgeAsSurfacePoints.Num();
+													if (NumSPoints > 0)
+													{ 
+														bool bTmp;
+														FVector3d Positions[2];
+														Positions[0] = AsR3Position(EdgeAsSurfacePoints[0], SurfaceMesh, bTmp);
+
+														int32 Cur = 1;
+														double AccumulatedDistance = 0.;
+
+														for (int i = 1; i < NumSPoints; ++i)
+														{
+															const FSurfacePoint& SurfacePoint = EdgeAsSurfacePoints[i];
+															Positions[Cur] = AsR3Position(SurfacePoint, SurfaceMesh, bTmp);
+															const double LineElementLength = (Positions[Cur] - Positions[1 - Cur]).Length();
+															AccumulatedDistance += LineElementLength;
+															if (SurfacePoint.PositionType == IntrinsicCorrespondenceUtils::FSurfacePoint::EPositionType::Edge)
+															{
+																const int32 SurfaceEID = SurfacePoint.Position.EdgePosition.EdgeID;
+																const FVector3d IntrinsicBaryCoords = EdgeDistanceToBaryCoords(Side, AccumulatedDistance);
+																FPointCorrespondence XingCorrespondence = {AsR2Position(IntrinsicBaryCoords), IntrinsicBaryCoords, SurfacePoint};
+																CrossingLocations.Add(SurfaceEID, XingCorrespondence);
+															}
+
+															Cur = 1 - Cur;
+														}
+													}
+													return MoveTemp(CrossingLocations);
+												};
+		for (int32 i = 0; i < 3; ++i)
+		{
+			SurfaceEdgeXings[i] = GetSurfaceEdgeCrossingLocations(i);
+		}
+	}
+	
+	// the intersection of a surface edge with the intrinsic triangle.
+	struct FSurfaceEdgeSegment
+	{
+		// end points where this surface edge intersects the boundary of the intrinsic triangle.
+		FPointCorrespondence P0;
+		FPointCorrespondence P1;
+
+		// corresponding surface mesh edge id
+		int32 SurfaceEID = -1;
+	};
+
+	// after the poke, new intrinsic edges will connect the new vertex to the corners of the original triangle.
+	// number these as  {0, 1, 2} according to the corners of the IntrinsicTri2D
+	
+	
+	// compute "normal coordinates", i.e. find the number of times each new intrinsic edge is intersected by surface mesh edge segments
+	// additionally, (for each new intrinsic edge) identify the intersecting surface edge that was closest to the new vertex.
+	
+	int32 NewEdgeNormCoords[3] = { 0, 0, 0 };         
+	FSurfaceEdgeSegment ClosestSurfaceEdgeSegment[3]; // per new intrinsic edge
+	{
+		auto UpdateNormCoordsAndClosest = [&NewEdgeNormCoords, &ClosestSurfaceEdgeSegment](const FPointCorrespondence& P0, const FPointCorrespondence& P1, int32 SurfaceEID, int32 ZeroOneOrTwo)
+											{
+												NewEdgeNormCoords[ZeroOneOrTwo] += 1;
+												FSurfaceEdgeSegment& ClosestSegment = ClosestSurfaceEdgeSegment[ZeroOneOrTwo];
+												const bool bHaveValidClosest = (ClosestSegment.SurfaceEID != -1);
+												if (bHaveValidClosest)
+												{
+													// by construction the surface edge segments can not intersect in the face of an intrinsic triangle
+													// so we need only check a single point of the old segment against the one defined by P0 to P1
+													const FVector2d OldSegmentCenter = 0.5 * (ClosestSegment.P0.LocalPos + ClosestSegment.P1.LocalPos);
+													const double SideTest = UE::Geometry::Orient(P0.LocalPos, P1.LocalPos, OldSegmentCenter);
+													if (SideTest > 0) // old surface segment is to the right of this new one segment 
+													{
+														ClosestSegment.P0 = P0;
+														ClosestSegment.P1 = P1;
+														ClosestSegment.SurfaceEID = SurfaceEID;
+													}
+												}
+												else
+												{
+													// just update since this is the first intersecting surface edge
+													ClosestSegment.P0 = P0;
+													ClosestSegment.P1 = P1;
+													ClosestSegment.SurfaceEID = SurfaceEID;
+												}
+											};
+
+		for (int32 Side = 0; Side < 3; ++Side)
+		{
+
+			const int32 NextSide     = (Side + 1) % 3;
+			const int32 NextNextSide = (Side + 2) % 3;
+
+			for (const TPair<int32,  FPointCorrespondence>& XingPair : SurfaceEdgeXings[Side])
+			{
+				const int32& SurfaceEID        = XingPair.Key;
+				const FPointCorrespondence& P0 = XingPair.Value;
+
+				if ( FPointCorrespondence* P1Ptr = SurfaceEdgeXings[NextSide].Find(SurfaceEID)) // the surface edge exits the next side
+				{
+					const FPointCorrespondence& P1 = *P1Ptr;
+					const double SignedArea = UE::Geometry::Orient(P0.LocalPos, P1.LocalPos, PokedPos); //same as (P1-P0)X(PokedPos - P0)
+
+					if (SignedArea <= 0) // surface segment crosses two new intrinsic edges.
+					{
+						// [todo] better treatment for == case?
+						UpdateNormCoordsAndClosest(P0, P1, SurfaceEID, Side);
+						UpdateNormCoordsAndClosest(P0, P1, SurfaceEID, NextNextSide);
+					}
+					if (SignedArea > 0) // surface segment crosses one new intrinsic edges.
+					{
+						UpdateNormCoordsAndClosest(P1, P0, SurfaceEID, NextSide);
+					}
+
+				}
+				else if (SurfaceEdgeXings[NextNextSide].Contains(SurfaceEID) == false) // the surface edge must exit the opp vertex
+				{
+					FVector3d P1BC(0., 0., 0.);  P1BC[NextNextSide] = 1.;
+
+					const FPointCorrespondence P1 = { IntrinsicTri2D[NextNextSide], P1BC, IntrinsicVertexPositions[OriginalVIDs[NextNextSide]] };;
+					const double SignedArea = UE::Geometry::Orient(P0.LocalPos, P1.LocalPos, PokedPos); // same as (P1-P0)X(PokedPos - P0)
+
+					if (SignedArea <= 0) // surface segment crosses one new intrinsic edges.
+					{
+						// [todo] better treatment for == case?
+						UpdateNormCoordsAndClosest(P0, P1, SurfaceEID, Side);
+					}
+					if (SignedArea > 0) // surface segment crosses one new intrinsic edges.
+					{
+						UpdateNormCoordsAndClosest(P1, P0, SurfaceEID, NextSide);
+					}
+				}
+			}
+		}
+	}
+
+
+	// intersection of the intrinsic mesh triangle (prior to poke) and the surface mesh triangle that ultimately supports the new vertex
+	// forms a convex polygon.   Convert the vertices of this polygon to FPointCorresponce. 
+	const TArray<FPointCorrespondence> BoundingConvexPolyVerts = [&]
+																{
+																	TArray<FPointCorrespondence> TmpBoundaryPoints;
+																	// method to add Correspondence points.  Makes sure we don't add a point twice
+																	// note, only need to check surface verts due to our usage and the fact surface edges only meet at surface verts. 
+																	TSet<int32> VisitedSurfaceVerts;
+																	auto AddPointCorrespondence = [&](const FPointCorrespondence& PointCorrespondence)
+																									{
+																										const FSurfacePoint& SP = PointCorrespondence.SurfacePos;
+																										if (IsVertexPoint(SP))
+																										{
+																											const int32 SurfaceVID = SP.Position.VertexPosition.VID;
+																											if (!VisitedSurfaceVerts.Contains(SurfaceVID))
+																											{
+																												TmpBoundaryPoints.Add(PointCorrespondence);
+																												VisitedSurfaceVerts.Add(SurfaceVID);
+																											}
+																										}
+																										else
+																										{
+																											TmpBoundaryPoints.Add(PointCorrespondence);
+																										}
+
+																									};
+																	
+																	TSet<int32> VisitedSurfaceEdges;
+																	for (int32 i = 0; i < 3; ++i)
+																	{
+																		// if it exists, this surface edge blocks
+																		// the new intrinsic vertex's view of the i-th original tri corner.
+																		const FSurfaceEdgeSegment& SurfaceEdgeSegment = ClosestSurfaceEdgeSegment[i];
+																		const int32 SurfaceEdgeID = SurfaceEdgeSegment.SurfaceEID;
+
+																		if (SurfaceEdgeID == -1) // no edge blocks..
+																		{
+																			const FVector2d& CornerLocalPos = IntrinsicTri2D[i];
+																			const int32 VID = OriginalVIDs[i];
+																			const FSurfacePoint& CornerSP = IntrinsicVertexPositions[VID];
+																			FVector3d IntrinsicBC(0., 0., 0.);
+																			IntrinsicBC[i] = 1.;
+
+																			FPointCorrespondence PointCorrespondence = { CornerLocalPos, IntrinsicBC, CornerSP };
+																			AddPointCorrespondence(PointCorrespondence);
+																		}
+																		else
+																		{
+																			if (!VisitedSurfaceEdges.Contains(SurfaceEdgeID))
+																			{
+																				VisitedSurfaceEdges.Add(SurfaceEdgeID);
+
+																				AddPointCorrespondence(SurfaceEdgeSegment.P0);
+																				AddPointCorrespondence(SurfaceEdgeSegment.P1);
+																			}
+																		}
+																	}
+																	return MoveTemp(TmpBoundaryPoints);
+																}();
+
+
+	
+	// Identify the surface triangle that is common to the bounding convex polygon verts.
+	const int32 SurfaceTID = [&]
+							{
+								const FDynamicMesh3& SurfaceMesh = *NormalCoordinates.SurfaceMesh;
+
+								// special case: are the points just the corners of a surface triangle? 
+								if (BoundingConvexPolyVerts.Num() == 3)
+								{
+									if (   IsVertexPoint(BoundingConvexPolyVerts[0].SurfacePos)
+										&& IsVertexPoint(BoundingConvexPolyVerts[1].SurfacePos)
+										&& IsVertexPoint(BoundingConvexPolyVerts[2].SurfacePos) )
+									{
+										int32 SurfaceVIDs[3] = { BoundingConvexPolyVerts[0].SurfacePos.Position.VertexPosition.VID,
+											                     BoundingConvexPolyVerts[1].SurfacePos.Position.VertexPosition.VID,
+											                     BoundingConvexPolyVerts[2].SurfacePos.Position.VertexPosition.VID
+										};
+
+										return SurfaceMesh.FindTriangle(SurfaceVIDs[0], SurfaceVIDs[1], SurfaceVIDs[2]);
+									}
+								}
+
+								// find surface triangles corresponding to one of the surface points.
+								// note: a surface point will correspond to 1, 2, or many surface triangles (face point, edge point, or vertex point).
+								const TArray<int32> CandidateTIDs = [&]
+																	{
+																		TArray<int32> TmpTIDs;
+																		// are any of the surface point already on a surface triangle face?
+																		int32 TID = -1;
+																		for (const FPointCorrespondence& BoundaryPoint : BoundingConvexPolyVerts)
+																		{
+																			const FSurfacePoint& SurfacePos = BoundaryPoint.SurfacePos;
+																			if (IsFacePoint(SurfacePos))
+																			{
+																				TID = SurfacePos.Position.TriPosition.TriID;
+																				break;
+																			}
+																		}
+
+																		if (TID != -1)
+																		{
+																			TmpTIDs.Add(TID);
+																		}
+																		else
+																		{
+																			const FSurfacePoint& SurfacePos = BoundingConvexPolyVerts[0].SurfacePos;
+																			// get all the surface triangles adjacent to this surface point
+																			TmpTIDs = GetAdjacentTriangles(SurfacePos, SurfaceMesh);
+																		}
+
+																		return MoveTemp(TmpTIDs);
+																	}();
+
+								// determine which candidate triangle is actually adjacent to all the surface points.
+								const int32 MutuallyAdjacentTID = [&]
+																	{
+																		for (int32 CandidateTID : CandidateTIDs)
+																		{
+																			bool bIsAdjacentToAll = true;
+																			// check adjacency to the boundary points.
+																			for (const FPointCorrespondence& BoundaryPoint : BoundingConvexPolyVerts)
+																			{
+																				const FSurfacePoint& SurfacePos = BoundaryPoint.SurfacePos;
+																				bIsAdjacentToAll = bIsAdjacentToAll && IsOnSurfaceTriangle(SurfacePos, CandidateTID, SurfaceMesh);
+																			}
+
+																			if (bIsAdjacentToAll)
+																			{
+																				return CandidateTID;
+																			}
+																		}
+																		return -1;
+																	}();
+
+								return MutuallyAdjacentTID;
+							}();
+
+	if (SurfaceTID == -1)
+	{
+		return EMeshResult::Failed_Unsupported;
+	}
+
+
+
+	// construct the surface barycentric coords for the poked position.
+	const FVector3d PokedSurfaceBaryCoords = [&]
+												{
+													const int32 NumBoundaryPts = BoundingConvexPolyVerts.Num();
+													// In the barycentric space of the intrinsic triangle:
+													// find the (min norm) linear combination convex-poly verticies that is equivalent to the new vertex location
+													TArray<double> MinNormSolution;
+													TArray<double> MatRows[3];
+													for (int32 i = 0; i < 3; ++i)
+													{
+														MatRows[i].AddUninitialized(NumBoundaryPts);
+
+														for (int32 j = 0; j < NumBoundaryPts; ++j)
+														{
+															const FVector3d& IntrinsicBarycentric = BoundingConvexPolyVerts[j].IntrinsicBC;
+															MatRows[i][j] = IntrinsicBarycentric[i];
+														}
+													}
+
+													const bool bValidInterpolation = ComputeMinNormSolution(MatRows[0], MatRows[1], MatRows[2], BaryCoordinates, MinNormSolution);
+
+													// convert the surface points to barycentric coords relative to the surface triangle.
+													TArray<FVector3d> SurfaceBaryCoords;
+													{
+														const FDynamicMesh3& SurfaceMesh = *NormalCoordinates.SurfaceMesh;
+														SurfaceBaryCoords.Reserve(NumBoundaryPts);
+														for (const FPointCorrespondence& BoundaryPoint : BoundingConvexPolyVerts)
+														{
+															const FSurfacePoint& SurfacePos = BoundaryPoint.SurfacePos;
+															SurfaceBaryCoords.Add(AsBarycenteric(SurfacePos, SurfaceTID, SurfaceMesh));
+														}
+													}
+
+													// interpolate the surface barycentric coordinates (using the coefficients from the intrinsic barycenterics)
+													// to construct the surface barycentric coordinates of the new vertex location.
+													// note: if interpolation failed just use the center of the polygon.
+													FVector3d SurfaceBarycoods(0., 0., 0.);
+													if (bValidInterpolation)
+													{
+														for (int32 i = 0; i < NumBoundaryPts; ++i)
+														{
+															SurfaceBarycoods += MinNormSolution[i] * SurfaceBaryCoords[i];
+														}
+													}
+													else
+													{
+														// just use the "center" of the convex polygon
+														for (int32 i = 0; i < NumBoundaryPts; ++i)
+														{
+															SurfaceBarycoods += SurfaceBaryCoords[i];
+														}
+														SurfaceBarycoods *= 1. / double(NumBoundaryPts);
+													}
+
+													return SurfaceBarycoods;
+												}();
+
+	// surface point for the new (poked) vertex
+	const FSurfacePoint PokedSurfacePoint(SurfaceTID, PokedSurfaceBaryCoords);
+
+	// update the topology.
+
+	// Add a new vertex and faces to the IntrinsicMesh.   
+	// Note: the r3 position will be wrong initially since this poke will just interpolate the corners of the intrinsic tri.
+	//       we fix this position as the last step in this function
+	EMeshResult PokeResult = MyBase::PokeTriangle(IntrinsicTID, BaryCoordinates, PokeInfo);
+	if (PokeResult != EMeshResult::Ok)
+	{
+		return PokeResult;
+	}
+
+
+	FIndex3i NewTris(IntrinsicTID, PokeInfo.NewTriangles[0], PokeInfo.NewTriangles[1]);
+	const int32 NewVID = PokeInfo.NewVertex;
+
+	// Need to update intrinsic information for the 3 triangles that resulted from the poke
+	// 1) a)update the number of surface edge crossings for each new intrinsic edge
+	//    b)update the roundabout data for any new intrinsic edge that is adj to a surface vertex
+	// 2) update the surface position of the new vertex.
+
+	{
+		// set/update number of surface edge crossings for the new edges. 
+		{
+			// original tri is and verts (a, b,c) and edges (a2b, b2c, c2a)
+			// is updated by the poke to have verts (a, b, new) and edges (a2b, b2new, new2a)
+			const FIndex3i T0EIDs = GetTriEdges(NewTris[0]);
+			const FIndex3i T1EIDs = GetTriEdges(NewTris[1]);
+	
+			// set the correct number of surface edges the cross each new intrinsic edge.
+			NormalCoordinates.NormalCoord.InsertAt(NewEdgeNormCoords[0], T0EIDs[2]); // a2new
+			NormalCoordinates.NormalCoord.InsertAt(NewEdgeNormCoords[2], T1EIDs[1]); // b2new
+			NormalCoordinates.NormalCoord.InsertAt(NewEdgeNormCoords[1], T0EIDs[1]); // c2new
+		}
+
+		// update the roundabout information.
+		{
+			// edge order: a2b, b2c, c2a
+			const FIndex3i  OriginalRO = NormalCoordinates.RoundaboutOrder[IntrinsicTID];
+
+			// copy the roundabout information to the new triangles with default -1 for the new directed edges.
+
+			FIndex3i NewTrisRO[3] = { FIndex3i(OriginalRO[0], -1, -1) // NewTri[0]   edge order: a2b, b2new, new2a
+									 ,FIndex3i(OriginalRO[1], -1, -1) // NewTri[1]   edge order: b2c, c2new, new2b
+									 ,FIndex3i(OriginalRO[2], -1, -1) // NewTri[2]   edge order: c2a, a2new, new2c
+									};
+
+			// update roundabout order for any new directed edge that starts at an implicit vertex that is also a surface vertex 	
+			// Note: by construction the new intrinsic vertex does not correspond to a vertex on the surface mesh so we need only consider a2new, b2new, c2new.
+
+
+			const FSurfacePoint OriginalTriSPs[3] = { IntrinsicVertexPositions[OriginalVIDs[0]],
+													  IntrinsicVertexPositions[OriginalVIDs[1]],
+													  IntrinsicVertexPositions[OriginalVIDs[2]] };
+
+			for (int32 i = 0; i < 3; ++i)
+			{
+				if (IsVertexPoint(OriginalTriSPs[i]))
+				{
+					const FIndex3i TriEIDs = GetTriEdges(NewTris[i]);
+					const int32 Radj = NewTrisRO[i][0];                                            // {R_ab,  R_bc,  R_ca }  for i = 0, 1, or 2
+					const int32 Eopp = NormalCoordinates.NumCornerEmanatingRefEdges(TriEIDs, 0);  // {Ebn_a, Ecn_b, Ean_c} for i = 0, 1, or 2
+					const int32 Rnew = (Radj + Eopp) % NormalCoordinates.RefVertDegree[OriginalVIDs[i]];
+
+					NewTrisRO[(i + 2) % 3][1] = Rnew;
+				}
+			}
+
+			NormalCoordinates.RoundaboutOrder.InsertAt(NewTrisRO[0], NewTris[0]);
+			NormalCoordinates.RoundaboutOrder.InsertAt(NewTrisRO[1], NewTris[1]);
+			NormalCoordinates.RoundaboutOrder.InsertAt(NewTrisRO[2], NewTris[2]);
+		}
+	}
+	// update the position of the new vertex ( R3 position and the surface point )
+	bool bIsValid;
+	const FVector3d R3Pos = IntrinsicCorrespondenceUtils::AsR3Position(PokedSurfacePoint, *NormalCoordinates.SurfaceMesh, bIsValid);
+	Vertices[NewVID] = R3Pos;
+	IntrinsicVertexPositions.InsertAt(PokedSurfacePoint, NewVID);
+
+	return EMeshResult::Ok;
+}
+
+
+
+UE::Geometry::EMeshResult UE::Geometry::FIntrinsicMesh::SplitEdge(int32 EdgeAB, FEdgeSplitInfo& SplitInfo, double SplitParameterT)
+{
+
+	using namespace FIntrinsicMeshImplUtils;
+
+	if (!IsEdge(EdgeAB))
+	{
+		return EMeshResult::Failed_NotAnEdge;
+	}
+
+	const FDynamicMesh3& SurfaceMesh = *NormalCoordinates.SurfaceMesh;
+
+	const FEdge OriginalEdge = GetEdge(EdgeAB);
+	// is the target intrinsic edge equivalent to a segment of a surface edge?
+	const bool bIsOnSurfaceEdge   = NormalCoordinates.IsSurfaceEdgeSegment(EdgeAB);
+	// say new vertex is f.
+	const int32 TID0 = OriginalEdge.Tri[0];
+	const int32 IndexOfe = GetTriEdges(TID0).IndexOf(EdgeAB);
+	const int32 TID1 = OriginalEdge.Tri[1];
+
+	const bool bIsBoundary = (TID1 == -1);
+	const int32 IndexOf1e = (!bIsBoundary) ? GetTriEdges(TID1).IndexOf(EdgeAB) : -1;
+
+	// Info about the original T0 tri, reordered to make the split edge the first edge..
+	const FIndex3i OriginalT0VIDs         = Permute(IndexOfe, GetTriangle(TID0));
+	const FIndex3i OriginalT0Edges        = Permute(IndexOfe, GetTriEdges(TID0));
+	const FVector3d OriginalT0EdgeLengths = Permute(IndexOfe, GetTriEdgeLengths(TID0));    // as (|ab|, |bc|, |ca|)
+	
+	// which, if any, of the original intrinsic edges are coincident with surface edges
+	// in the order { edgeAB, edgeBC, edgeCA, edgeDB }
+	const bool Kroneckers[4] = { NormalCoordinates.IsSurfaceEdgeSegment(OriginalT0Edges[0]), 
+	                             NormalCoordinates.IsSurfaceEdgeSegment(OriginalT0Edges[1]), 
+							     NormalCoordinates.IsSurfaceEdgeSegment(OriginalT0Edges[2]),
+								 (bIsBoundary) ? false : NormalCoordinates.IsSurfaceEdgeSegment(Permute(IndexOf1e, GetTriEdges(TID1))[2])};
+
+	
+
+	const double IntrinsicEdgeSplitDistance = SplitParameterT * OriginalT0EdgeLengths[0];
+
+	// description of surface edge intersection with intrinsic edge AB
+	struct FEdgeDistanceCorrespondence
+	{
+		int32  SurfaceEID = -1;      // surface edge
+		double DistOnEdge;           // distance measured along the intrinsic edge
+		FSurfacePoint  SurfacePoint; // location on the surface mesh
+
+	};
+
+	auto AsR3Pos = [&SurfaceMesh](const FSurfacePoint& SP)
+				{
+					bool bTemp;
+					return AsR3Position(SP, SurfaceMesh, bTemp);
+				};
+
+	// -- prior to actually doing the edge split, we find the two points (surface edge intersections) that bracket the proposed edge split location
+	//    and pre-compute the normal coordinates for the 4 edges that are changed by the split
+
+	// initialize the bounding points with the ends of the edge being split. 
+	FEdgeDistanceCorrespondence BoundingPoints[2] = { {-1, 0., IntrinsicVertexPositions[OriginalT0VIDs[0]]}, 
+													  {-1, OriginalT0EdgeLengths[0], IntrinsicVertexPositions[OriginalT0VIDs[1]]} };
+
+	int32 UpdatedNormCoord = 0; // after split edgeAB becomes edge [a, f]													
+	int32 NewEdgeNormCoords[3] = { 0, 0, 0}; //N_fb, N_fc, N_fd  
+	
+	{
+		if (NormalCoordinates.NumEdgeCrossing(EdgeAB) == 0) // no surface edges cross the intrinsic edge being split
+		{
+			// potentially propagating -1 from NormalCoord[EdgeAB] may explicitly mark these as a segment of EdgeAB  
+			UpdatedNormCoord     = NormalCoordinates.NormalCoord[EdgeAB];  // becomes edgeAF
+			NewEdgeNormCoords[0] = NormalCoordinates.NormalCoord[EdgeAB];  // new edge, edgeFB
+
+
+			NewEdgeNormCoords[1] = FMath::Max(NormalCoordinates.NumEdgeCrossing(OriginalT0Edges[1]), NormalCoordinates.NumEdgeCrossing(OriginalT0Edges[2]));
+
+			// poke makes t0 = {a, f, c} and  t2 = { f, b, c}  ( note: t1 = {f, a, d}  t3 = f, d, b}  don't exist in the boundary edge case )
+		
+			if (!bIsBoundary)
+			{
+				// Info about the original T1 tri, reordered to make the split edge the first edge..
+			
+				const FIndex3i OriginalT1Edges        = Permute(IndexOf1e, GetTriEdges(TID1));
+				const FVector3d OriginalT1EdgeLengths = Permute(IndexOf1e, GetTriEdgeLengths(TID1));   // as (|ba|, |ad|, |db})
+
+				NewEdgeNormCoords[2] = FMath::Max(NormalCoordinates.NumEdgeCrossing(OriginalT1Edges[1]), NormalCoordinates.NumEdgeCrossing(OriginalT1Edges[2]));
+
+			}
+		} 
+		else
+		{
+		
+
+			const int32 IntrinsicEID = OriginalT0Edges[0];
+		
+
+			const FIndex3i OriginalT1Edges = (TID1 != -1) ? Permute(IndexOf1e, GetTriEdges(TID1)) : FIndex3i(-1, -1, -1);
+
+
+		
+
+			// sequence of surface points corresponding to surface edges intersecting intrinsic edge  ab. 
+			// ordered relative to  OriginalEdge.Tri[0];
+			const TArray<FSurfacePoint> ABEdgeAsSurfacePoints = [&] 
+																{
+																	const double CoalesceThreshold = 0.;
+																	const bool bReverse = false;
+																	return TraceEdge(IntrinsicEID, CoalesceThreshold, bReverse);
+																}();
+
+
+			auto GetExitEdge = [&](const int32 TriID, const int32 P)
+									{
+										int32 ExitEID = -1;
+										if (TriID != -1)
+										{
+											const TTuple<int32, int32> ExitEIDandP = FNormalCoordSurfaceTraceImpl::GetCrossingExit(*this, NormalCoordinates, TriID, EdgeAB, P);
+											ExitEID = (ExitEIDandP.Get<1>() != 0) ? ExitEIDandP.Get<0>() : -1;
+										}
+										return ExitEID;
+									};
+
+		
+			// account for surface edges that cross the intrinsic edgeAB at the endpoints and the new edges, ie edgeFC (or edgeFD)
+			{
+				const int32 Ebc_a = NormalCoordinates.NumCornerEmanatingRefEdges(OriginalT0Edges, 0); // surface edges that start at 'a' and exit side bc
+				const int32 Eca_b = NormalCoordinates.NumCornerEmanatingRefEdges(OriginalT0Edges, 1); // surface edges that start at 'b' and exit side ca
+
+				NewEdgeNormCoords[1] += Eca_b + Eca_b;  // new edge - edgeFC
+
+				if (!bIsBoundary)
+				{
+					const int32 Ead_b = NormalCoordinates.NumCornerEmanatingRefEdges(OriginalT1Edges, 0); // surface edges that start at 'b' and exit side ad
+					const int32 Edb_a = NormalCoordinates.NumCornerEmanatingRefEdges(OriginalT1Edges, 1); // surface edges that start at 'a' and exit side db
+
+					NewEdgeNormCoords[2] += Ead_b + Edb_a;  // new edge - edgeFD
+				}
+			}
+
+			// count the number of surface edges that cross both edgeAB and the edgeFC, and those that cross edgeAB and edgeFD
+			// Also identify the two surface edge crossing of edgeAB that bracket the vertex produced by the split.
+			{
+				// double buffer.
+				int32 Cur = 1;
+				FVector3d Positions[2];
+		
+				
+				Positions[0] = AsR3Pos(ABEdgeAsSurfacePoints[0]);
+
+				const int32 NumSPoints = ABEdgeAsSurfacePoints.Num();
+
+		
+		
+				double AccumulatedDistance = 0;
+				for (int32 i= 1; i < NumSPoints - 1; ++i) // the first and last points are the implicit edge start/end. 
+				{
+					const FSurfacePoint& SurfacePoint = ABEdgeAsSurfacePoints[i];
+					Positions[Cur] = AsR3Pos(SurfacePoint);
+					const double LineElementLength = (Positions[Cur] - Positions[1 - Cur]).Length();
+					AccumulatedDistance += LineElementLength;
+			
+					Cur = 1 - Cur; // swap double buffer.
+
+					// which side does this surface edge exit?
+			
+					const int32 T0ExitEID = GetExitEdge(TID0, i);
+					const int32 T1ExitEID = GetExitEdge(TID1, NumSPoints - i);
+			
+					if (AccumulatedDistance < IntrinsicEdgeSplitDistance) // crossing to left of split
+					{	
+						UpdatedNormCoord += 1; // edgeAF
+
+						checkSlow(IsEdgePoint(SurfacePoint));
+						int32 SurfaceEID = SurfacePoint.Position.EdgePosition.EdgeID;
+						BoundingPoints[0] = {SurfaceEID, AccumulatedDistance, SurfacePoint};
+				
+						if (OriginalT0Edges[1] == T0ExitEID)
+						{ 
+							NewEdgeNormCoords[1] +=1; // new edge - edgeFC
+						}
+						if (T1ExitEID != -1 && OriginalT1Edges[2] == T1ExitEID)
+						{
+							NewEdgeNormCoords[2] += 1; // new edge - edgeFD
+						}
+					}
+					else // crossing to right of split
+					{ 
+						NewEdgeNormCoords[0] += 1;  // edgeFB
+
+						if (AccumulatedDistance < BoundingPoints[1].DistOnEdge)
+						{
+							checkSlow(IsEdgePoint(SurfacePoint));
+							int32 SurfaceEID = SurfacePoint.Position.EdgePosition.EdgeID;
+							BoundingPoints[1] = { SurfaceEID, AccumulatedDistance, SurfacePoint };
+						}
+
+						if (OriginalT0Edges[2] == T0ExitEID)
+						{
+							NewEdgeNormCoords[1] += 1; // new edge - edgeFC
+						}
+						if (T1ExitEID != -1 && OriginalT1Edges[1] == T1ExitEID)
+						{
+							NewEdgeNormCoords[2] += 1; // new edge - edgeFD
+						}
+
+					}
+				}
+			}
+		}
+	}
+	
+	// do the actual edge split
+	UE::Geometry::EMeshResult SplitResult = MyBase::SplitEdge(EdgeAB, SplitInfo, SplitParameterT);
+
+	if (SplitResult != EMeshResult::Ok)
+	{
+		return SplitResult;
+	}
+	
+	// update the number of edge crossings for the edge configuration resulting from the split
+	// original edge[a,b] is now [a,f] new edges are [f,b], [f,c] and [f,d]
+	// Note: this must be done before updating the roundabout information
+	NormalCoordinates.NormalCoord[SplitInfo.OriginalEdge] = UpdatedNormCoord;
+	NormalCoordinates.NormalCoord.InsertAt(NewEdgeNormCoords[0], SplitInfo.NewEdges[0]);
+	NormalCoordinates.NormalCoord.InsertAt(NewEdgeNormCoords[1], SplitInfo.NewEdges[1]);
+	if (!bIsBoundary)
+	{
+		NormalCoordinates.NormalCoord.InsertAt(NewEdgeNormCoords[2], SplitInfo.NewEdges[2]);
+	}
+
+	// fix the roundabout orders
+	if (!bIsBoundary) // for 4 triangles. 
+	{ 
+		
+		const FIndex3i OriginalT0RO   = Permute(IndexOfe, NormalCoordinates.RoundaboutOrder[TID0]);  // {ROa2b, ROb2c, ROc2a}
+		const FIndex3i OriginalT1RO   = Permute(IndexOf1e, NormalCoordinates.RoundaboutOrder[TID1]); // {ROb2a, ROa2d, ROd2b}
+		const FIndex3i OriginalT1VIDs = Permute(IndexOf1e, GetTriangle(TID1)); // {b, a, d} 
+
+		// initialize with the correct RO information for the outer diamond c2a, a2d, b2c, d2b
+		const int32 Rc2a = OriginalT0RO[2];
+		const int32 Ra2d = OriginalT1RO[1];
+		const int32 Rd2b = OriginalT1RO[2];
+		const int32 Rb2c = OriginalT0RO[1];
+		const int32 Ra2b = OriginalT0RO[0];
+
+		// T0 and T1 after split {a,f, c} and {f, a, d}.. although permuted to the original order.
+		FIndex3i UpdatedRO[2] = { FIndex3i(-1, -1, Rc2a), 
+								  FIndex3i(-1, Ra2d, -1) };
+		// tris {f,b, c}  and {f, d, b}
+		FIndex3i NewTrisRO[2] = { FIndex3i(-1, Rb2c, -1),  
+								  FIndex3i(-1, Rd2b, -1) };
+		
+
+		// correct the roundabout order consistent with the new point f being just inside the original T0
+
+		// directed edge a2f RO:  UpdatedRO[0][0]
+		{
+			const int32 A = OriginalT0VIDs[0];
+			if (IsVertexPoint(IntrinsicVertexPositions[A]))
+			{
+				UpdatedRO[0][0] = Ra2b;
+			}
+		}
+		// directed edge b2f RO: NewTriRO[1][2] 
+		{
+			const int32 B = OriginalT0VIDs[1];
+			if (IsVertexPoint(IntrinsicVertexPositions[B]))
+			{
+				const FIndex3i NewTri0Edges = GetTriEdges(SplitInfo.NewTriangles[0]);
+				const int32 Kronecker_bc = (int32)Kroneckers[1];
+				const int32 Ecf_b = NormalCoordinates.NumCornerEmanatingRefEdges(NewTri0Edges, 1); // num surface edges coming from b and crossing edgeCF.
+				NewTrisRO[1][2] = (Rb2c + Ecf_b + Kronecker_bc)  % NormalCoordinates.RefVertDegree[B];
+			}
+		}
+		// directed edges c2f RO: NewTriRO[0][2]
+		{
+			const int32 C = OriginalT0VIDs[2];
+			if (IsVertexPoint(IntrinsicVertexPositions[C]))
+			{
+				const FIndex3i T0UpdatedEdges = Permute(IndexOfe, GetTriEdges(TID0));
+				const int32 Kronecker_ca  = (int32)Kroneckers[2];
+				const int32 Eaf_c = NormalCoordinates.NumCornerEmanatingRefEdges(T0UpdatedEdges, 2); // num surface edges coming from c and crossing edgeAF.
+				const int32 Rc2f = (Rc2a + Eaf_c + Kronecker_ca) % NormalCoordinates.RefVertDegree[C];
+				NewTrisRO[0][2] = Rc2f;
+			}
+		}
+		// directed edge d2f RO: UpdatedRO[1][2]
+		{
+			const int32 D = OriginalT1VIDs[2];
+			if (IsVertexPoint(IntrinsicVertexPositions[D]))
+			{
+				const FIndex3i NewTri1Edges = GetTriEdges(SplitInfo.NewTriangles[1]);
+				const int32 Kronecker_bd = (int32)Kroneckers[3];
+				const int32 Ebf_d = NormalCoordinates.NumCornerEmanatingRefEdges(NewTri1Edges, 1); // num surface edges coming from d and crossing edgeBF
+				const int32 Rd2f  = (Rd2b + Ebf_d + Kronecker_bd) % NormalCoordinates.RefVertDegree[D];
+			
+				UpdatedRO[1][2] = Rd2f;
+			}
+		}
+		NormalCoordinates.RoundaboutOrder[TID0] = Permute( (3-IndexOfe)%3,  UpdatedRO[0]);  
+		NormalCoordinates.RoundaboutOrder[TID1] = Permute( (3-IndexOf1e)%3, UpdatedRO[1]); 
+		NormalCoordinates.RoundaboutOrder.InsertAt(NewTrisRO[0], SplitInfo.NewTriangles[0]);
+		NormalCoordinates.RoundaboutOrder.InsertAt(NewTrisRO[1], SplitInfo.NewTriangles[1]);
+	}
+	else // only 2 triangles.
+	{
+		const FIndex3i OriginalT0RO = Permute(IndexOfe, NormalCoordinates.RoundaboutOrder[TID0]);
+		// initialize with the correct RO information for the outer diamond c2a, a2d, b2c, d2b
+		
+		const int32 Ra2b = OriginalT0RO[0];
+		const int32 Rb2c = OriginalT0RO[1];
+		const int32 Rc2a = OriginalT0RO[2];
+		
+
+		FIndex3i UpdatedRO = FIndex3i(-1, -1, Rc2a);  // {ROa2f, ROf2c, ROc2a}
+		FIndex3i NewTriRO  = FIndex3i(-1, Rb2c, -1);  // {Rof2b, ROb2c, ROc2f} tris {f,b, c}
+		
+		// directed edge a2f RO:  UpdatedRO[0][0]
+		{
+			const int32 A = OriginalT0VIDs[0];
+			if (IsVertexPoint(IntrinsicVertexPositions[A]))
+			{
+				UpdatedRO[0] = Ra2b;
+			}
+		}
+		// directed edges c2f RO: NewTriRO[2]
+		{
+			const int32 C = OriginalT0VIDs[2];
+			if(IsVertexPoint(IntrinsicVertexPositions[C]))
+			{
+				const FIndex3i T0UpdatedEdges = Permute(IndexOfe, GetTriEdges(TID0));
+				const int32 Kronecker_ca  = (int32)Kroneckers[2];
+				const int32 Eaf_c = NormalCoordinates.NumCornerEmanatingRefEdges(T0UpdatedEdges, 2); // num surface edges coming from c and crossing edgeAF.
+				const int32 Rc2f  = (Rc2a + Eaf_c + Kronecker_ca) % NormalCoordinates.RefVertDegree[C];
+				NewTriRO[2] = Rc2f;
+			}
+		}
+
+		NormalCoordinates.RoundaboutOrder[TID0] = Permute((3-IndexOfe)%3, UpdatedRO); 
+		NormalCoordinates.RoundaboutOrder.InsertAt(NewTriRO, SplitInfo.NewTriangles[0]);
+	}
+
+	// update the location of the new vertex, both as surface position and R3.
+	{
+		// the two surface points on the split edge that bracket the new vertex
+		const FSurfacePoint& LeftSP  = BoundingPoints[0].SurfacePoint;
+		const FSurfacePoint& RightSP = BoundingPoints[1].SurfacePoint;
+
+
+		if (bIsOnSurfaceEdge)
+		{
+			//    the bounding surface points are either surface vert points or surface edge points
+
+			// find the surface edge
+			const int32 SurfaceEID = [&]
+									{
+										int32 SurfaceEID = -1;
+										if (IsEdgePoint(LeftSP))
+										{
+											SurfaceEID = LeftSP.Position.EdgePosition.EdgeID;
+										} 
+										else if (IsEdgePoint(RightSP))
+										{
+											SurfaceEID = RightSP.Position.EdgePosition.EdgeID;
+										}
+										else
+										{
+											// neither bounding point was an "edge" point, 
+											// so they must both be "vertex" points
+											checkSlow(IsVertexPoint(LeftSP) && IsVertexPoint(RightSP));
+											const int32 SurfaceVIDs[2] = { LeftSP.Position.VertexPosition.VID,  RightSP.Position.VertexPosition.VID};
+											SurfaceEID = SurfaceMesh.FindEdge(SurfaceVIDs[0], SurfaceVIDs[1]);
+											checkSlow(SurfaceEID != -1);
+										}
+										return SurfaceEID;
+									}();
+
+			// compute alpha location relative to the surface edge.
+			const FEdge  SurfaceEdge    = SurfaceMesh.GetEdge(SurfaceEID);
+			const FIndex2i SurfaceEdgeV = SurfaceEdge.Vert;
+			
+			// orient the surface edge in the same direction as the intrinsic edge ( the intrinsic edge was ordered by its T0)
+			// note: the intrinsic edge vertices can not be the same because the intrinsic edge is a segment of the surface edge
+
+			const int32 StartV = [&]
+								{
+									if (IsVertexPoint(LeftSP))
+									{
+										return (int32)(LeftSP.Position.VertexPosition.VID == SurfaceEdgeV.B);
+									}
+									else if (IsVertexPoint(RightSP))
+									{
+										return (int32)(RightSP.Position.VertexPosition.VID != SurfaceEdgeV.B);
+									}
+									
+									checkSlow(IsEdgePoint(RightSP) && IsEdgePoint(LeftSP));
+									return (int32)(RightSP.Position.EdgePosition.Alpha > LeftSP.Position.EdgePosition.Alpha); 
+								}();
+
+			// surface edge end points, ordered correctly.	
+			const FVector3d SurfaceEdgeEndPts[2] = {SurfaceMesh.GetVertex(SurfaceEdgeV[StartV]), SurfaceMesh.GetVertex(SurfaceEdgeV[1-StartV])};
+
+			const double SurfaceEdgeLength  = ( SurfaceEdgeEndPts[1] - SurfaceEdgeEndPts[0] ).Length();
+			const double SurfaceDistToSplit = ( AsR3Pos(LeftSP) - SurfaceEdgeEndPts[0] ).Length()  + IntrinsicEdgeSplitDistance;
+			const double R3Alpha            = 1. - SurfaceDistToSplit / SurfaceEdgeLength;
+			checkSlow(SurfaceDistToSplit <= SurfaceEdgeLength);
+
+		
+
+			// compute the R3 position of the new vertex
+			const FVector3d R3Pos = R3Alpha * SurfaceEdgeEndPts[0] + (1. - R3Alpha) * SurfaceEdgeEndPts[1];
+
+			// record the R3 position and the surface position.
+			Vertices[SplitInfo.NewVertex] = R3Pos;
+			const double SurfaceAlpha     = (StartV == 0) ? R3Alpha : 1. - R3Alpha;
+			IntrinsicVertexPositions.InsertAt(FSurfacePoint(SurfaceEID, SurfaceAlpha), SplitInfo.NewVertex);
+		
+		}
+		else
+		{
+			// compute the r3 positions of the boundary points
+			const FVector3d LeftR3  = AsR3Pos(LeftSP);
+			const FVector3d RightR3 = AsR3Pos(RightSP);
+		
+			// compute the local alpha of the new vertex between the boundary points
+			const double BoundingPtsSeperation = (LeftR3 - RightR3).Length();
+			const double SplitToRightBoundary  = BoundingPoints[1].DistOnEdge - IntrinsicEdgeSplitDistance;
+			const double Alpha = SplitToRightBoundary / BoundingPtsSeperation;
+			checkSlow(Alpha >= 0. && Alpha <=1.);
+
+	
+
+			// find the surface triangle the new vertex lives on
+			const int32 SurfaceTID = [&]
+									{
+										int32 TID = -1;
+										if (IsFacePoint(LeftSP))
+										{
+											TID = LeftSP.Position.TriPosition.TriID;
+										}
+										else if (IsFacePoint(RightSP))
+										{
+											TID = RightSP.Position.TriPosition.TriID;
+										}
+										if (TID == -1)
+										{
+											TArray<int32> AdjSurfaceTIDs = GetAdjacentTriangles(LeftSP, SurfaceMesh);
+											for (const int32 CandidtateTID : AdjSurfaceTIDs)
+											{
+												if (IsOnSurfaceTriangle(RightSP, CandidtateTID, SurfaceMesh))
+												{
+													TID = CandidtateTID;
+													break;
+												}
+											}
+										}
+										return TID;
+									}();
+
+			checkSlow(SurfaceTID != -1);
+
+			// compute the barycentric coordinates of the boundary points relative to the surface triangle
+			const FVector3d LeftBC  = AsBarycenteric(LeftSP, SurfaceTID, SurfaceMesh);
+			const FVector3d RightBC = AsBarycenteric(RightSP, SurfaceTID, SurfaceMesh);
+	
+			// compute the R3 and Barycentric location of the new vertex
+			const FVector3d R3Pos        = Alpha * LeftR3 + (1.-Alpha) * RightR3;
+			const FVector3d SurfaceBCPos = Alpha * LeftBC + (1.-Alpha) * RightBC;
+
+			// record the R3 position and the surface position.
+			Vertices[SplitInfo.NewVertex] = R3Pos;
+			IntrinsicVertexPositions.InsertAt(FSurfacePoint(SurfaceTID, SurfaceBCPos), SplitInfo.NewVertex);
+		
+		}
+	}
+
+	
+	return EMeshResult::Ok;
+}
+
diff --git a/Engine/Plugins/Runtime/GeometryProcessing/Source/DynamicMesh/Public/Operations/IntrinsicCorrespondenceUtils.h b/Engine/Plugins/Runtime/GeometryProcessing/Source/DynamicMesh/Public/Operations/IntrinsicCorrespondenceUtils.h
--- a/Engine/Plugins/Runtime/GeometryProcessing/Source/DynamicMesh/Public/Operations/IntrinsicCorrespondenceUtils.h
+++ b/Engine/Plugins/Runtime/GeometryProcessing/Source/DynamicMesh/Public/Operations/IntrinsicCorrespondenceUtils.h
@@ -9,7 +9,7 @@
 #include "VectorTypes.h"
 #include "MathUtil.h"
 
-
+struct foofoo { foofoo(){}; double x; double y; double z;};
 namespace UE
 {
 namespace Geometry
@@ -47,10 +47,11 @@ namespace IntrinsicCorrespondenceUtils
 		struct FTrianglePosition
 		{
 			int32 TriID;                    // Reference (extrinsic) triangle
-			double BarycentricCoords[3];    // Barycentric coordinates within referenced triangle
+			FVector3d  BarycentricCoords;   // Barycentric coordinates within referenced triangle
 		};
 		union FSurfacePositionUnion
 		{
+			FSurfacePositionUnion() { VertexPosition.VID = -1;}
 			FVertexPosition   VertexPosition;
 			FEdgePosition     EdgePosition;
 			FTrianglePosition TriPosition;
@@ -69,7 +70,29 @@ namespace IntrinsicCorrespondenceUtils
 	* @param Mesh         - Dynamic mesh that defines a surface
 	* @param bValidPoint  - will be false if the vert, edge or tri that defines the surface point is not a member of Mesh
 	*/
-	FVector3d AsR3Position(const FSurfacePoint& SurfacePoint, const FDynamicMesh3& Mesh, bool& bValidPoint);
+	FVector3d DYNAMICMESH_API AsR3Position(const FSurfacePoint& SurfacePoint, const FDynamicMesh3& Mesh, bool& bValidPoint);
+
+	/**
+	* @return true if the SurfacePoint corresponds to a vertex
+	*/
+	static inline bool IsVertexPoint(const FSurfacePoint& SurfacePoint)
+	{
+		return SurfacePoint.PositionType == FSurfacePoint::EPositionType::Vertex;
+	}
+	/**
+	* @return true if the SurfacePoint corresponds to an edge
+	*/
+	static inline bool IsEdgePoint(const FSurfacePoint& SurfacePoint)
+	{
+		return SurfacePoint.PositionType == FSurfacePoint::EPositionType::Edge;
+	}
+	/**
+	* @return true if the SurfacePoint corresponds to a face
+	*/
+	static inline bool IsFacePoint(const FSurfacePoint& SurfacePoint)
+	{
+		return SurfacePoint.PositionType == FSurfacePoint::EPositionType::Triangle;
+	}
 
 	// A Connection stores the information needed to define a local reference direction for each vertex and triangle.
 	struct DYNAMICMESH_API FMeshConnection
@@ -98,6 +121,20 @@ namespace IntrinsicCorrespondenceUtils
 			int32 CIdx; // index of curve relative to the edge crossed
 		};
 
+		/**
+		* @return the number of surface edges that cross this intrinsic edge.  
+		* Note: it is the responsibility of the caller to ensure IntrinsicEID is a valid edge on the intrinsic mesh.
+		*/ 
+		inline int32 NumEdgeCrossing(const int32 IntrinsicEID) const;
+
+		/**
+		* @return true if the specified intrinsic edge ( indicated by IntrinsicEID) is a segment of a surface edge.  
+		* Note: By construction the initial intrinsic edges will have the same ID as the corresponding surface mesh edges, but edge splits 
+		* will disrupt this.
+		* Also note: it is the responsibility of the caller to ensure IntrinsicEID is a valid edge on the intrinsic mesh.
+		*/ 
+		inline bool IsSurfaceEdgeSegment(const int32 IntrinsicEID) const;
+
 		/**
 		* @return number of surface edges within the specified intrinsic triangle that are adjacent to the specified triangle corner.
 		* This excludes any surface edges that are along the intrinsic triangle boundary edges.
@@ -131,10 +168,11 @@ namespace IntrinsicCorrespondenceUtils
 		*/
 		int32 GetEdgeOrder(const int32 VID, const int32 SurfaceEID) const;
 
-		TArray<int32>     NormalCoord;     // Per intrinsic edge: number of times src mesh edge crosses. NormalCoord[EID] is 0 if and only if same as edge on surface mesh   
-		TArray<int32>     RefVertDegree;   // Per vertex in ref mesh, caches the result of counting the edges in RefMesh.VtxEdgeIter().  
-		TArray<FIndex3i>  RoundaboutOrder; // Per intrinsic Triangle, per (directed) edge: indicates the closest (when traveling ccw) reference mesh edge  
-		TArray<FIndex3i>  EdgeOrder;       // Per surface triangle, per directed edge: number of edges past reference edge traveling ccw.
+		TDynamicVector<int32>     NormalCoord;     // Per intrinsic edge: number of times src mesh edge crosses. By convention NormalCoord[EID] = -1 iff same as edge on surface mesh   
+		TDynamicVector<FIndex3i>  RoundaboutOrder; // Per intrinsic Triangle, per (directed) edge: indicates the closest (when traveling ccw) reference mesh edge
+
+		TArray<int32>             RefVertDegree;   // Per vertex in ref mesh, caches the result of counting the edges in RefMesh.VtxEdgeIter().  
+		TArray<FIndex3i>          EdgeOrder;       // Per surface triangle, per directed edge: number of edges past reference edge traveling ccw.
 	};
 
 
@@ -176,7 +214,15 @@ namespace IntrinsicCorrespondenceUtils
 		TDynamicVector<FGeometricInfo>       GeometricVertexInfo;             // Per-intrinsic Vertex - Simple geometric information about the extrinsic surface in the neighborhood of intrinsic vert 	
 	};
 
+	int32 FNormalCoordinates::NumEdgeCrossing(const int32 IntrinsicEID) const
+	{
+		return TMathUtil<int32>::Max(0, NormalCoord[IntrinsicEID]);
+	}
 
+	bool FNormalCoordinates::IsSurfaceEdgeSegment(const int32 IntrinsicEID) const
+	{
+		return NormalCoord[IntrinsicEID] == -1;
+	}
 
 
 
@@ -191,9 +237,10 @@ namespace IntrinsicCorrespondenceUtils
 		const int32 Edgejk = IntrinsicTriEIDs[j_index];
 		const int32 Edgeki = IntrinsicTriEIDs[k_index];
 
-		const int32 N_ij = NormalCoord[Edgeij];
-		const int32 N_jk = NormalCoord[Edgejk];
-		const int32 N_ki = NormalCoord[Edgeki];
+		const int32 N_ij = NumEdgeCrossing(Edgeij);
+		const int32 N_jk = NumEdgeCrossing(Edgejk);
+		const int32 N_ki = NumEdgeCrossing(Edgeki);
+
 
 
 		return TMathUtil<int32>::Max(0, N_ij - N_jk - N_ki);
@@ -209,9 +256,9 @@ namespace IntrinsicCorrespondenceUtils
 		const int32 Edgejk = IntrinsicTriEIDs[j_index];
 		const int32 Edgeki = IntrinsicTriEIDs[k_index];
 
-		const int32 N_ij = NormalCoord[Edgeij];
-		const int32 N_jk = NormalCoord[Edgejk];
-		const int32 N_ki = NormalCoord[Edgeki];
+		const int32 N_ij = NumEdgeCrossing(Edgeij);
+		const int32 N_jk = NumEdgeCrossing(Edgejk);
+		const int32 N_ki = NumEdgeCrossing(Edgeki);
 
 		// see Gillespi et al, eqn 14
 
diff --git a/Engine/Plugins/Runtime/GeometryProcessing/Source/DynamicMesh/Public/Operations/IntrinsicTriangulationMesh.h b/Engine/Plugins/Runtime/GeometryProcessing/Source/DynamicMesh/Public/Operations/IntrinsicTriangulationMesh.h
--- a/Engine/Plugins/Runtime/GeometryProcessing/Source/DynamicMesh/Public/Operations/IntrinsicTriangulationMesh.h
+++ b/Engine/Plugins/Runtime/GeometryProcessing/Source/DynamicMesh/Public/Operations/IntrinsicTriangulationMesh.h
@@ -16,29 +16,49 @@
 namespace Geometry
 {
 /**
-* FSimpleIntrinsicEdgeFlipMesh supports edge flips, but no other operations that change
-* the mesh connectivity, and no operations that affect the vertices.
-*
+* Intrinsic Meshes:
+* 
 * An Intrinsic Mesh can be thought of as a mesh that overlays another mesh, sharing
-* the original surface (and in this case vertex locations), but the edges of the intrinsic mesh
-* (while constrained to be on the original mesh surface) need not align with the edges of the original mesh and the
-* lengths of these edges are measured on the surface of the original mesh (not the R3 distance).
+* the original surface.  The edges of the intrinsic mesh are constrained to be on the original mesh surface,
+* but need not align with the original mesh edges. The lengths of the intrinsic edges are measured on 
+* the surface of the original mesh (not the R3 distance).  The original mesh vertices are a subset of 
+* the intrinsic mesh vertices and intrinsic edge splits and triangle pokes may introduce new intrinsic mesh vertices.
 *
-* The FSimpleIntrinsicEdgeFlipMesh is designed to work with the function FlipToDelaunay() to make an
-* Intrinsic Delaunay Triangulation (IDT) of the same surface as a given FDynamicMesh,
-* allowing for more robust cotangent-Laplacians (see LaplacianMatrixAssembly.h)
-*
-* This simple mesh alone does not support computing surface positions from intrinsic mesh positions,
-* to do such computations FIntrinsicEdgeFlipMesh or FIntrinsicTriangulation should be used.
-*
 * Note: the implementation is a simple triangle-based mesh but there is no restriction that
 * edge(a, b) is unique (e.g. multiple intrinsic edges may connect the same two vertices with different paths).
 * In fact this mesh allows triangles and edges with repeated vertices. Such structures arise naturally
 * with intrinsic triangulation, for example an edge that starts and ends at the
 * same vertex may encircle a mesh.
+* 
+* The bookkeeping that manages the correspondence between locations on the intrinsic mesh and on the surface mesh
+* is implemented either using integer-based "Normal Coordinates" 
+* ( cf 'Integer Coordinates for Intrinsic Geometry Processing' M. Gillespie, N. Sharp, and K. Crane, TOG , Vol. 40,  December 2021. )
+* as used by FIntrinsicEdgeFlipMesh and FIntrinsicMesh; 
+* or alternately floating-point based generalized directions "Signpost data" 
+* ( cf 'Signpost Coordinates inspired by "Navigating Intrinsic Triangulations' Sharp, Soliman and Crane [2019, ACM Transactions on Graphics])
+* as used by FIntrinsicTriangulation.
 *
+*/
+
+/**
+* The FSimpleIntrinsicEdgeFlipMesh is designed to work with the function FlipToDelaunay() to make an
+* Intrinsic Delaunay Triangulation (IDT) of the same surface as a given FDynamicMesh,
+* allowing for more robust cotangent-Laplacians (see LaplacianMatrixAssembly.h)
+* 
+* 
+* The vertices in this FSimpleIntrinsicEdgeFlipMesh are exactly those found in the surface mesh
+* since it supports edge flips, but no other operations that change the mesh connectivity, 
+* and no operations that affect the vertices.  
+* 
+* This mesh alone does not support computing the intersections of intrinsic edges with the surface mesh edges
+* and as a result it can not be used to easily extract the path of an intrinsic edge.
+* To do such computations FIntrinsicEdgeFlipMesh, FIntrinsicMesh or FIntrinsicTriangulation should be used.
+*
 * The API of FIntrinsicEdgeFlipMesh is similar to FDynamicMesh3, but only has a minimal subset of methods and some
-* such as ::FlipEdge() and ::GetEdgeOpposingV() have very different implementations
+* such as ::FlipEdge() and ::GetEdgeOpposingV() have very different implementations.
+* 
+* In addition to managing vertices, triangles and edges this class also tracks edge lengths (as measured on the surface mesh)
+* and for convenience the internal angles for each triangle.
 */
 class DYNAMICMESH_API FSimpleIntrinsicEdgeFlipMesh
 {
@@ -198,6 +218,22 @@
 		return VertexRefCounts.IsValid(VertexID);
 	}
 
+	/** @return true if the specified vertex is adjacent to a boundary edge*/
+	inline bool IsBoundaryVertex(int32 VertexID) const
+	{
+		if (IsVertex(VertexID))
+		{
+			for (int32 EID : VertexEdgeLists.Values(VertexID))
+			{
+				if (Edges[EID].Tri[1] == InvalidID)
+				{
+					return true;
+				}
+			}
+		}
+		return false;
+	}
+
 	/** @return the r3 position of the specified vertex */
 	inline FVector3d GetVertex(int32 VertexID) const
 	{
@@ -210,6 +246,12 @@
 	{
 		return (int32)TriangleRefCounts.GetMaxIndex();
 	}
+	
+	/** @return enumerable object for valid triangle indices suitable for use with range-based for, ie for ( int i : TriangleIndicesItr() ) */
+	FRefCountVector::IndexEnumerable TriangleIndicesItr() const
+	{
+		return TriangleRefCounts.Indices();
+	}
 
 	/** @return true if intrinsic mesh contains specified triangle */
 	inline bool IsTriangle(int32 TriangleID) const
@@ -231,15 +273,6 @@
 		return TriangleEdges[TriangleID];
 	}
 
-
-protected:
-	
-	/** @return three requested edge lengths */
-	FVector3d GetEdgeLengthTriple(const FIndex3i& EIDs) const
-	{
-		return FVector3d(EdgeLengths[EIDs.A], EdgeLengths[EIDs.B], EdgeLengths[EIDs.C]);
-	}
-
 	/** 
 	* cyclic permutation of the src vector such that the 'Index' entry in the src is the first entry in the result
 	* NB: Index must be 0, 1, or 2.
@@ -251,6 +284,15 @@
 		return Vector3Type(Src[Index], Src[AddOneModThree[Index]], Src[AddTwoModThree[Index]]);
 	}
 
+protected:
+	
+	/** @return three requested edge lengths */
+	FVector3d GetEdgeLengthTriple(const FIndex3i& EIDs) const
+	{
+		return FVector3d(EdgeLengths[EIDs.A], EdgeLengths[EIDs.B], EdgeLengths[EIDs.C]);
+	}
+
+
 	/**
 	* computes the angles based on the edge lengths, so must have valid EdgeLengths for this tri before calling.
 	* the order of the angles matches the order of the vertices that define the triangle, e.g. if (a, b, c) = GetTriangle()
@@ -338,16 +380,174 @@
 
 
 /**
-* FIntrinsicEdgeFlipMesh supports edge flips, but no other operations that change
-* the mesh connectivity, and no operations that affect the vertices.
+* FSimpleIntrinsicMesh extends FSimpleIntrinsicEdgeFlipMesh to support edge splits and triangle pokes in addition to edge flips.
+* The surface mesh vertices will be a subset of the intrinsic mesh vertex set. 
+* 
+* This is intended as a base class for more sophisticated intrinsic mesh classes, as it does not explicitly track
+* the topology changes relative to the original surface mesh.  For example it does not record the locations on the surface mesh 
+* of any new intrinsic vertices or the locations of intrinsic mesh edge intersections with the surface mesh edges.
 *
-* Unlike, FSimpleIntrinsicEdgeFlipMesh from which it is derived, a correspondence can be 
-* formed between the two meshes: the surface mesh edges can be traced across this intrinsic mesh, 
-* and this in-turn can be used to trace the intrinsic mesh edges across the surface mesh.  
+* FIntrinsicMesh and FIntrinsicTriangulation are derived from this class and each maintain their own connection 
+* with the surface mesh ( "normal coordinates" or "signpost data") that can track these surface mesh locations 
+* and can be used to reconstruct paths (e.g. intrinsic edges) on the surface mesh.
+*/
+class DYNAMICMESH_API FSimpleIntrinsicMesh : public FSimpleIntrinsicEdgeFlipMesh
+{
+public:
+	typedef FSimpleIntrinsicEdgeFlipMesh  MyBase;
+	using FEdge = FDynamicMesh3::FEdge;
+	using FEdgeFlipInfo     = DynamicMeshInfo::FEdgeFlipInfo;
+	using FPokeTriangleInfo = DynamicMeshInfo::FPokeTriangleInfo;
+	using FEdgeSplitInfo    = DynamicMeshInfo::FEdgeSplitInfo;
+
+	/** InvalidID indicates that a vertex/edge/triangle ID is invalid */
+	constexpr static int InvalidID = IndexConstants::InvalidID;
+
+	/** Constructor does ID-matching deep copy the basic mesh topology (but no attributes or groups) */
+	FSimpleIntrinsicMesh(const FDynamicMesh3& SrcMesh) : MyBase(SrcMesh)
+	{};
+
+	FSimpleIntrinsicMesh() = delete;
+	FSimpleIntrinsicMesh(const FSimpleIntrinsicMesh&) = delete;
+
+
+	/**
+	* Insert a new vertex inside an intrinsic triangle, ie do a 1 to 3 triangle split
+	* @param TID             - index of triangle to poke
+	* @param BaryCoordinates - barycentric coordinates of poke position
+	* @param PokeInfo        - returned information about new and modified mesh elements
+	* 
+	* @return Ok on success, or enum value indicates why operation cannot be applied. Mesh remains unmodified on error.
+	*/
+	EMeshResult PokeTriangle(int32 TID, const FVector3d& BaryCoordinates, FPokeTriangleInfo& PokeInfo);
+
+
+	/**
+	* Split an intrinsic edge of the mesh by inserting a vertex. This creates a new triangle on either side of the edge (ie a 2-4 split).
+	* If the original edge had vertices [a,b], with triangles t0=[a,b,c] and t1=[b,a,d],  then the split inserts new vertex f.
+	* After the split t0=[a,f,c] and t1=[f,a,d], and we have t2=[f,b,c] and t3=[f,d,b]  (it's best to draw it out on paper...)
+	* SplitInfo.OriginalTriangles = {t0, t1} and SplitInfo.NewTriangles = {t2, t3}
+	*
+	* @param EdgeAB          - index of the edge to be split
+	* @param SplitInfo       - returned information about new and modified mesh elements
+	* @param SplitParameterT - defines the position along the edge that we split at, must be between 0 and 1, and is assumed to be based on the order of vertices in t0
+	* 
+	* @return Ok on success, or enum value indicates why operation cannot be applied. Mesh remains unmodified on error.
+	*/
+	EMeshResult SplitEdge(int32 EdgeAB, FEdgeSplitInfo& SplitInfo, double SplitParameterT = 0.5);
+
+protected:
+
+	/**
+	* Updates the mesh connectivity by adding a new vertex at the vertex-averaged location. Note: this does not update any of the intrinsic lengths, positions, angles, etc.
+	* those quantities must be updated after.  See ::PokeTriangle()
+	*/ 
+	EMeshResult PokeTriangleTopology(int32 TID, FPokeTriangleInfo& PokeInfo);
+
+	/**
+	* Updates the mesh connectivity by adding a new vertex at the vertex-averaged location. Note: this does not update any of the intrinsic lengths, positions, angles, etc.
+	* those quantities must be updated after.  See ::SplitEdge()
+	*/
+	EMeshResult SplitEdgeTopology(int32 EdgeAB, FEdgeSplitInfo& SplitInfo);
+
+	/** allocate, or clear existing edge list for specified vertex*/
+	inline void AllocateEdgesList(int32 VertexID)
+	{
+		if (VertexID < (int)VertexEdgeLists.Size())
+		{
+			VertexEdgeLists.Clear(VertexID);
+		}
+		VertexEdgeLists.AllocateAt(VertexID);
+	}
+
+	/** add new vertex to the mesh topology, does not update any of the intrinsic quantities*/
+	inline int32 AppendVertex(const FVector3d& vPos)
+	{
+		int32 newVertID = VertexRefCounts.Allocate();
+		Vertices.InsertAt(vPos, newVertID);
+		AllocateEdgesList(newVertID);
+		return newVertID;
+	}
+
+	/** add new edge to the mesh topology, does not update any of the intrinsic quantities*/
+	inline int32 AddEdgeInternal(int32 vA, int32 vB, int32 tA, int32 tB = InvalidID)
+	{
+		if (vB < vA) {
+			int32 t = vB; vB = vA; vA = t;
+		}
+		int32 eid = EdgeRefCounts.Allocate();
+		Edges.InsertAt(FEdge{ {vA, vB},{tA, tB} }, eid);
+		VertexEdgeLists.Insert(vA, eid);
+		if (vA != vB)
+		{
+			VertexEdgeLists.Insert(vB, eid);
+		}
+
+		return eid;
+	}
+	/** add new triangle to the mesh topology, does not update any of the intrinsic quantities*/
+	inline int32 AddTriangleInternal(int32 a, int32 b, int32 c, int32 e0, int32 e1, int32 e2)
+	{
+		int32 tid = TriangleRefCounts.Allocate();
+		Triangles.InsertAt(FIndex3i(a, b, c), tid);
+		TriangleEdges.InsertAt(FIndex3i(e0, e1, e2), tid);
+		return tid;
+	}
+};
+
+/**
+* TEdgeCorrespondence allows intrinsic edges to be traced across the underlying surface mesh.
+* When computing a large number intrinsic edge traces, this is to be preferred over the TraceEdge()
+* intrinsic mesh methods because it avoids redundant computations.  But when tracing a small number
+* of intrinsic edges the large up-front cost of computing information the entire mesh will make use
+* of this class sub-optimal.
+* 
+* During construction a list of surface mesh edges crossed by each intrinsic edge is stored.
+* The Trace operation for an intrinsic edge uses this information to unfold the crossed surface triangles
+* into a 2d triangle strip upon which the actual trace takes place.
+* 
+* This is only for use with intrinsic mesh implementations that support  'normal coordinates' when doing flips/splits and pokes.
+* such as the FIntrinsicMesh and FIntrinsicEdgeFlipMesh.
+*/ 
+template <typename IntrinsicMeshType>
+struct DYNAMICMESH_API TEdgeCorrespondence
+{
+
+	using FEdgeAndCrossingIdx = IntrinsicCorrespondenceUtils::FNormalCoordinates::FEdgeAndCrossingIdx;
+	using FSurfacePoint = IntrinsicCorrespondenceUtils::FSurfacePoint;
+
+
+	TEdgeCorrespondence(const IntrinsicMeshType& IntrinsicMesh)
+	{
+		Setup(IntrinsicMesh);
+	}
+
+	void Setup(const IntrinsicMeshType& IntrinsicMesh);
+	/**
+	* @return Array of FSurfacePoints relative to the ExtrinsicMesh that represent the specified intrinsic mesh edge.
+	* This is directed relative to traversal of the first adjacent triangle (GetEdgeT().[0]), unless bReverse is true
+	*
+	* @param IntrinsicEID      - ID of the intrinsic mesh edge to be traced.
+	* @param CoalesceThreshold - in barycentric units [0,1], edge crossings within this threshold are snapped to the nearest vertex.
+	*                            and any resulting repetitions of vertex surface points are replaced with a single vertex surface point.
+	* @param bReverse          - if true, trace edge direction is reversed
+	*/
+	TArray<UE::Geometry::IntrinsicCorrespondenceUtils::FSurfacePoint> TraceEdge(int32 IntrinsicEID, double CoalesceThreshold = 0., bool bReverse = false) const;
+
+
+	const IntrinsicMeshType*               IntrinsicMesh;         // the intrinsic mesh
+	const FDynamicMesh3*                   SurfaceMesh;           // the surface mesh 
+	TArray<TArray<int32>>                  SurfaceEdgesCrossed;   // array of surface edge crossings per intrinsic edge.
+};
+
+/**
+* FIntrinsicEdgeFlipMesh augments the base class FSimpleIntrinsicEdgeFlipMesh with
+* the addition of a 'normal coordinate' based correspondence with the surface mesh.
 *
-* The implementation uses integer coordinates to robustly construct these paths.
+* This class supports both the trace of intrinsic mesh edges across the surface mesh,
+* and the trace of surface mesh edges across the intrinsic mesh.
 * 
-* Note, because only edge flips are supported (and not edge split or triangle pokes), the
+* Note, because only edge flips are supported (no edge split or triangle pokes), the
 * vertex set of the intrinsic mesh and the surface (i.e. extrinsic) mesh will be the same.
 * 
 * The expected use
@@ -372,7 +572,9 @@
 *          auto TracedIntrinsicEdge  =  Correspondence.TraceEdge(IntrinsicEdgeID);
 *      } 
 *
-* Note: this isn't currently expected to work with surface meshes that have bow-ties.
+* NB: The lifetime of this structure should not exceed that of the original
+* FDynamicMesh as this class holds a pointer to that mesh to reference locations on its surface.
+* Also, this class isn't currently expected to work with surface meshes that have bow-ties.
 */
 class DYNAMICMESH_API FIntrinsicEdgeFlipMesh : public  FSimpleIntrinsicEdgeFlipMesh
 {
@@ -417,40 +619,140 @@
 	EMeshResult FlipEdge(int32 EID, FEdgeFlipInfo& EdgeFlipInfo);
 
 
-	/**
-	* FEdgeCorrespondence allows intrinsic edges to be traced across the surface mesh.
-	* During construction a list of surface mesh edges crossed by each intrinsic edge is stored.
-	* The Trace operation for an intrinsic edge uses this information to unfold the crossed surface triangles
-	* into a 2d triangle strip upon which the actual trace takes place.
-	*/ 
-	struct DYNAMICMESH_API FEdgeCorrespondence
-	{
+	typedef TEdgeCorrespondence< FIntrinsicEdgeFlipMesh >  FEdgeCorrespondence;
 
-		using FEdgeAndCrossingIdx = IntrinsicCorrespondenceUtils::FNormalCoordinates::FEdgeAndCrossingIdx;
-
-		FEdgeCorrespondence(const FIntrinsicEdgeFlipMesh& IntrinsicEdgeFlipTriangulation);
-
-		/**
-		* @return Array of FSurfacePoints relative to the ExtrinsicMesh that represent the specified intrinsic mesh edge.
-		* This is directed relative to traversal of the first adjacent triangle (GetEdgeT().[0]), unless bReverse is true
-		*
-		* @param IntrinsicEID      - ID of the intrinsic mesh edge to be traced.
-		* @param CoalesceThreshold - in barycentric units [0,1], edge crossings within this threshold are snapped to the nearest vertex.
-		*                            and any resulting repetitions of vertex surface points are replaced with a single vertex surface point.
-		* @param bReverse          - if true, trace edge direction is reversed
-		*/
-		TArray<FSurfacePoint> TraceEdge(int32 IntrinsicEID, double CoalesceThreshold = 0., bool bReverse = false) const;
-
-
-		const FIntrinsicEdgeFlipMesh*          IntrinsicMesh;         // the intrinsic mesh
-		const FDynamicMesh3*                   SurfaceMesh;           // the surface mesh 
-		TArray<TArray<int32>>                  SurfaceEdgesCrossed;   // array of surface edge crossings per intrinsic edge.
-	};
-
 	FEdgeCorrespondence ComputeEdgeCorrespondence() const 
 	{
 		return FEdgeCorrespondence(*this);
 	};
 
 	/**
+	* @return NormalCoordinate data structure that counts the number of regular surface (extrinsic) mesh edge-crossings for each intrinsic mesh edge. 
+	*/
+	const IntrinsicCorrespondenceUtils::FNormalCoordinates& GetNormalCoordinates() const
+	{
+		return NormalCoordinates;
+	}
+
+	/**
+	* Trace the specified surface (extrinsic) mesh edge across the intrinsic mesh.  By default this traces from vertex EdgeV.A to vertex EdgeV.B
+	* where EdgeV = SurfaceMesh.GetEdgeV(SurfaceEID).
+	* @param SurfaceEID   - the surface edge to be traced
+	* @param bReverse     - reverses the direction of the trace giving EdgeV.B -> EdgeV.A
+	* 
+	* @return array of intrinsic edges crossed between the specified vertices.  This can be used to compute the actual MeshPoints on the surface if desired.
+	*/ 
+	TArray<FEdgeAndCrossingIdx> GetImplicitEdgeCrossings(const int32 SurfaceEID, const bool bReverse = false) const;
+
+	/**
+	* Trace the specified surface (extrinsic) mesh edge across the intrinsic mesh.  By default this traces from vertex EdgeV.A to vertex EdgeV.B
+	* where EdgeV = SurfaceMesh.GetEdgeV(SurfaceEID).
+	* @param SurfaceEID   - the surface edge to be traced
+	* @param bReverse     - reverses the direction of the trace giving EdgeV.B -> EdgeV.A
+	* 
+	* @return array of surface points defined relative to the intrinsic mesh.  In particular, the array will be a vertex followed by a sequence of (intrinsic) edge crossings.
+	*/ 
+	TArray<FSurfacePoint> TraceSurfaceEdge(int32 SurfaceEID, double CoalesceThreshold, bool bReverse = false) const;
+
+	/**
+	* @return true if the intrinsic vertex corresponds to a surface vertex.  
+	*/ 
+	bool IsSurfaceVertex(int VID) const 
+	{
+		return  (IsVertex(VID)) ? true : false;
+	}
+
+	/**
+	* @return the position of the intrinsic vertex relative to the underlying surface mesh
+	*/ 
+	FSurfacePoint GetVertexSurfacePoint(int32 VID) const
+	{
+		// The FIntrinsicEdgeFlipMesh verts correspond to surface verts ( this mesh does not insert or delete verts).
+		return FSurfacePoint(VID);
+	}
+
+protected:
+
+	using FNormalCoordinates = IntrinsicCorrespondenceUtils::FNormalCoordinates;
+	FNormalCoordinates  NormalCoordinates;      // connection and integer-based coordinates used in reconstructing {intrinsic, surface} edges on the {surface, intrinsic} mesh
+};
+
+
+
+/**
+* FIntrinsicMesh extends the FSimpleIntrinsicMesh with the addition of a 'normal coordinate' 
+* based correspondence with the surface mesh.
+*
+* This class supports both the trace of intrinsic mesh edges across the surface mesh,
+* and the trace of surface mesh edges across the intrinsic mesh.
+* 
+* This supports mesh operations that: SplitEdge, PokeTriangle, FlipEdge.
+*
+* Both SplitEdge and PokeTriangle will generate new vertices and triangles in the
+* intrinsic triangulation.  These vertices will live on the surface of the original mesh
+* and their location (in R3 and relative to the original surface ) is computed by interpolation.
+*
+* NB: The lifetime of this structure should not exceed that of the original
+* FDynamicMesh as this class holds a pointer to that mesh to reference locations on its surface.
+* Also, this class isn't currently expected to work with surface meshes that have bow-ties.
+*/
+
+class DYNAMICMESH_API FIntrinsicMesh  : public  FSimpleIntrinsicMesh
+{
+public:
+	typedef FSimpleIntrinsicMesh     MyBase;
+
+	using FEdge = FIntrinsicEdgeFlipMesh::FEdge;
+	using FEdgeFlipInfo = FIntrinsicEdgeFlipMesh::FEdgeFlipInfo;
+	using FSurfacePoint = IntrinsicCorrespondenceUtils::FSurfacePoint;
+	using FEdgeAndCrossingIdx = IntrinsicCorrespondenceUtils::FNormalCoordinates::FEdgeAndCrossingIdx;
+	using FPokeTriangleInfo   = DynamicMeshInfo::FPokeTriangleInfo;
+	using FEdgeSplitInfo      = DynamicMeshInfo::FEdgeSplitInfo;
+
+	FIntrinsicMesh() = delete;
+	FIntrinsicMesh(const FIntrinsicMesh&) = delete;
+
+	FIntrinsicMesh(const FDynamicMesh3& SurfaceMesh);
+
+	/** @return pointer the reference mesh, note the FIntrinsicTriangulation does not manage the lifetime of this mesh */
+	const FDynamicMesh3* GetExtrinsicMesh()  const
+	{
+		return NormalCoordinates.SurfaceMesh;
+	}
+
+
+	/**
+	* Flip a single edge in the Intrinsic Mesh.
+	* @param EdgeFlipInfo [out] populated on return.
+	*
+	* @return a success or failure info.
+	*/
+	EMeshResult FlipEdge(int32 EID, FEdgeFlipInfo& EdgeFlipInfo);
+
+	/**
+	* Insert a new vertex inside an intrinsic triangle, ie do a 1 to 3 triangle split
+	* @param TID             - index of triangle to poke
+	* @param BaryCoordinates - barycentric coordinates of poke position
+	* @param PokeInfo        - returned information about new and modified mesh elements
+	* 
+	* @return Ok on success, or enum value indicates why operation cannot be applied. Mesh remains unmodified on error.
+	*/
+	EMeshResult PokeTriangle(int32 TID, const FVector3d& BaryCoordinates, FPokeTriangleInfo& PokeInfo);
+
+	/**
+	* Split an intrinsic edge of the mesh by inserting a vertex. This creates a new triangle on either side of the edge (ie a 2-4 split).
+	* If the original edge had vertices [a,b], with triangles t0=[a,b,c] and t1=[b,a,d],  then the split inserts new vertex f.
+	* After the split t0=[a,f,c] and t1=[f,a,d], and we have t2=[f,b,c] and t3=[f,d,b]  (it's best to draw it out on paper...)
+	* SplitInfo.OriginalTriangles = {t0, t1} and SplitInfo.NewTriangles = {t2, t3}
+	*
+	* @param EdgeAB          - index of the edge to be split
+	* @param SplitInfo       - returned information about new and modified mesh elements
+	* @param SplitParameterT - defines the position along the edge that we split at, must be between 0 and 1, and is assumed to be based on the order of vertices in t0
+	* 
+	* @return Ok on success, or enum value indicates why operation cannot be applied. Mesh remains unmodified on error.
+	*/
+	EMeshResult SplitEdge(int32 EdgeAB, FEdgeSplitInfo& SplitInfo, double SplitParameterT = 0.5);
+
+
+	/**
 	* @return NormalCoordinate data structure that counts regular surface (extrinsic) mesh edge-crossing for each intrinsic mesh edge. 
@@ -461,5 +763,18 @@
 	}
 
 	/**
+	* @return Array of FSurfacePoints relative to the surface mesh that represent the specified intrinsic mesh edge.
+	* This is directed relative to traversal of the first adjacent triangle (GetEdgeT()[0]), unless bReverse is true
+	*
+	* @param CoalesceThreshold - in barycentric units [0,1], edge crossings with in this threshold are snapped to the nearest vertex.
+	*                            and any resulting repetitions of vertex surface points are replaced with a single vertex surface point.
+	* @param bReverse          - if true, trace edge direction is reversed
+	* 
+	* Note: when tracing few edges, this method is preferred to using the FEdgeCorrespondence based tracing (see FEdgeCorrespondence below).  
+	* But when tracing a large percentage of the mesh edges the FEdgeCorrespondence will be faster as it reuses intermediate results.
+	*/
+	TArray<FSurfacePoint> TraceEdge(int32 IntrinsicEID, double CoalesceThreshold = 0., bool bReverse = false) const;
+
+	/**
 	* Trace the specified surface (extrinsic) mesh edge across the intrinsic mesh.  By default this traces from vertex EdgeV.A to vertex EdgeV.B
 	* where EdgeV = SurfaceMesh.GetEdgeV(SurfaceEID).
@@ -480,24 +795,46 @@
 	*/ 
 	TArray<FSurfacePoint> TraceSurfaceEdge(int32 SurfaceEID, double CoalesceThreshold, bool bReverse = false) const;
 
+
+	typedef TEdgeCorrespondence< FIntrinsicMesh >  FEdgeCorrespondence;
+
+	FEdgeCorrespondence ComputeEdgeCorrespondence() const
+	{
+		return FEdgeCorrespondence(*this);
+	};
+
+
+	/**
+	* @return position in r3 of an intrinsic vertex.
+	*/
+	FVector3d GetVertexPosition(int32 VID) const
+	{
+		return GetVertex(VID);
+	}
+
+	/**
+	* @return the position of the intrinsic vertex relative to the underlying surface mesh
+	*/ 
+	FSurfacePoint GetVertexSurfacePoint(int32 VID) const
+	{
+		return IntrinsicVertexPositions[VID];
+	}
+
+
 protected:
-
 	using FNormalCoordinates = IntrinsicCorrespondenceUtils::FNormalCoordinates;
 	FNormalCoordinates  NormalCoordinates;      // connection and integer-based coordinates used in reconstructing {intrinsic, surface} edges on the {surface, intrinsic} mesh
+
+	TDynamicVector<FSurfacePoint>        IntrinsicVertexPositions;        // Per-intrinsic Vertex - Locates intrinsic mesh vertices relative to the surface defined by extrinsic mesh.
 };
 
 
+
 /**
 * Class that manages the intrinsic triangulation of a given FDynamicMesh.
 *
 * This supports mesh operations that: SplitEdge, PokeTriangle, FlipEdge.
 *
-* An Intrinsic triangulation can be thought of as a mesh that overlays another mesh, sharing
-* the original surface, but the intrinsic triangles can fold over the edges of the original mesh
-* That is to say, the edges of the intrinsic mesh are still on the mesh surface
-* but need not align with the edges of the original mesh and the lengths of these edges are measured
-* on the surface of the original mesh (not the R3 distance).
 *
-*
 * Both SplitEdge and PokeTriangle will generate new vertices and triangles in the
 * intrinsic triangulation.  These vertices will live on the surface of the original mesh
@@ -508,12 +845,12 @@
 * FDynamicMesh as this class holds a pointer to that mesh to reference locations on its surface.
 * Also, this class isn't currently expected to work with surface meshes that have bow-ties.
 */ 
-class DYNAMICMESH_API FIntrinsicTriangulation : public FSimpleIntrinsicEdgeFlipMesh
+class DYNAMICMESH_API FIntrinsicTriangulation : public  FSimpleIntrinsicMesh
 {
 
 public:
 
-	typedef FSimpleIntrinsicEdgeFlipMesh MyBase;
+	typedef FSimpleIntrinsicMesh MyBase;
 
 	using FEdgeFlipInfo = DynamicMeshInfo::FEdgeFlipInfo;
 	using FEdgeSplitInfo = DynamicMeshInfo::FEdgeSplitInfo;
@@ -614,61 +951,6 @@
 	double UpdateVertexByEdgeTrace(const int32 UpdateVID, const int32 StartVID, const double TracePolarDir, const double TraceDistance);
 
 
-	/**
-	* Updates the mesh connectivity by adding a new vertex at the vertex-averaged location. Note: this does not update any of the intrinsic lengths, positions, angles, etc.
-	* those quantities must be updated after.  See ::PokeTriangle()
-	*/ 
-	EMeshResult PokeTriangleTopology(int32 TID, FPokeTriangleInfo& PokeInfo);
-	
-	/**
-	* Updates the mesh connectivity by adding a new vertex at the vertex-averaged location. Note: this does not update any of the intrinsic lengths, positions, angles, etc.
-	* those quantities must be updated after.  See ::SplitEdge()
-	*/
-	EMeshResult SplitEdgeTopology(int32 EdgeAB, FEdgeSplitInfo& SplitInfo);
-
-	/** allocate, or clear existing edge list for specified vertex*/
-	inline void AllocateEdgesList(int32 VertexID)
-	{
-		if (VertexID < (int)VertexEdgeLists.Size())
-		{
-			VertexEdgeLists.Clear(VertexID);
-		}
-		VertexEdgeLists.AllocateAt(VertexID);
-	}
-
-	/** add new vertex to the mesh topology, does not update any of the intrinsic quantities*/
-	inline int32 AppendVertex(const FVector3d& vPos)
-	{
-		int32 newVertID = VertexRefCounts.Allocate();
-		Vertices.InsertAt(vPos, newVertID);
-		AllocateEdgesList(newVertID);
-		return newVertID;
-	}
-
-	/** add new edge to the mesh topology, does not update any of the intrinsic quantities*/
-	inline int32 AddEdgeInternal(int32 vA, int32 vB, int32 tA, int32 tB = InvalidID)
-	{
-		if (vB < vA) {
-			int32 t = vB; vB = vA; vA = t;
-		}
-		int32 eid = EdgeRefCounts.Allocate();
-		Edges.InsertAt(FEdge{ {vA, vB},{tA, tB} }, eid);
-		VertexEdgeLists.Insert(vA, eid);
-		if (vA != vB)
-		{ 
-			VertexEdgeLists.Insert(vB, eid);
-		}
-
-		return eid;
-	}
-	/** add new triangle to the mesh topology, does not update any of the intrinsic quantities*/
-	inline int32 AddTriangleInternal(int32 a, int32 b, int32 c, int32 e0, int32 e1, int32 e2)
-	{
-		int32 tid = TriangleRefCounts.Allocate();
-		Triangles.InsertAt(FIndex3i(a, b, c), tid);
-		TriangleEdges.InsertAt(FIndex3i(e0, e1, e2), tid);
-		return tid;
-	}
 protected:
 	using FSignpost = IntrinsicCorrespondenceUtils::FSignpost;
 	FSignpost SignpostData;     // connection and directional-based coordinates used in reconstructing {intrinsic, surface} edges on the {surface, intrinsic} mesh vert 		
@@ -686,6 +968,8 @@
 */
 int32 DYNAMICMESH_API FlipToDelaunay(FSimpleIntrinsicEdgeFlipMesh& IntrinsicMesh, TSet<int>& Uncorrected, const int32 MaxFlipCount = TMathUtilConstants<int>::MaxReal);
 int32 DYNAMICMESH_API FlipToDelaunay(FIntrinsicEdgeFlipMesh& IntrinsicMesh, TSet<int>& Uncorrected, const int32 MaxFlipCount = TMathUtilConstants<int>::MaxReal);
+int32 DYNAMICMESH_API FlipToDelaunay(FSimpleIntrinsicMesh& IntrinsicMesh, TSet<int>& Uncorrected, const int32 MaxFlipCount = TMathUtilConstants<int>::MaxReal);
+int32 DYNAMICMESH_API FlipToDelaunay(FIntrinsicMesh& IntrinsicMesh, TSet<int>& Uncorrected, const int32 MaxFlipCount = TMathUtilConstants<int>::MaxReal);
 int32 DYNAMICMESH_API FlipToDelaunay(FIntrinsicTriangulation& IntrinsicMesh, TSet<int>& Uncorrected, const int32 MaxFlipCount = TMathUtilConstants<int>::MaxReal);
 
 }; // end namespace Geometry
