Remove extraneous debug code.
#rnx

diff --git a/Engine/Plugins/Runtime/GeometryProcessing/Source/DynamicMesh/Public/Operations/IntrinsicCorrespondenceUtils.h b/Engine/Plugins/Runtime/GeometryProcessing/Source/DynamicMesh/Public/Operations/IntrinsicCorrespondenceUtils.h
--- a/Engine/Plugins/Runtime/GeometryProcessing/Source/DynamicMesh/Public/Operations/IntrinsicCorrespondenceUtils.h
+++ b/Engine/Plugins/Runtime/GeometryProcessing/Source/DynamicMesh/Public/Operations/IntrinsicCorrespondenceUtils.h
@@ -9,7 +9,6 @@
 #include "VectorTypes.h"
 #include "MathUtil.h"
 
-struct foofoo { foofoo(){}; double x; double y; double z;};
 namespace UE
 {
 namespace Geometry
