Geometry: Intrinsic Triangulation and Geodesics.  Fixed bugs related to the special case when an intrinsic triangle has only two distinct verts (as happens when an intrinsic triangle wraps the tip of a cone).   Also added  mesh reset methods and allowing the intrinsic mesh to optionally own a copy of the surface mesh it uses.

diff --git a/Engine/Plugins/Runtime/GeometryProcessing/Source/DynamicMesh/Private/Operations/IntrinsicCorrespondenceUtils.cpp b/Engine/Plugins/Runtime/GeometryProcessing/Source/DynamicMesh/Private/Operations/IntrinsicCorrespondenceUtils.cpp
--- a/Engine/Plugins/Runtime/GeometryProcessing/Source/DynamicMesh/Private/Operations/IntrinsicCorrespondenceUtils.cpp
+++ b/Engine/Plugins/Runtime/GeometryProcessing/Source/DynamicMesh/Private/Operations/IntrinsicCorrespondenceUtils.cpp
@@ -92,10 +92,28 @@ FVector3d IntrinsicCorrespondenceUtils::AsR3Position(const IntrinsicCorresponden
 /**------------------------------------------------------------------------------
 * FMeshConnection
 *-------------------------------------------------------------------------------*/
-
 IntrinsicCorrespondenceUtils::FMeshConnection::FMeshConnection(const FDynamicMesh3& SurfMesh)
-	:SurfaceMesh(&SurfMesh)
 {
+	Reset(SurfMesh);
+}
+
+void IntrinsicCorrespondenceUtils::FMeshConnection::Reset(TUniquePtr<FDynamicMesh3>& SurfMesh)
+{
+	Reset(*SurfMesh);
+
+	OwnedSurfaceMesh = MoveTemp(SurfMesh);
+	SurfaceMesh = OwnedSurfaceMesh.Get();
+}
+
+void IntrinsicCorrespondenceUtils::FMeshConnection::Reset(const FDynamicMesh3& SurfMesh)
+{
+	OwnedSurfaceMesh.Release();
+
+	SurfaceMesh = &SurfMesh;
+
+	VIDToReferenceEID.Reset();
+	TIDToReferenceEID.Reset();
+
 	// pick local reference directions for each vertex and triangle.
 	const int32 MaxVertexID   = SurfMesh.MaxVertexID();
 	const int32 MaxTriangleID = SurfMesh.MaxTriangleID();
@@ -154,8 +172,33 @@ IntrinsicCorrespondenceUtils::FMeshConnection::FMeshConnection(const FDynamicMes
 *-------------------------------------------------------------------------------*/
 
 IntrinsicCorrespondenceUtils::FNormalCoordinates::FNormalCoordinates(const FDynamicMesh3& SurfMesh)
-	: MyBase(SurfMesh)
+	: MyBase()
+{
+	Reset(SurfMesh);
+}
+
+void IntrinsicCorrespondenceUtils::FNormalCoordinates::Reset(TUniquePtr<FDynamicMesh3>& SurfMesh)
+{
+	MyBase::Reset(SurfMesh);
+
+	RebuildNormalCoordinates(*this->SurfaceMesh);
+}
+
+void IntrinsicCorrespondenceUtils::FNormalCoordinates::Reset(const FDynamicMesh3& SurfMesh)
 {
+	MyBase::Reset(SurfMesh);
+
+	RebuildNormalCoordinates(SurfMesh);
+}
+
+void IntrinsicCorrespondenceUtils::FNormalCoordinates::RebuildNormalCoordinates(const FDynamicMesh3& SurfMesh)
+{
+	NormalCoord.Clear();
+	RoundaboutOrder.Clear();
+
+	RefVertDegree.Reset();
+	EdgeOrder.Reset();
+
 	const int32 MaxEID = SurfMesh.MaxEdgeID();
 	const int32 MaxVID = SurfMesh.MaxVertexID();
 	const int32 MaxTID = SurfMesh.MaxTriangleID();
diff --git a/Engine/Plugins/Runtime/GeometryProcessing/Source/DynamicMesh/Private/Operations/IntrinsicTriangulationMesh.cpp b/Engine/Plugins/Runtime/GeometryProcessing/Source/DynamicMesh/Private/Operations/IntrinsicTriangulationMesh.cpp
--- a/Engine/Plugins/Runtime/GeometryProcessing/Source/DynamicMesh/Private/Operations/IntrinsicTriangulationMesh.cpp
+++ b/Engine/Plugins/Runtime/GeometryProcessing/Source/DynamicMesh/Private/Operations/IntrinsicTriangulationMesh.cpp
@@ -991,73 +991,95 @@ namespace FNormalCoordSurfaceTraceImpl
 			// where to start when visiting the intrinsic triangles that are adjacent to the startVID
 			const int32 VistorStartEID = IdentifyInitialAdjacentEdge(IntrinsicMesh, StartVID);
 
-			// when visiting each adjacent intrinsic triangle (in CCW order) test if the surface edge we trace is inside this triangle.
-			auto EdgeFinder = [&](int32 IntrinsicTID, int32 IntrinsicEID, int32 IdxOf)->bool
-								{
-
+			// test for the special case when the StartVID has valence one on the intrinsic mesh (ie a single triangle is wrapped into a cone)
+			const FIndex2i IntrinsicEdgeT  = IntrinsicMesh.GetEdgeT(VistorStartEID);
+			if (IntrinsicEdgeT.A == IntrinsicEdgeT.B)
+			{
+				const int32 IntrinsicEID     = VistorStartEID;
+				const int32 IntrinsicTID     = IntrinsicEdgeT.A;
+				const FIndex3i IntrinsicVIDs = IntrinsicMesh.GetTriangle(IntrinsicTID);
+				const int32 IdxOf            = IntrinsicVIDs.IndexOf(StartVID);
+				const int32 ThisRoundabout   = NCoords.RoundaboutOrder[IntrinsicTID][IdxOf];
+				const bool bOnSurfaceEdge    = NCoords.IsSurfaceEdgeSegment(IntrinsicEID);
+				const bool bEquivalentToSurface = bOnSurfaceEdge  && (ThisRoundabout == OrderOfTraceEID);
+
+				FinderInfo.bIsAlsoSurfaceEdge = bEquivalentToSurface;
+				FinderInfo.TID                = IntrinsicTID;
+				FinderInfo.EID                = IntrinsicEID;
+				FinderInfo.IdxOf              = IdxOf;
+				FinderInfo.FirstRoundabout    = ThisRoundabout;
+				FinderInfo.SecondRoundabout   = ThisRoundabout;
+			}
+			else
+			{ 
+				// when visiting each adjacent intrinsic triangle (in CCW order) test if the surface edge we trace is inside this triangle.
+				auto EdgeFinder = [&](int32 IntrinsicTID, int32 IntrinsicEID, int32 IdxOf)->bool
+									{
 
-									const FIndex3i IntrinsicTriEIDs = IntrinsicMesh.GetTriEdges(IntrinsicTID);
-									const int32 ThisRoundabout      = NCoords.RoundaboutOrder[IntrinsicTID][IdxOf];
-									const bool bOnSurfaceEdge       = NCoords.IsSurfaceEdgeSegment(IntrinsicEID);
-									const bool bEquivalentToSurface = bOnSurfaceEdge  && (ThisRoundabout == OrderOfTraceEID);
-									int32 NextRoundabout = -1;
 
-									bool bShouldBreak = false;
+										const FIndex3i IntrinsicTriEIDs = IntrinsicMesh.GetTriEdges(IntrinsicTID);
+										const int32 ThisRoundabout      = NCoords.RoundaboutOrder[IntrinsicTID][IdxOf];
+										const bool bOnSurfaceEdge       = NCoords.IsSurfaceEdgeSegment(IntrinsicEID);
+										const bool bEquivalentToSurface = bOnSurfaceEdge  && (ThisRoundabout == OrderOfTraceEID);
+										int32 NextRoundabout = -1;
 
-									if (bEquivalentToSurface) // test if the current intrinsic edge is the same as the surface mesh trace edge
-									{
-										bShouldBreak = true;
-									}
-									else // test if the edge starts in this intrinsic triangle
-									{
-										const int32 NextIntrinsicEID      = IntrinsicTriEIDs[(IdxOf + 2) % 3];
-										const FIndex2i NextIntrinsicEdgeT = IntrinsicMesh.GetEdgeT(NextIntrinsicEID);
-										const int32 NextIntrinsicTID      = (NextIntrinsicEdgeT.A == IntrinsicTID) ? NextIntrinsicEdgeT.B : NextIntrinsicEdgeT.A;
+										bool bShouldBreak = false;
 
-										if (NextIntrinsicTID != -1)
+										if (bEquivalentToSurface) // test if the current intrinsic edge is the same as the surface mesh trace edge
+										{
+											bShouldBreak = true;
+										}
+										else // test if the edge starts in this intrinsic triangle
 										{
-											const int32 NextIndexOf = IntrinsicMesh.GetTriEdges(NextIntrinsicTID).IndexOf(NextIntrinsicEID);
-											NextRoundabout          = NCoords.RoundaboutOrder[NextIntrinsicTID][NextIndexOf];
-											if (NextRoundabout < ThisRoundabout) // Order = {NextRO |zero cut | ThisRO}
+											const int32 NextIntrinsicEID      = IntrinsicTriEIDs[(IdxOf + 2) % 3];
+											const FIndex2i NextIntrinsicEdgeT = IntrinsicMesh.GetEdgeT(NextIntrinsicEID);
+											const int32 NextIntrinsicTID      = (NextIntrinsicEdgeT.A == IntrinsicTID) ? NextIntrinsicEdgeT.B : NextIntrinsicEdgeT.A;
+
+											if (NextIntrinsicTID != -1)
 											{
-												// we have crossed the zero roundabout cut ( e.g. NextRo = 2oclock and ThisRo = 11oclock)
-												if ( NextRoundabout > OrderOfTraceEID) // Order = {NextRO, Trace | zero cut | ThisRO}
+												const int32 NextIndexOf = IntrinsicMesh.GetTriEdges(NextIntrinsicTID).IndexOf(NextIntrinsicEID);
+												NextRoundabout          = NCoords.RoundaboutOrder[NextIntrinsicTID][NextIndexOf];
+												if (NextRoundabout < ThisRoundabout) // Order = {NextRO |zero cut | ThisRO}
 												{
-													bShouldBreak = true;
+													// we have crossed the zero roundabout cut ( e.g. NextRo = 2oclock and ThisRo = 11oclock)
+													if ( NextRoundabout > OrderOfTraceEID) // Order = {NextRO, Trace | zero cut | ThisRO}
+													{
+														bShouldBreak = true;
+													}
+													else 
+													{ 
+														NextRoundabout += ValenceOfStartVID; // Order = {NextRO, | zero cut |, Trace(?), ThisRO}
+													}
 												}
-												else 
-												{ 
-													NextRoundabout += ValenceOfStartVID; // Order = {NextRO, | zero cut |, Trace(?), ThisRO}
+												checkSlow(IntrinsicMesh.GetTriangle(NextIntrinsicTID)[NextIndexOf] == StartVID);
+												//checkSlow(NextRoundabout != 0 || ThisRoundabout != 0)
+												if ((NextRoundabout > OrderOfTraceEID) && (OrderOfTraceEID >= ThisRoundabout))
+												{
+													bShouldBreak = true;
 												}
 											}
-											checkSlow(IntrinsicMesh.GetTriangle(NextIntrinsicTID)[NextIndexOf] == StartVID);
-											//checkSlow(NextRoundabout != 0 || ThisRoundabout != 0)
-											if ((NextRoundabout > OrderOfTraceEID) && (OrderOfTraceEID >= ThisRoundabout))
+											else // mesh boundary case and we made it to the last triangle w/o finding this edge. it must be in this one.
 											{
 												bShouldBreak = true;
 											}
 										}
-										else // mesh boundary case and we made it to the last triangle w/o finding this edge. it must be in this one.
+
+										if (bShouldBreak)
 										{
-											bShouldBreak = true;
+											FinderInfo.bIsAlsoSurfaceEdge = bEquivalentToSurface;
+											FinderInfo.TID                = IntrinsicTID;
+											FinderInfo.EID                = IntrinsicEID;
+											FinderInfo.IdxOf              = IdxOf;
+											FinderInfo.FirstRoundabout    = ThisRoundabout;
+											FinderInfo.SecondRoundabout   = NextRoundabout;
 										}
-									}
 
-									if (bShouldBreak)
-									{
-										FinderInfo.bIsAlsoSurfaceEdge = bEquivalentToSurface;
-										FinderInfo.TID                = IntrinsicTID;
-										FinderInfo.EID                = IntrinsicEID;
-										FinderInfo.IdxOf              = IdxOf;
-										FinderInfo.FirstRoundabout    = ThisRoundabout;
-										FinderInfo.SecondRoundabout   = NextRoundabout;
-									}
-
-									return bShouldBreak;
-								};
-			VisitVertexAdjacentElements(IntrinsicMesh, StartVID, VistorStartEID, EdgeFinder);
+										return bShouldBreak;
+									};
+				VisitVertexAdjacentElements(IntrinsicMesh, StartVID, VistorStartEID, EdgeFinder);
 
-			checkSlow(FinderInfo.TID != -1); // should have found something!
+				checkSlow(FinderInfo.TID != -1); // should have found something!
+			}
 		}
 
 
@@ -1371,8 +1393,9 @@ namespace FNormalCoordSurfaceTraceImpl
 					const FIndex3i TriBVIDs = HostMesh.GetTriangle(HostEdgeT.B);
 
 					const int32 HostStartVID = StartSurfacePoint.Position.VertexPosition.VID;
-					
-					return (TriAVIDs.IndexOf(HostStartVID) != -1) ? HostEdgeT.A : HostEdgeT.B;
+					const int32 HostTID =  (TriAVIDs.IndexOf(HostStartVID) != -1) ? HostEdgeT.A : HostEdgeT.B;
+					checkSlow(HostMesh.GetTriangle(HostTID).IndexOf(HostStartVID) != -1);
+					return HostTID;
 				}
 				else if (IsEdgePoint(StartSurfacePoint))
 				{
@@ -1741,18 +1764,41 @@ int32 UE::Geometry::FlipToDelaunay(FIntrinsicEdgeFlipMesh& IntrinsicMesh, TSet<i
 * FSimpleIntrinsicEdgeFlipMesh Methods
 *------------------------------------------------------------------------------ */
 
-FSimpleIntrinsicEdgeFlipMesh::FSimpleIntrinsicEdgeFlipMesh(const FDynamicMesh3& SrcMesh) :
-	Vertices(SrcMesh.GetVerticesBuffer()),
-	VertexRefCounts{ SrcMesh.GetVerticesRefCounts() },
-	VertexEdgeLists{ SrcMesh.GetVertexEdges() },
+FSimpleIntrinsicEdgeFlipMesh::FSimpleIntrinsicEdgeFlipMesh(const FDynamicMesh3& SrcMesh)
+{
+	Reset(SrcMesh);
+}
+void FSimpleIntrinsicEdgeFlipMesh::Clear()
+{
+	Vertices.Clear();
+	VertexRefCounts.Clear();
+	VertexEdgeLists.Reset();
 
-	Triangles{ SrcMesh.GetTrianglesBuffer() },
-	TriangleRefCounts{ SrcMesh.GetTrianglesRefCounts() },
-	TriangleEdges{ SrcMesh.GetTriangleEdges() },
+	Triangles.Clear();
+	TriangleRefCounts.Clear();
+	TriangleEdges.Clear();
+	
+	Edges.Clear();
+	EdgeRefCounts.Clear();
+
+	EdgeLengths.Clear();
+	InternalAngles.Clear();
+}
 
-	Edges{ SrcMesh.GetEdgesBuffer() },
-	EdgeRefCounts{ SrcMesh.GetEdgesRefCounts() }
+void FSimpleIntrinsicEdgeFlipMesh::Reset(const FDynamicMesh3& SrcMesh)
 {
+	Clear();
+
+	Vertices        = SrcMesh.GetVerticesBuffer();
+	VertexRefCounts = SrcMesh.GetVerticesRefCounts();
+	VertexEdgeLists = SrcMesh.GetVertexEdges();
+
+	Triangles         = SrcMesh.GetTrianglesBuffer();
+	TriangleRefCounts = SrcMesh.GetTrianglesRefCounts();
+	TriangleEdges     = SrcMesh.GetTriangleEdges();
+
+	Edges         = SrcMesh.GetEdgesBuffer();
+	EdgeRefCounts = SrcMesh.GetEdgesRefCounts();
 	
 	const int32 MaxEID = MaxEdgeID();
 	EdgeLengths.SetNum(MaxEID);
@@ -2011,6 +2057,16 @@ EMeshResult FSimpleIntrinsicEdgeFlipMesh::FlipEdge(int32 EID, FEdgeFlipInfo& Edg
 		}
 	}
 
+	// prohibit case where one of the ends of the edge has degree one (this looks like a triangle rolled into a cone ).
+	{
+		FIndex2i EdgeT = GetEdgeT(EID);
+		
+		if (EdgeT.A == EdgeT.B)
+		{
+			return EMeshResult::Failed_Unsupported;
+		}
+	}
+
 	// compute the length of edge after flip
 	const double PostFlipLength = GetOpposingVerticesDistance(EID);
 
@@ -2958,9 +3014,10 @@ UE::Geometry::EMeshResult  UE::Geometry::FIntrinsicTriangulation::SplitEdge(int3
 *  FIntrinsicEdgeFlipMesh Methods
 *-------------------------------------------------------------------------------*/
 FIntrinsicEdgeFlipMesh::FIntrinsicEdgeFlipMesh(const FDynamicMesh3& SurfaceMesh)
-	: MyBase(SurfaceMesh)
-	, NormalCoordinates(SurfaceMesh)
 {
+	// construct the intrinsic mesh directly from this surface mesh
+	MyBase::Reset(SurfaceMesh);
+	NormalCoordinates.Reset(SurfaceMesh);
 }
 
 EMeshResult FIntrinsicEdgeFlipMesh::FlipEdge(int32 EID, FEdgeFlipInfo& EdgeFlipInfo)
diff --git a/Engine/Plugins/Runtime/GeometryProcessing/Source/DynamicMesh/Public/Operations/IntrinsicCorrespondenceUtils.h b/Engine/Plugins/Runtime/GeometryProcessing/Source/DynamicMesh/Public/Operations/IntrinsicCorrespondenceUtils.h
--- a/Engine/Plugins/Runtime/GeometryProcessing/Source/DynamicMesh/Public/Operations/IntrinsicCorrespondenceUtils.h
+++ b/Engine/Plugins/Runtime/GeometryProcessing/Source/DynamicMesh/Public/Operations/IntrinsicCorrespondenceUtils.h
@@ -96,10 +96,23 @@ namespace IntrinsicCorrespondenceUtils
 	// A Connection stores the information needed to define a local reference direction for each vertex and triangle.
 	struct DYNAMICMESH_API FMeshConnection
 	{
+		FMeshConnection() : SurfaceMesh(nullptr) {};
+
 		FMeshConnection(const FDynamicMesh3& SurfMesh);
-		const FDynamicMesh3* SurfaceMesh;               // Pointer back to original mesh 
-		TArray<int32> VIDToReferenceEID;                // Reference (original mesh) Edge per (original mesh) Vertex, defines local zero angle on tangent plane
-		TArray<int32> TIDToReferenceEID;                // Reference (original mesh) Edge per (original mesh) Triangle, defines local zero angle on tangent plane
+
+		/** Reset the connection to the given (not owned) surface mesh*/
+		void Reset(const FDynamicMesh3& SurfMesh);     
+
+		/** Reset the connection to used a internally owned surface mesh*/
+		void Reset(TUniquePtr<FDynamicMesh3>& SurfMesh);
+		
+
+		
+
+		const FDynamicMesh3* SurfaceMesh;                  // Pointer back the surface mesh used by the connection.  Will agree with OwnedSurfaceMesh if that is not null 
+		TUniquePtr<FDynamicMesh3> OwnedSurfaceMesh;        // Optional, the connection can own a surface mesh.
+		TArray<int32> VIDToReferenceEID;                   // Reference (original mesh) Edge per (original mesh) Vertex, defines local zero angle on tangent plane
+		TArray<int32> TIDToReferenceEID;                   // Reference (original mesh) Edge per (original mesh) Triangle, defines local zero angle on tangent plane
 	};
 
 	/**
@@ -111,8 +124,18 @@ namespace IntrinsicCorrespondenceUtils
 	{
 		typedef FMeshConnection  MyBase;
 
+		FNormalCoordinates(): MyBase() {}
+
 		FNormalCoordinates(const FDynamicMesh3& SurfMesh);
 
+		/** Reset the connection to the given (not owned) surface mesh*/
+		void Reset(const FDynamicMesh3& SurfMesh);
+
+		/** Reset the connection to used an internally owned surface mesh*/
+		void Reset(TUniquePtr<FDynamicMesh3>& SurfMesh);
+
+
+
 		struct FEdgeAndCrossingIdx
 		{
 			int32 TID;  // triangle traversed
@@ -172,6 +195,10 @@ namespace IntrinsicCorrespondenceUtils
 
 		TArray<int32>             RefVertDegree;   // Per vertex in ref mesh, caches the result of counting the edges in RefMesh.VtxEdgeIter().  
 		TArray<FIndex3i>          EdgeOrder;       // Per surface triangle, per directed edge: number of edges past reference edge traveling ccw.
+
+	protected:
+		/** rebuild the normal coordinate data in this connection. Assumes the base connection data is consistent with this surface mesh*/
+		void RebuildNormalCoordinates(const FDynamicMesh3& SurfMesh);
 	};
 
 
@@ -338,6 +365,7 @@ namespace IntrinsicCorrespondenceUtils
 			CurTriEIDs = Mesh.GetTriEdges(CurTID);
 			IndexOf    = CurTriEIDs.IndexOf(CurEID);
 
+			checkSlow(Mesh.GetTriangle(CurTID)[IndexOf] == VID);
 		} while (CurEID != StartEID);
 
 		return true;
diff --git a/Engine/Plugins/Runtime/GeometryProcessing/Source/DynamicMesh/Public/Operations/IntrinsicTriangulationMesh.h b/Engine/Plugins/Runtime/GeometryProcessing/Source/DynamicMesh/Public/Operations/IntrinsicTriangulationMesh.h
--- a/Engine/Plugins/Runtime/GeometryProcessing/Source/DynamicMesh/Public/Operations/IntrinsicTriangulationMesh.h
+++ b/Engine/Plugins/Runtime/GeometryProcessing/Source/DynamicMesh/Public/Operations/IntrinsicTriangulationMesh.h
@@ -69,12 +69,18 @@ public:
 	/** InvalidID indicates that a vertex/edge/triangle ID is invalid */
 	constexpr static int InvalidID = IndexConstants::InvalidID;
 
+	FSimpleIntrinsicEdgeFlipMesh(){};
+
 	/** Constructor does ID-matching deep copy the basic mesh topology (but no attributes or groups) */
 	FSimpleIntrinsicEdgeFlipMesh(const FDynamicMesh3& SrcMesh);
 
-	FSimpleIntrinsicEdgeFlipMesh() = delete;
 	FSimpleIntrinsicEdgeFlipMesh(const FSimpleIntrinsicEdgeFlipMesh&) = delete;
 
+	/** Discard all data */
+	void Clear();
+
+	/** Reset the mesh to an ID-matching deep copy the basic mesh topology (but no attributes or groups) */
+	void Reset(const FDynamicMesh3& SrcMesh);
 	/**
 	* Flip a single edge in the Intrinsic Mesh.
 	* @param EdgeFlipInfo [out] populated on return.
@@ -241,6 +247,12 @@ public:
 		return Vertices[VertexID];
 	}
 
+	/** @return upper bound on the vertex ID used in the mesh.  i.e. all edge IDs < MaxVertexID*/
+	int32 MaxVertexID() const
+	{
+		return (int32)VertexRefCounts.GetMaxIndex();
+	}
+
 	/** @return upper bound on the triangle ID used in the mesh.  i.e. all triangle IDs < MaxTriangleID */
 	int32 MaxTriangleID() const
 	{
@@ -266,6 +278,13 @@ public:
 		return Triangles[TriangleID];
 	}
 
+	/** @return  reference to vertex ids for the three corners of the triangle (e.g. a, b, c)*/
+	inline const FIndex3i& GetTriangleRef(int TriangleID) const
+	{
+		checkSlow(IsTriangle(TriangleID));
+		return Triangles[TriangleID];
+	}
+
 	/** @return edge ids for the three sides of the triangle (e.g. ab, bc, ca)*/
 	inline FIndex3i GetTriEdges(int32 TriangleID) const
 	{
@@ -403,11 +422,13 @@ public:
 	/** InvalidID indicates that a vertex/edge/triangle ID is invalid */
 	constexpr static int InvalidID = IndexConstants::InvalidID;
 
+	FSimpleIntrinsicMesh(){};
+
 	/** Constructor does ID-matching deep copy the basic mesh topology (but no attributes or groups) */
 	FSimpleIntrinsicMesh(const FDynamicMesh3& SrcMesh) : MyBase(SrcMesh)
 	{};
 
-	FSimpleIntrinsicMesh() = delete;
+	
 	FSimpleIntrinsicMesh(const FSimpleIntrinsicMesh&) = delete;
 
 
