fix mac build of ConstraintedDelaunay2: Graph.Edges() iterator does not return items by reference

diff --git a/Engine/Plugins/Runtime/GeometryProcessing/Source/GeometryAlgorithms/Public/ConstrainedDelaunay2.h b/Engine/Plugins/Runtime/GeometryProcessing/Source/GeometryAlgorithms/Public/ConstrainedDelaunay2.h
--- a/Engine/Plugins/Runtime/GeometryProcessing/Source/GeometryAlgorithms/Public/ConstrainedDelaunay2.h
+++ b/Engine/Plugins/Runtime/GeometryProcessing/Source/GeometryAlgorithms/Public/ConstrainedDelaunay2.h
@@ -86,7 +86,7 @@ struct TConstrainedDelaunay2
 				GraphToDTVertIdxMap[Idx] = -1;
 			}
 		}
-		for (const FDynamicGraph::FEdge& Edge : Graph.Edges())
+		for (const FDynamicGraph::FEdge Edge : Graph.Edges())
 		{
 			Edges.Add(FIndex2i(GraphToDTVertIdxMap[Edge.A], GraphToDTVertIdxMap[Edge.B]));
 		}
