SetCollisionGeometryTool: add Mesh to Level Set / SDF conversion

diff --git a/Engine/Plugins/Experimental/MeshModelingToolsetExp/Source/MeshModelingToolsExp/Private/Physics/CollisionGeometryVisualization.cpp b/Engine/Plugins/Experimental/MeshModelingToolsetExp/Source/MeshModelingToolsExp/Private/Physics/CollisionGeometryVisualization.cpp
--- a/Engine/Plugins/Experimental/MeshModelingToolsetExp/Source/MeshModelingToolsExp/Private/Physics/CollisionGeometryVisualization.cpp
+++ b/Engine/Plugins/Experimental/MeshModelingToolsetExp/Source/MeshModelingToolsExp/Private/Physics/CollisionGeometryVisualization.cpp
@@ -3,6 +3,7 @@
 #include "Physics/CollisionGeometryVisualization.h"
 #include "Generators/LineSegmentGenerators.h"
 #include "Util/ColorConstants.h"
+#include "Chaos/Levelset.h"
 
 using namespace UE::Geometry;
 
@@ -13,6 +14,7 @@ void UE::PhysicsTools::InitializePreviewGeometryLines(const FPhysicsDataCollecti
 	FColor SphereColor = LineColor;
 	FColor BoxColor = LineColor;
 	FColor ConvexColor = LineColor;
+	FColor LevelSetColor = LineColor;
 	FColor CapsuleColor = LineColor;
 	int32 OverrideColorIdx = 0;
 
@@ -141,7 +143,42 @@ void UE::PhysicsTools::InitializePreviewGeometryLines(const FPhysicsDataCollecti
 			}
 		});
 	}
+	
+	// for Level Sets draw the grid cells where phi < 0
+	for (int32 Index = 0; Index < AggGeom.LevelSetElems.Num(); Index++)
+	{
+		PreviewGeom->CreateOrUpdateLineSet(FString::Printf(TEXT("Level Set %d"), Index), 1, [&](int32 UnusedIndex, TArray<FRenderableLine>& LinesOut)
+		{
+			const FColor Color = GetNextColor(LevelSetColor);
+			const FKLevelSetElem& LevelSet = AggGeom.LevelSetElems[Index];
+			
+			FTransform ElemTransform = LevelSet.GetTransform();
+			ElemTransform.ScaleTranslation(PhysicsData.ExternalScale3D);
+			ElemTransform.MultiplyScale3D(PhysicsData.ExternalScale3D);
+
+			auto GenerateBoxSegmentsFromFBox = [&](const FBox& Box)
+			{
+				const FVector3d Center = Box.GetCenter();
+				const FVector3d HalfDimensions = 0.5 * (Box.Max - Box.Min);
+
+				UE::Geometry::GenerateBoxSegments<double>(HalfDimensions, Center, FVector3d::UnitX(), FVector3d::UnitY(), FVector3d::UnitZ(), ElemTransform,
+					[&](const FVector3d& A, const FVector3d& B) { LinesOut.Add(FRenderableLine((FVector)A, (FVector)B, Color, LineThickness, DepthBias)); });
+			};
+
+			const FBox TotalGridBox = LevelSet.UntransformedAABB();
+			GenerateBoxSegmentsFromFBox(TotalGridBox);
 
+			TArray<FBox> CellBoxes;
+			const double Threshold = UE_KINDA_SMALL_NUMBER;		// allow slightly greater than zero for visualization purposes
+			LevelSet.GetInteriorGridCells(CellBoxes, Threshold);
+
+			for (const FBox& CellBox : CellBoxes)
+			{
+				GenerateBoxSegmentsFromFBox(CellBox);
+			}
+
+		});
+	}
 
 	// Unclear whether we actually use these in the Engine, for UBodySetup? Does not appear to be supported by UxX import system,
 	// and online documentation suggests they may only be supported for cloth?
diff --git a/Engine/Plugins/Experimental/MeshModelingToolsetExp/Source/MeshModelingToolsExp/Private/Physics/CollisionPropertySets.cpp b/Engine/Plugins/Experimental/MeshModelingToolsetExp/Source/MeshModelingToolsExp/Private/Physics/CollisionPropertySets.cpp
--- a/Engine/Plugins/Experimental/MeshModelingToolsetExp/Source/MeshModelingToolsExp/Private/Physics/CollisionPropertySets.cpp
+++ b/Engine/Plugins/Experimental/MeshModelingToolsetExp/Source/MeshModelingToolsExp/Private/Physics/CollisionPropertySets.cpp
@@ -57,6 +57,13 @@ void UE::PhysicsTools::InitializePhysicsToolObjectPropertySet(const FPhysicsData
 		ConvexData.Element = Convex;
 		PropSet->Convexes.Add(ConvexData);
 	}
+
+	for (const FKLevelSetElem& LevelSet: AggGeom.LevelSetElems)
+	{
+		FPhysicsLevelSetData LevelSetData;
+		LevelSetData.Element = LevelSet;
+		PropSet->LevelSets.Add(LevelSetData);
+	}
 }
 
 
@@ -68,4 +75,5 @@ void UPhysicsObjectToolPropertySet::Reset()
 	Boxes.Reset();
 	Capsules.Reset();
 	Convexes.Reset();
+	LevelSets.Reset();
 }
\ No newline at end of file
diff --git a/Engine/Plugins/Experimental/MeshModelingToolsetExp/Source/MeshModelingToolsExp/Private/Physics/SetCollisionGeometryTool.cpp b/Engine/Plugins/Experimental/MeshModelingToolsetExp/Source/MeshModelingToolsExp/Private/Physics/SetCollisionGeometryTool.cpp
--- a/Engine/Plugins/Experimental/MeshModelingToolsetExp/Source/MeshModelingToolsExp/Private/Physics/SetCollisionGeometryTool.cpp
+++ b/Engine/Plugins/Experimental/MeshModelingToolsetExp/Source/MeshModelingToolsExp/Private/Physics/SetCollisionGeometryTool.cpp
@@ -101,6 +101,9 @@ public:
 			UseShapeGenerator->Generate_ProjectedHulls(NewCollision->Geometry,
 				(FMeshSimpleShapeApproximation::EProjectedHullAxisMode)(int32)SweepAxis);
 			break;
+		case ECollisionGeometryType::LevelSets:
+			UseShapeGenerator->Generate_LevelSets(NewCollision->Geometry);
+			break;
 		case ECollisionGeometryType::MinVolume:
 			UseShapeGenerator->Generate_MinVolume(NewCollision->Geometry);
 			break;
@@ -220,6 +223,7 @@ void USetCollisionGeometryTool::Setup()
 	Settings->WatchProperty(Settings->bSimplifyPolygons, [this](bool) { InvalidateCompute(); });
 	Settings->WatchProperty(Settings->HullTolerance, [this](float) { InvalidateCompute(); });
 	Settings->WatchProperty(Settings->SweepAxis, [this](EProjectedHullAxis) { InvalidateCompute(); });
+	Settings->WatchProperty(Settings->LevelSetResolution, [this](int32) { InvalidateCompute(); });
 
 	if (InitialSourceMeshes.Num() == 1)
 	{
@@ -295,6 +299,7 @@ TUniquePtr<UE::Geometry::TGenericDataOperator<FPhysicsDataCollection>> USetColli
 	Op->UseShapeGenerator->ConvexDecompositionErrorTolerance = Settings->AddHullsErrorTolerance;
 	Op->UseShapeGenerator->ConvexDecompositionMinPartThickness = Settings->MinPartThickness;
 	Op->UseShapeGenerator->HullSimplifyTolerance = Settings->HullTolerance;
+	Op->UseShapeGenerator->LevelSetGridResolution = Settings->LevelSetResolution;
 
 	Op->ComputeType = Settings->GeometryType;
 	Op->bAppendToExisting = Settings->bAppendToExisting;
diff --git a/Engine/Plugins/Experimental/MeshModelingToolsetExp/Source/MeshModelingToolsExp/Public/Physics/CollisionPropertySets.h b/Engine/Plugins/Experimental/MeshModelingToolsetExp/Source/MeshModelingToolsExp/Public/Physics/CollisionPropertySets.h
--- a/Engine/Plugins/Experimental/MeshModelingToolsetExp/Source/MeshModelingToolsExp/Public/Physics/CollisionPropertySets.h
+++ b/Engine/Plugins/Experimental/MeshModelingToolsetExp/Source/MeshModelingToolsExp/Public/Physics/CollisionPropertySets.h
@@ -89,6 +89,14 @@ struct MESHMODELINGTOOLSEXP_API FPhysicsConvexData
 	FKShapeElem Element;
 };
 
+USTRUCT()
+struct MESHMODELINGTOOLSEXP_API FPhysicsLevelSetData
+{
+	GENERATED_BODY()
+
+	UPROPERTY(VisibleAnywhere, Category = LevelSet)
+	FKShapeElem Element;
+};
 
 UCLASS()
 class MESHMODELINGTOOLSEXP_API UPhysicsObjectToolPropertySet : public UInteractiveToolPropertySet
@@ -113,6 +121,9 @@ public:
 	UPROPERTY(VisibleAnywhere, Category = PhysicsData)
 	TArray<FPhysicsConvexData> Convexes;
 
+	UPROPERTY(VisibleAnywhere, Category = PhysicsData)
+	TArray<FPhysicsLevelSetData> LevelSets;
+
 	void Reset();
 };
 
diff --git a/Engine/Plugins/Experimental/MeshModelingToolsetExp/Source/MeshModelingToolsExp/Public/Physics/SetCollisionGeometryTool.h b/Engine/Plugins/Experimental/MeshModelingToolsetExp/Source/MeshModelingToolsExp/Public/Physics/SetCollisionGeometryTool.h
--- a/Engine/Plugins/Experimental/MeshModelingToolsetExp/Source/MeshModelingToolsExp/Public/Physics/SetCollisionGeometryTool.h
+++ b/Engine/Plugins/Experimental/MeshModelingToolsetExp/Source/MeshModelingToolsExp/Public/Physics/SetCollisionGeometryTool.h
@@ -57,6 +57,7 @@ enum class ECollisionGeometryType
 	Capsules = 4,
 	ConvexHulls = 5,
 	SweptHulls = 6,
+	LevelSets = 7,
 	MinVolume = 10,
 
 	None = 11
@@ -149,6 +150,13 @@ public:
 		EditConditionHides, EditCondition = "GeometryType == ECollisionGeometryType::SweptHulls"))
 	EProjectedHullAxis SweepAxis = EProjectedHullAxis::SmallestVolume;
 
+	// Level Set settings
+
+	/** Level set grid resolution along longest grid axis */
+	UPROPERTY(EditAnywhere, Category = LevelSets, meta = (UIMin = "3", UIMax = "100", ClampMin = "3", ClampMax = "1000",
+		EditConditionHides, EditCondition = "GeometryType == ECollisionGeometryType::LevelSets"))
+	int32 LevelSetResolution = 10;
+
 	UPROPERTY(EditAnywhere, Category = OutputOptions)
 	bool bAppendToExisting = false;
 
diff --git a/Engine/Plugins/Runtime/GeometryProcessing/Source/DynamicMesh/Private/ShapeApproximation/MeshSimpleShapeApproximation.cpp b/Engine/Plugins/Runtime/GeometryProcessing/Source/DynamicMesh/Private/ShapeApproximation/MeshSimpleShapeApproximation.cpp
--- a/Engine/Plugins/Runtime/GeometryProcessing/Source/DynamicMesh/Private/ShapeApproximation/MeshSimpleShapeApproximation.cpp
+++ b/Engine/Plugins/Runtime/GeometryProcessing/Source/DynamicMesh/Private/ShapeApproximation/MeshSimpleShapeApproximation.cpp
@@ -8,13 +8,15 @@
 #include "FitCapsule3.h"
 #include "CompGeom/ConvexDecomposition3.h"
 //#include "DynamicMesh/DynamicMeshAABBTree3.h"
-
+#include "Implicit/SweepingMeshSDF.h"
 #include "ShapeApproximation/ShapeDetection3.h"
 #include "MeshQueries.h"
 #include "Operations/MeshConvexHull.h"
 #include "Operations/MeshProjectionHull.h"
 #include "Util/ProgressCancel.h"
 
+#define LOCTEXT_NAMESPACE "MeshSimpleShapeApproximation"
+
 using namespace UE::Geometry;
 
 
@@ -410,7 +412,49 @@ void FMeshSimpleShapeApproximation::Generate_ProjectedHulls(FSimpleShapeSet3d& S
 }
 
 
+void FMeshSimpleShapeApproximation::Generate_LevelSets(FSimpleShapeSet3d& ShapeSetOut, FProgressCancel* Progress)
+{
+	FCriticalSection GeometryLock;
+	ParallelFor(SourceMeshes.Num(), [&](int32 MeshIndex)
+	{
+		if (GetDetectedSimpleShape(SourceMeshCaches[MeshIndex], ShapeSetOut, GeometryLock))
+		{
+			return;
+		}
+
+		const FAxisAlignedBox3d Bounds = SourceMeshes[MeshIndex]->GetBounds();
+		const double CellSize = Bounds.MaxDim() / LevelSetGridResolution;
+
+		TMeshAABBTree3<FDynamicMesh3> Spatial(SourceMeshes[MeshIndex]);
+
+		TSweepingMeshSDF<FDynamicMesh3> SDF;
+		SDF.Mesh = SourceMeshes[MeshIndex];
+		SDF.Spatial = &Spatial;
+		SDF.ComputeMode = TSweepingMeshSDF<FDynamicMesh3>::EComputeModes::NarrowBand_SpatialFloodFill;
+		SDF.CellSize = CellSize;
+		SDF.NarrowBandMaxDistance = 2.0 * CellSize;
+		SDF.ExactBandWidth = FMath::CeilToInt(SDF.NarrowBandMaxDistance / CellSize);
+		SDF.ExpandBounds = 2.0 * CellSize * FVector3d::One();
+
+		if (SDF.Compute(Bounds))
+		{
+			FLevelSetShape3d NewLevelSet;
+			NewLevelSet.GridTransform = FTransform((FVector3d)SDF.GridOrigin);
+			NewLevelSet.Grid = MoveTemp(SDF.Grid);
+			NewLevelSet.CellSize = SDF.CellSize;
 
+			GeometryLock.Lock();
+			ShapeSetOut.LevelSets.Add(MoveTemp(NewLevelSet));
+			GeometryLock.Unlock();
+		}
+		else if (Progress)
+		{
+			GeometryLock.Lock();
+			Progress->AddWarning(LOCTEXT("Generate_LevelSets_Failed", "Generating a new Level Set failed"), FProgressCancel::EMessageLevel::UserWarning);
+			GeometryLock.Unlock();
+		}
+	});
+}
 
 void FMeshSimpleShapeApproximation::Generate_MinVolume(FSimpleShapeSet3d& ShapeSetOut)
 {
@@ -465,4 +509,6 @@ void FMeshSimpleShapeApproximation::Generate_MinVolume(FSimpleShapeSet3d& ShapeS
 			GeometryLock.Unlock();
 		}
 	});
-}
\ No newline at end of file
+}
+
+#undef LOCTEXT_NAMESPACE
diff --git a/Engine/Plugins/Runtime/GeometryProcessing/Source/DynamicMesh/Private/ShapeApproximation/SimpleShapeSet3.cpp b/Engine/Plugins/Runtime/GeometryProcessing/Source/DynamicMesh/Private/ShapeApproximation/SimpleShapeSet3.cpp
--- a/Engine/Plugins/Runtime/GeometryProcessing/Source/DynamicMesh/Private/ShapeApproximation/SimpleShapeSet3.cpp
+++ b/Engine/Plugins/Runtime/GeometryProcessing/Source/DynamicMesh/Private/ShapeApproximation/SimpleShapeSet3.cpp
@@ -55,6 +55,10 @@ static void FilterContained(const FSimpleShapeSet3d& Geometry, const FSphereShap
 			{
 				bContained = UE::Geometry::IsInside(Sphere.Sphere, Geometry.Convexes[ElemIdx].Mesh.VerticesItr());
 			}
+			else if (Elements[j].Type == ESimpleShapeType::LevelSet)
+			{
+				bContained = UE::Geometry::IsInside(Sphere.Sphere, Geometry.LevelSets[ElemIdx].GetGridBox());
+			}
 			else
 			{
 				ensure(false);		// not implemented yet?
@@ -96,6 +100,10 @@ static void FilterContained(const FSimpleShapeSet3d& Geometry, const UE::Geometr
 			{
 				bContained = UE::Geometry::IsInside(Capsule.Capsule, Geometry.Convexes[ElemIdx].Mesh.VerticesItr());
 			}
+			else if (Elements[j].Type == ESimpleShapeType::LevelSet)
+			{
+				bContained = UE::Geometry::IsInside(Capsule.Capsule, Geometry.LevelSets[ElemIdx].GetGridBox());
+			}
 			else
 			{
 				ensure(false);
@@ -137,6 +145,10 @@ static void FilterContained(const FSimpleShapeSet3d& Geometry, const FBoxShape3d
 			{
 				bContained = UE::Geometry::IsInside(Box.Box, Geometry.Convexes[ElemIdx].Mesh.VerticesItr());
 			}
+			else if (Elements[j].Type == ESimpleShapeType::LevelSet)
+			{
+				bContained = UE::Geometry::IsInside(Box.Box, Geometry.LevelSets[ElemIdx].GetGridBox());
+			}
 			else
 			{
 				ensure(false);
@@ -187,6 +199,10 @@ static void FilterContained(const FSimpleShapeSet3d& Geometry, const FConvexShap
 			{
 				bContained = UE::Geometry::IsInsideHull<double>(Planes, Geometry.Convexes[ElemIdx].Mesh.VerticesItr());
 			}
+			else if (Elements[j].Type == ESimpleShapeType::LevelSet)
+			{
+				bContained = UE::Geometry::IsInsideHull<double>(Planes, Geometry.LevelSets[ElemIdx].GetGridBox());
+			}
 			else
 			{
 				ensure(false);
@@ -201,6 +217,53 @@ static void FilterContained(const FSimpleShapeSet3d& Geometry, const FConvexShap
 }
 
 
+static void FilterContained(const FSimpleShapeSet3d& Geometry, FLevelSetShape3d& LevelSetShape, const TArray<FSimpleShapeElementKey>& Elements, int32 k, TArray<bool>& RemovedInOut)
+{
+	FOrientedBox3d GridBox = LevelSetShape.GetGridBox();
+
+	int32 N = Elements.Num();
+	for (int32 j = k + 1; j < N; ++j)
+	{
+		if (RemovedInOut[j] == false)
+		{
+			// TODO: use the full GridTransform, not just translation. JIRA UE-155269
+			TTriLinearGridInterpolant<FDenseGrid3f> GridInterp(&LevelSetShape.Grid, LevelSetShape.GridTransform.GetTranslation(), LevelSetShape.CellSize, LevelSetShape.Grid.GetDimensions());
+
+			bool bContained = false;
+			int32 ElemIdx = Elements[j].Index;
+			if (Elements[j].Type == ESimpleShapeType::Sphere)
+			{
+				bContained = UE::Geometry::IsInside(GridInterp, Geometry.Spheres[ElemIdx].Sphere);
+			}
+			else if (Elements[j].Type == ESimpleShapeType::Box)
+			{
+				bContained = UE::Geometry::IsInside(GridInterp, Geometry.Boxes[ElemIdx].Box);
+			}
+			else if (Elements[j].Type == ESimpleShapeType::Capsule)
+			{
+				bContained = UE::Geometry::IsInside(GridInterp, Geometry.Capsules[ElemIdx].Capsule);
+			}
+			else if (Elements[j].Type == ESimpleShapeType::Convex)
+			{
+				bContained = UE::Geometry::IsInside<double>(GridInterp, Geometry.Convexes[ElemIdx].Mesh.VerticesItr());
+			}
+			else if (Elements[j].Type == ESimpleShapeType::LevelSet)
+			{
+				bContained = UE::Geometry::IsInside(GridInterp, Geometry.LevelSets[ElemIdx].GetGridBox());
+			}
+			else
+			{
+				ensure(false);
+			}
+
+			if (bContained)
+			{
+				RemovedInOut[j] = true;
+			}
+		}
+	}
+}
+
 
 static void GetElementsList(FSimpleShapeSet3d& GeometrySet, TArray<FSimpleShapeElementKey>& Elements)
 {
@@ -221,6 +284,11 @@ static void GetElementsList(FSimpleShapeSet3d& GeometrySet, TArray<FSimpleShapeE
 		FVector2d VolArea = TMeshQueries<FDynamicMesh3>::GetVolumeArea(GeometrySet.Convexes[k].Mesh);
 		Elements.Add(FSimpleShapeElementKey(ESimpleShapeType::Convex, k, VolArea.X));
 	}
+	for (int32 k = 0; k < GeometrySet.LevelSets.Num(); ++k)
+	{
+		Elements.Add(FSimpleShapeElementKey(ESimpleShapeType::LevelSet, k, GeometrySet.LevelSets[k].GetGridBox().Volume()));
+	}
+
 }
 
 static void GetElementsSortedByDecreasing(FSimpleShapeSet3d& GeometrySet, TArray<FSimpleShapeElementKey>& Elements)
@@ -264,6 +332,11 @@ void FSimpleShapeSet3d::RemoveContainedGeometry()
 		{
 			FilterContained(*this, Convexes[ElemIdx], Elements, k, Removed);
 		}
+		else if (ElemType == ESimpleShapeType::LevelSet)
+		{
+			// Pass in zero here instead of k because we don't currently have a good volume value for level set objects
+			FilterContained(*this, LevelSets[ElemIdx], Elements, 0, Removed);
+		}
 		else
 		{
 			ensure(false);
@@ -294,6 +367,9 @@ void FSimpleShapeSet3d::RemoveContainedGeometry()
 			case ESimpleShapeType::Convex:
 				NewSet.Convexes.Add(Convexes[ElemIdx]);		// todo movetemp here...
 				break;
+			case ESimpleShapeType::LevelSet:
+				NewSet.LevelSets.Add(LevelSets[ElemIdx]);		// todo movetemp here...
+				break;
 			}
 		}
 	}
@@ -303,6 +379,7 @@ void FSimpleShapeSet3d::RemoveContainedGeometry()
 	Boxes = MoveTemp(NewSet.Boxes);
 	Capsules = MoveTemp(NewSet.Capsules);
 	Convexes = MoveTemp(NewSet.Convexes);
+	LevelSets = MoveTemp(NewSet.LevelSets);
 }
 
 
@@ -338,6 +415,12 @@ static void TransformCapsuleShape(UE::Geometry::FCapsuleShape3d& CapsuleShape, c
 	CapsuleShape.Capsule.Radius *= RadiusScale;
 }
 
+static void TransformLevelSetShape(UE::Geometry::FLevelSetShape3d& LevelSetShape, const FTransform3d& Transform)
+{
+	const FTransform3d Sub(-0.5 * LevelSetShape.CellSize * FVector3d::One());
+	const FTransform3d Add(0.5 * LevelSetShape.CellSize * FVector3d::One());
+	LevelSetShape.GridTransform = LevelSetShape.GridTransform * Sub * Transform * Add;
+}
 
 
 
@@ -386,7 +469,20 @@ static void TransformCapsuleShape(UE::Geometry::FCapsuleShape3d& CapsuleShape, c
 	CapsuleShape.Capsule.Radius *= RadiusScale;
 }
 
+static void TransformLevelSetShape(UE::Geometry::FLevelSetShape3d& LevelSetShape, const TArray<FTransform3d>& TransformSequence)
+{
+	const FTransform3d Sub(-0.5 * LevelSetShape.CellSize * FVector3d::One());
+	const FTransform3d Add(0.5 * LevelSetShape.CellSize * FVector3d::One());
+
+	LevelSetShape.GridTransform = LevelSetShape.GridTransform * Sub;
 
+	for (const FTransform3d& XForm : TransformSequence)
+	{
+		LevelSetShape.GridTransform = LevelSetShape.GridTransform * XForm;
+	}
+
+	LevelSetShape.GridTransform = LevelSetShape.GridTransform * Add;
+}
 
 
 void FSimpleShapeSet3d::Append(const FSimpleShapeSet3d& OtherShapeSet)
@@ -411,6 +507,12 @@ void FSimpleShapeSet3d::Append(const FSimpleShapeSet3d& OtherShapeSet)
 	{
 		Convexes.Add(ConvexShape);
 	}
+
+	LevelSets.Reserve(LevelSets.Num() + OtherShapeSet.LevelSets.Num());
+	for (const FLevelSetShape3d& LevelSetShape : OtherShapeSet.LevelSets)
+	{
+		LevelSets.Add(LevelSetShape);
+	}
 }
 
 
@@ -440,6 +542,13 @@ void FSimpleShapeSet3d::Append(const FSimpleShapeSet3d& OtherShapeSet, const FTr
 		Convexes.Add(ConvexShape);
 		MeshTransforms::ApplyTransform(Convexes.Last().Mesh, Transform);
 	}
+
+	LevelSets.Reserve(LevelSets.Num() + OtherShapeSet.LevelSets.Num());
+	for (const FLevelSetShape3d& LevelSetShape : OtherShapeSet.LevelSets)
+	{
+		LevelSets.Add(LevelSetShape);
+		TransformLevelSetShape(LevelSets.Last(), Transform);
+	}
 }
 
 
@@ -473,6 +582,13 @@ void FSimpleShapeSet3d::Append(const FSimpleShapeSet3d& OtherShapeSet, const TAr
 			MeshTransforms::ApplyTransform(Convexes.Last().Mesh, XForm);
 		}
 	}
+
+	LevelSets.Reserve(LevelSets.Num() + OtherShapeSet.LevelSets.Num());
+	for (const FLevelSetShape3d& LevelSetShape : OtherShapeSet.LevelSets)
+	{
+		LevelSets.Add(LevelSetShape);
+		TransformLevelSetShape(LevelSets.Last(), TransformSequence);
+	}
 }
 
 
@@ -506,6 +622,9 @@ void FSimpleShapeSet3d::FilterByVolume(int32 MaximumCount)
 		case ESimpleShapeType::Convex:
 			NewSet.Convexes.Add(Convexes[ElemIdx]);		// todo movetemp here...
 			break;
+		case ESimpleShapeType::LevelSet:
+			NewSet.LevelSets.Add(LevelSets[ElemIdx]);		// todo movetemp here...
+			break;
 		}
 	}
 
@@ -513,6 +632,7 @@ void FSimpleShapeSet3d::FilterByVolume(int32 MaximumCount)
 	Boxes = MoveTemp(NewSet.Boxes);
 	Capsules = MoveTemp(NewSet.Capsules);
 	Convexes = MoveTemp(NewSet.Convexes);
+	LevelSets = MoveTemp(NewSet.LevelSets);
 }
 
 
@@ -538,4 +658,9 @@ void FSimpleShapeSet3d::ApplyTransform(const FTransform3d& Transform)
 	{
 		MeshTransforms::ApplyTransform(ConvexShape.Mesh, Transform);
 	}
+
+	for (FLevelSetShape3d& LevelSetShape : LevelSets)
+	{
+		TransformLevelSetShape(LevelSetShape, Transform);
+	}
 }
diff --git a/Engine/Plugins/Runtime/GeometryProcessing/Source/DynamicMesh/Public/ShapeApproximation/MeshSimpleShapeApproximation.h b/Engine/Plugins/Runtime/GeometryProcessing/Source/DynamicMesh/Public/ShapeApproximation/MeshSimpleShapeApproximation.h
--- a/Engine/Plugins/Runtime/GeometryProcessing/Source/DynamicMesh/Public/ShapeApproximation/MeshSimpleShapeApproximation.h
+++ b/Engine/Plugins/Runtime/GeometryProcessing/Source/DynamicMesh/Public/ShapeApproximation/MeshSimpleShapeApproximation.h
@@ -79,6 +79,9 @@ public:
 
 	bool bUseExactComputationForBox = false;
 
+	/** Level Set Grid resolution along longest axis */
+	int32 LevelSetGridResolution = 10;
+
 	//
 	// setup/initialization
 	//
@@ -155,6 +158,8 @@ public:
 	 */
 	void Generate_ProjectedHulls(FSimpleShapeSet3d& ShapeSetOut, EProjectedHullAxisMode AxisMode);
 
+	void Generate_LevelSets(FSimpleShapeSet3d& ShapeSetOut, FProgressCancel* Progress = nullptr);
+	
 
 	/**
 	 * Fit containing axis-aligned box, oriented box, capsule, and sphere to each input mesh, and
diff --git a/Engine/Plugins/Runtime/GeometryProcessing/Source/DynamicMesh/Public/ShapeApproximation/SimpleShapeSet3.h b/Engine/Plugins/Runtime/GeometryProcessing/Source/DynamicMesh/Public/ShapeApproximation/SimpleShapeSet3.h
--- a/Engine/Plugins/Runtime/GeometryProcessing/Source/DynamicMesh/Public/ShapeApproximation/SimpleShapeSet3.h
+++ b/Engine/Plugins/Runtime/GeometryProcessing/Source/DynamicMesh/Public/ShapeApproximation/SimpleShapeSet3.h
@@ -9,6 +9,7 @@
 #include "CapsuleTypes.h"
 #include "SphereTypes.h"
 #include "DynamicMesh/DynamicMesh3.h"
+#include "Spatial/DenseGrid3.h"
 #include "TransformTypes.h"
 
 
@@ -24,7 +25,8 @@ enum class ESimpleShapeType
 	Sphere = 2,
 	Box = 4,
 	Capsule = 8,
-	Convex = 16
+	Convex = 16,
+	LevelSet = 32
 };
 ENUM_CLASS_FLAGS(ESimpleShapeType);
 
@@ -106,6 +108,42 @@ struct DYNAMICMESH_API FConvexShape3d
 	ESimpleShapeType GetShapeType() const { return ESimpleShapeType::Convex; }
 };
 
+/**
+ * FLevelSetShape is a 3D signed distance function sampled on a regular grid
+ */
+
+struct DYNAMICMESH_API FLevelSetShape3d
+{
+	FTransform3d GridTransform;
+	FDenseGrid3f Grid;
+	float CellSize;
+
+	FLevelSetShape3d() = default;
+
+	FLevelSetShape3d(const FTransform3d& GridTransform, const FDenseGrid3f& Grid, float CellSize) :
+		GridTransform(GridTransform),
+		Grid(Grid),
+		CellSize(CellSize)
+	{
+	}
+
+	FLevelSetShape3d(FTransform3d&& GridTransform, FDenseGrid3f&& Grid, float CellSize) :
+		GridTransform(MoveTemp(GridTransform)),
+		Grid(MoveTemp(Grid)),
+		CellSize(CellSize)
+	{
+	}
+
+	ESimpleShapeType GetShapeType() const { return ESimpleShapeType::LevelSet; }
+
+	FOrientedBox3d GetGridBox() const
+	{
+		const FVector3d BoxExtents = FVector3d(0.5 * CellSize * Grid.GetDimensions());
+		const FFrame3d BoxFrame(GridTransform.GetTranslation(), GridTransform.GetRotation());
+		return FOrientedBox3d(BoxFrame, BoxExtents);
+	}
+};
+
 
 /**
  * FSimpleShapeSet stores a set of simple geometry shapes useful for things like collision detection/etc.
@@ -117,9 +155,10 @@ struct DYNAMICMESH_API FSimpleShapeSet3d
 	TArray<FBoxShape3d> Boxes;
 	TArray<UE::Geometry::FCapsuleShape3d> Capsules;
 	TArray<FConvexShape3d> Convexes;
+	TArray<FLevelSetShape3d> LevelSets;
 
 	/** @return total number of elements in all sets */
-	int32 TotalElementsNum() const { return Spheres.Num() + Boxes.Num() + Capsules.Num() + Convexes.Num(); }
+	int32 TotalElementsNum() const { return Spheres.Num() + Boxes.Num() + Capsules.Num() + Convexes.Num() + LevelSets.Num(); }
 
 	/**
 	 * Append elements of another shape set
diff --git a/Engine/Plugins/Runtime/MeshModelingToolset/Source/ModelingComponents/Private/Physics/PhysicsDataCollection.cpp b/Engine/Plugins/Runtime/MeshModelingToolset/Source/ModelingComponents/Private/Physics/PhysicsDataCollection.cpp
--- a/Engine/Plugins/Runtime/MeshModelingToolset/Source/ModelingComponents/Private/Physics/PhysicsDataCollection.cpp
+++ b/Engine/Plugins/Runtime/MeshModelingToolset/Source/ModelingComponents/Private/Physics/PhysicsDataCollection.cpp
@@ -126,6 +126,14 @@ void FPhysicsDataCollection::CopyGeometryToAggregate()
 
 		AggGeom.ConvexElems.Add(Element);
 	}
+
+	for (UE::Geometry::FLevelSetShape3d& LevelSetGeom : Geometry.LevelSets)
+	{
+		FKLevelSetElem Element;
+		UE::Geometry::GetFKElement(LevelSetGeom.GridTransform, LevelSetGeom.Grid, LevelSetGeom.CellSize, Element);
+		AggGeom.LevelSetElems.Emplace(MoveTemp(Element));
+	}
+
 }
 
 
diff --git a/Engine/Plugins/Runtime/MeshModelingToolset/Source/ModelingComponents/Public/Physics/CollisionGeometryConversion.h b/Engine/Plugins/Runtime/MeshModelingToolset/Source/ModelingComponents/Public/Physics/CollisionGeometryConversion.h
--- a/Engine/Plugins/Runtime/MeshModelingToolset/Source/ModelingComponents/Public/Physics/CollisionGeometryConversion.h
+++ b/Engine/Plugins/Runtime/MeshModelingToolset/Source/ModelingComponents/Public/Physics/CollisionGeometryConversion.h
@@ -10,7 +10,7 @@
 #include "SphereTypes.h"
 #include "DynamicMesh/DynamicMesh3.h"
 #include "ShapeApproximation/SimpleShapeSet3.h"
-
+#include "Spatial/DenseGrid3.h"
 #include "PhysicsEngine/AggregateGeom.h"
 
 
@@ -70,7 +70,34 @@ namespace UE
 			ConvexInOut.UpdateElemBox();
 		}
 
+		/**
+		* Create FKLevelSet from the given grid information
+		*/
+		inline void GetFKElement(const FTransform3d& GridTransform, const FDenseGrid3f& Grid, float CellSize, FKLevelSetElem& LevelSetOut)
+		{
+			const FVector3i InGridDims = Grid.GetDimensions();
+			const FIntVector OutGridDims(InGridDims[0], InGridDims[1], InGridDims[2]);
+
+			TArray<double> OutGridValues;
+			OutGridValues.Init(0.0, OutGridDims[0] * OutGridDims[1] * OutGridDims[2]);
+
+			for (int I = 0; I < OutGridDims[0]; ++I)
+			{
+				for (int J = 0; J < OutGridDims[1]; ++J)
+				{
+					for (int K = 0; K < OutGridDims[2]; ++K)
+					{
+						const int InBufferIndex = I + InGridDims[0] * (J + InGridDims[1] * K);
+						const int OutBufferIndex = K + OutGridDims[2] * (J + OutGridDims[1] * I);
+						OutGridValues[OutBufferIndex] = Grid[InBufferIndex];
+					}
+				}
+			}
 
+			FTransform3d ChaosTransform = GridTransform;
+			ChaosTransform.AddToTranslation(-0.5 * CellSize * FVector::One());
+			LevelSetOut.BuildLevelSet(ChaosTransform, OutGridValues, OutGridDims, CellSize);
+		}
 
 
 		/**
@@ -130,6 +157,41 @@ namespace UE
 			GetShape(ConvexElem, ConvexShapeOut.Mesh);
 		}
 
+		/**
+		* Convert FKLevelSetElem to FLevelSetShape3d
+		*/
+		inline void GetShape(const FKLevelSetElem& LevelSetElem, FLevelSetShape3d& LevelSetShapeOut)
+		{
+			FTransform InGridTransform;
+			TArray<double> InGridValues;
+			FIntVector GridDims;
+			float InGridCellSize;
+
+			LevelSetElem.GetLevelSetData(InGridTransform, InGridValues, GridDims, InGridCellSize);
+
+			if (!ensure(InGridValues.Num() == GridDims[0] * GridDims[1] * GridDims[2]))
+			{
+				return;
+			}
+
+			LevelSetShapeOut.GridTransform = InGridTransform;
+			LevelSetShapeOut.GridTransform.AddToTranslation(0.5 * InGridCellSize * FVector::One());
+			LevelSetShapeOut.CellSize = InGridCellSize;
+			LevelSetShapeOut.Grid.Resize(GridDims[0], GridDims[1], GridDims[2]);
+
+			for (int I = 0; I < GridDims[0]; ++I)
+			{
+				for (int J = 0; J < GridDims[1]; ++J)
+				{
+					for (int K = 0; K < GridDims[2]; ++K)
+					{
+						const int InBufferIndex = K + GridDims[2] * (J + GridDims[1] * I);
+						const int OutBufferIndex = I + GridDims[0] * (J + GridDims[1] * K);
+						LevelSetShapeOut.Grid[OutBufferIndex] = InGridValues[InBufferIndex];
+					}
+				}
+			}
+		}
 
 		/**
 		 * Convert FKAggregateGeom to FSimpleShapeSet3d
@@ -163,6 +225,14 @@ namespace UE
 				GetShape(ConvexElem, ConvexShape);
 				ShapeSetOut.Convexes.Add(ConvexShape);
 			}
+
+			for (const FKLevelSetElem& LevelSetElem : AggGeom.LevelSetElems)
+			{
+				FLevelSetShape3d LevelSetShape;
+				GetShape(LevelSetElem, LevelSetShape);
+				ShapeSetOut.LevelSets.Add(LevelSetShape);
+			}
+
 		}
 
 
diff --git a/Engine/Source/Runtime/Engine/Classes/PhysicsEngine/AggregateGeom.h b/Engine/Source/Runtime/Engine/Classes/PhysicsEngine/AggregateGeom.h
--- a/Engine/Source/Runtime/Engine/Classes/PhysicsEngine/AggregateGeom.h
+++ b/Engine/Source/Runtime/Engine/Classes/PhysicsEngine/AggregateGeom.h
@@ -174,5 +174,6 @@ private:
 		SphylElems = Other.SphylElems;
 		ConvexElems = Other.ConvexElems;
 		TaperedCapsuleElems = Other.TaperedCapsuleElems;
+		LevelSetElems = Other.LevelSetElems;
 	}
 };
diff --git a/Engine/Source/Runtime/Engine/Classes/PhysicsEngine/LevelSetElem.h b/Engine/Source/Runtime/Engine/Classes/PhysicsEngine/LevelSetElem.h
--- a/Engine/Source/Runtime/Engine/Classes/PhysicsEngine/LevelSetElem.h
+++ b/Engine/Source/Runtime/Engine/Classes/PhysicsEngine/LevelSetElem.h
@@ -32,7 +32,9 @@ struct FKLevelSetElem : public FKShapeElem
 		return *this;
 	}
 
-	ENGINE_API void BuildLevelSet(const FVector& GridOrigin, const TArray<double>& GridValues, const FIntVector& GridDims, float GridCellSize);
+	ENGINE_API void BuildLevelSet(const FTransform& GridTransform, const TArray<double>& GridValues, const FIntVector& GridDims, float GridCellSize);
+
+	ENGINE_API void GetLevelSetData(FTransform& OutGridTransform, TArray<double>& OutGridValues, FIntVector& OutGridDims, float& OutGridCellSize) const;
 
 	ENGINE_API FTransform GetTransform() const
 	{
@@ -49,8 +51,8 @@ struct FKLevelSetElem : public FKShapeElem
 
 	// Draw helpers
 
-	/** Get geometry of all cells where the level set function is less than or equal to zero */
-	ENGINE_API void GetInteriorGridCells(TArray<FBox>& CellBoxes) const;
+	/** Get geometry of all cells where the level set function is less than or equal to InteriorThreshold */
+	ENGINE_API void GetInteriorGridCells(TArray<FBox>& CellBoxes, double InteriorThreshold = 0.0) const;
 
 	/** Get geometry of all cell faces where level set function changes sign */
 	ENGINE_API void GetZeroIsosurfaceGridCellFaces(TArray<FVector3f>& Vertices, TArray<FIntVector>& Tris) const;
@@ -61,6 +63,7 @@ struct FKLevelSetElem : public FKShapeElem
 	ENGINE_API void GetElemSolid(const FTransform& ElemTM, const FVector& Scale3D, const FMaterialRenderProxy* MaterialRenderProxy, int32 ViewIndex, class FMeshElementCollector& Collector) const;
 
 	ENGINE_API FBox CalcAABB(const FTransform& BoneTM, const FVector& Scale3D) const;
+	ENGINE_API FBox UntransformedAABB() const;
 
 	bool Serialize(FArchive& Ar);
 
@@ -78,7 +81,7 @@ private:
 	TSharedPtr<Chaos::FLevelSet, ESPMode::ThreadSafe> LevelSet;
 
 	/** Helper function to safely copy instances of this shape*/
-	void CloneElem(const FKLevelSetElem& Other);
+	ENGINE_API void CloneElem(const FKLevelSetElem& Other);
 };
 
 /* Enable our own serialization function to handle FLevelSet */
diff --git a/Engine/Source/Runtime/Engine/Private/Collision/KAggregateGeom.cpp b/Engine/Source/Runtime/Engine/Private/Collision/KAggregateGeom.cpp
--- a/Engine/Source/Runtime/Engine/Private/Collision/KAggregateGeom.cpp
+++ b/Engine/Source/Runtime/Engine/Private/Collision/KAggregateGeom.cpp
@@ -524,19 +524,55 @@ FBox FKLevelSetElem::CalcAABB(const FTransform& BoneTM, const FVector& Scale3D)
 	return Box;
 }
 
-void FKLevelSetElem::BuildLevelSet(const FVector& GridOrigin, const TArray<double>& GridValues, const FIntVector& GridDims, float GridCellSize)
+FBox FKLevelSetElem::UntransformedAABB() const
 {
-	const Chaos::FVec3 Min = GridOrigin;
-	const Chaos::FVec3 Max = GridOrigin + Chaos::FVec3(GridCellSize * (FVector)GridDims);
+	FBox Box;
+
+	if (LevelSet.IsValid())
+	{
+		Box = FBox(LevelSet->BoundingBox().Min(), LevelSet->BoundingBox().Max());
+		return Box;
+	}
+
+	return Box;
+}
+
+void FKLevelSetElem::BuildLevelSet(const FTransform& GridTransform, const TArray<double>& GridValues, const FIntVector& GridDims, float GridCellSize)
+{
+	const Chaos::FVec3 Min(0, 0, 0);
+	const Chaos::FVec3 Max = Min + Chaos::FVec3(GridCellSize * (FVector)GridDims);
 	const Chaos::TVec3<int32> ChaosDim(GridDims[0], GridDims[1], GridDims[2]);
 	
 	Chaos::TUniformGrid<Chaos::FReal, 3> ChaosGrid(Min, Max, ChaosDim);
 	Chaos::TArrayND<Chaos::FReal, 3> Phi( Chaos::TVec3<int32>(GridDims[0], GridDims[1], GridDims[2]), GridValues );
 
 	LevelSet = TSharedPtr<Chaos::FLevelSet>(new Chaos::FLevelSet(MoveTemp(ChaosGrid), MoveTemp(Phi), 0));
+
+	Transform = GridTransform;
 }
 
-void FKLevelSetElem::GetInteriorGridCells( TArray<FBox>& CellBoxes ) const
+void FKLevelSetElem::GetLevelSetData(FTransform& OutGridTransform, TArray<double>& OutGridValues, FIntVector& OutGridDims, float& OutGridCellSize) const
+{
+	if (!ensure(LevelSet.IsValid()))
+	{
+		return;
+	}
+	const Chaos::TUniformGrid<Chaos::FReal, 3>& ChaosGrid = LevelSet->GetGrid();
+
+	OutGridTransform = Transform;
+	OutGridDims = FIntVector(ChaosGrid.Counts()[1], ChaosGrid.Counts()[0], ChaosGrid.Counts()[2]);
+
+	OutGridCellSize = ChaosGrid.Dx()[0];
+
+	OutGridValues.Init(0.0, LevelSet->GetPhiArray().Num());
+	for (int32 Index = 0; Index < ChaosGrid.GetNumCells(); ++Index)
+	{
+		OutGridValues[Index] = LevelSet->GetPhiArray()[Index];
+	}
+}
+
+
+void FKLevelSetElem::GetInteriorGridCells( TArray<FBox>& CellBoxes, double InteriorThreshold) const
 {
 	if (LevelSet.IsValid())
 	{
@@ -550,7 +586,7 @@ void FKLevelSetElem::GetInteriorGridCells( TArray<FBox>& CellBoxes ) const
 				{
 					const double Value = LevelSet->GetPhiArray()(i, j, k);
 
-					if (Value <= 0.0)
+					if (Value <= InteriorThreshold)
 					{
 						const FVector3d CellMin = Grid.MinCorner() + Grid.Dx() * FVector3d(i, j, k);
 						const FVector3d CellMax = CellMin + Grid.Dx() * FVector3d(1, 1, 1);
diff --git a/Engine/Source/Runtime/Engine/Private/PhysicsEngine/PhysDrawing.cpp b/Engine/Source/Runtime/Engine/Private/PhysicsEngine/PhysDrawing.cpp
--- a/Engine/Source/Runtime/Engine/Private/PhysicsEngine/PhysDrawing.cpp
+++ b/Engine/Source/Runtime/Engine/Private/PhysicsEngine/PhysDrawing.cpp
@@ -928,10 +928,11 @@ void FKLevelSetElem::DrawElemWire(class FPrimitiveDrawInterface* PDI, const FTra
 	TArray<FBox> Boxes;
 
 	// Bounding box
-	Boxes.Add(CalcAABB(FTransform::Identity, FVector::OneVector));
+	Boxes.Add(UntransformedAABB());
 
 	// Cells with negative Phi values
-	GetInteriorGridCells(Boxes);
+	const double Threshold = UE_KINDA_SMALL_NUMBER;		// allow slightly greater than zero for visualization purposes
+	GetInteriorGridCells(Boxes, UE_KINDA_SMALL_NUMBER);
 
 	for (const FBox& Box : Boxes)
 	{
diff --git a/Engine/Source/Runtime/GeometryCore/Public/Implicit/GridInterpolant.h b/Engine/Source/Runtime/GeometryCore/Public/Implicit/GridInterpolant.h
--- a/Engine/Source/Runtime/GeometryCore/Public/Implicit/GridInterpolant.h
+++ b/Engine/Source/Runtime/GeometryCore/Public/Implicit/GridInterpolant.h
@@ -6,6 +6,7 @@
 
 #include "MathUtil.h"
 #include "VectorTypes.h"
+#include "IntVectorTypes.h"
 
 namespace UE
 {
@@ -47,6 +48,16 @@ public:
 			  GridOrigin.Z + CellSize * Dimensions.Z });
 	}
 
+	FVector3i Cell(const FVector3d& Pt) const
+	{
+		// compute integer coordinates
+		FVector3i CellCoords;
+		CellCoords.X = (int)((Pt.X - GridOrigin.X) / CellSize);
+		CellCoords.Y = (int)((Pt.Y - GridOrigin.Y) / CellSize);
+		CellCoords.Z = (int)((Pt.Z - GridOrigin.Z) / CellSize);
+
+		return CellCoords;
+	}
 
 	double Value(const FVector3d& Pt) const
 	{
diff --git a/Engine/Source/Runtime/GeometryCore/Public/Intersection/ContainmentQueries3.h b/Engine/Source/Runtime/GeometryCore/Public/Intersection/ContainmentQueries3.h
--- a/Engine/Source/Runtime/GeometryCore/Public/Intersection/ContainmentQueries3.h
+++ b/Engine/Source/Runtime/GeometryCore/Public/Intersection/ContainmentQueries3.h
@@ -8,6 +8,7 @@
 #include "CapsuleTypes.h"
 #include "OrientedBoxTypes.h"
 #include "HalfspaceTypes.h"
+#include "Implicit/GridInterpolant.h"
 #include "Intersection/IntersectionQueries3.h"
 
 
@@ -45,6 +46,13 @@ namespace UE
 			return true;
 		}
 
+		/** @return true if the AABB of InnerGrid is completely inside OuterSphere */
+		template<typename RealType, typename GridType>
+		bool IsInside(const TSphere3<RealType>& OuterSphere, const TTriLinearGridInterpolant<GridType>& InnerGrid)
+		{
+			return IsInside(OuterSphere, TOrientedBox3<RealType>(InnerGrid.Bounds()));
+		}
+
 
 
 		//
@@ -77,6 +85,13 @@ namespace UE
 			return true;
 		}
 
+		/** @return true if the AABB of InnerGrid is completely inside OuterCapsule */
+		template<typename RealType, typename GridType>
+		bool IsInside(const TCapsule3<RealType>& OuterCapsule, const TTriLinearGridInterpolant<GridType>& InnerGrid)
+		{
+			return IsInside(OuterCapsule, TOrientedBox3<RealType>(InnerGrid.Bounds()));
+		}
+
 
 
 		//
@@ -109,6 +124,12 @@ namespace UE
 			return true;
 		}
 
+		/** @return true if the AABB of InnerGrid is completely inside OuterBox */
+		template<typename RealType, typename GridType>
+		bool IsInside(const TOrientedBox3<RealType>& OuterBox, const TTriLinearGridInterpolant<GridType>& InnerGrid)
+		{
+			return IsInside(OuterBox, TOrientedBox3<RealType>(InnerGrid.Bounds()));
+		}
 
 
 		//
@@ -147,6 +168,49 @@ namespace UE
 		template<typename RealType, typename EnumerablePointsType, typename E = decltype(DeclVal<EnumerablePointsType>().begin())>
 		bool IsInsideHull(TArrayView<THalfspace3<RealType>> Halfspaces, EnumerablePointsType EnumerablePts);
 
+		template<typename RealType, typename GridType>
+		bool IsInsideHull(TArrayView<THalfspace3<RealType>> Halfspaces, const TTriLinearGridInterpolant<GridType>& InnerGrid)
+		{
+			return IsInsideHull(Halfspaces, TOrientedBox3<RealType>(InnerGrid.Bounds()));
+		}
+
+
+		//
+		// Signed Distance Field containment queries
+		//
+
+		/**
+		 * Test if the axis-aligned bounding box of InnerSphere is completely inside the negative region of the signed distance field discretized on OuterGrid
+		 */
+		template<typename RealType, typename GridType>
+		bool IsInside(const TTriLinearGridInterpolant<GridType>& OuterGrid, const TSphere3<RealType>& InnerSphere);
+
+		/**
+		 * Test if the axis-aligned bounding box of InnerCapsule is completely inside the negative region of the signed distance field discretized on OuterGrid
+		 */
+		template<typename RealType, typename GridType>
+		bool IsInside(const TTriLinearGridInterpolant<GridType>& OuterGrid, const TCapsule3<RealType>& InnerCapsule);
+
+		/**
+		 * Test if the axis-aligned bounding box of InnerBox is completely inside the negative region of the signed distance field discretized on OuterGrid
+		 */
+		template<typename RealType, typename GridType>
+		bool IsInside(const TTriLinearGridInterpolant<GridType>& OuterGrid, const TOrientedBox3<RealType>& InnerBox);
+
+		/**
+		 * Test if all enumerable points are inside the negative region of the signed distance field discretized on OuterGrid
+		 */
+		template<typename RealType, typename GridType, typename EnumerablePointsType, typename E = decltype(DeclVal<EnumerablePointsType>().begin())>
+		bool IsInside(const TTriLinearGridInterpolant<GridType>& OuterGrid, EnumerablePointsType EnumerablePts);
+
+		/**
+		 * Test if the axis-aligned bounding box of InnerGrid is completely inside the negative region of the signed distance field discretized on OuterGrid
+		 */
+		template<typename RealType, typename GridType1, typename GridType2>
+		bool IsInside(const TTriLinearGridInterpolant<GridType1>& OuterGrid, const TTriLinearGridInterpolant<GridType2>& InnerGrid)
+		{
+			return IsInside(OuterGrid, TOrientedBox3<RealType>(InnerGrid.Bounds()));
+		}
 
 	}
 }
@@ -209,4 +273,96 @@ bool UE::Geometry::IsInsideHull(TArrayView<THalfspace3<RealType>> Halfspaces, En
 		}
 	}
 	return true;
-}
\ No newline at end of file
+}
+
+
+
+namespace UE
+{
+	namespace Geometry
+	{
+		//
+		// Signed distance field containment tests
+		// 
+		// To test if an object is inside an SDF, we first rasterize the object's AABB onto the SDF grid. If all of the AABB cells have negative SDF values,
+		// we report that the object is inside. Note this is conservative, not exact -- so an object may actually be inside the implicit surface defined by
+		// the SDF, but if its AABB is not fully inside then we will report "not inside."
+		//
+
+		template<typename RealType, typename GridType>
+		bool IsInside(const TTriLinearGridInterpolant<GridType>& OuterGrid, const UE::Geometry::TAxisAlignedBox3<RealType>& InnerAABB)
+		{
+			const FVector3i MinCell = OuterGrid.Cell(InnerAABB.Min);
+			FVector3i MaxCell = OuterGrid.Cell(InnerAABB.Max);
+
+			for (int32 Dim = 0; Dim < 3; ++Dim)
+			{
+				if (MinCell[Dim] < 0 || MaxCell[Dim] >= OuterGrid.Dimensions[Dim])
+				{
+					return false;	// AABB extends outside the grid
+				}
+			}
+
+			MaxCell += FVector3i(1, 1, 1);
+
+			for (int I = MinCell[0]; I < MaxCell[0]; ++I)
+			{
+				for (int J = MinCell[1]; J < MaxCell[1]; ++J)
+				{
+					for (int K = MinCell[2]; K < MaxCell[2]; ++K)
+					{
+						RealType GridCellValue = OuterGrid.Grid->GetValue(I, J, K);
+
+						if (GridCellValue >= 0)	// one of the cells containing the AABB is outside the implicit surface
+						{
+							return false;
+						}
+					}
+				}
+			}
+
+			return true;
+		}
+
+		template<typename RealType, typename GridType>
+		bool IsInside(const TTriLinearGridInterpolant<GridType>& OuterGrid, const TSphere3<RealType>& InnerSphere)
+		{
+			TAxisAlignedBox3<RealType> InnerAABB(InnerSphere.Center - InnerSphere.Radius * TVector<RealType>::OneVector,
+				InnerSphere.Center + InnerSphere.Radius * TVector<RealType>::OneVector);
+
+			// TODO: TSphere3 has a SignedDistance function, so we could get a better discretization
+			return IsInside(OuterGrid, InnerAABB);
+		}
+
+		template<typename RealType, typename GridType>
+		bool IsInside(const TTriLinearGridInterpolant<GridType>& OuterGrid, const TCapsule3<RealType>& InnerCapsule)
+		{
+			// TODO: TCapsule3 has a SignedDistance function, so we could get a better discretization
+			return IsInside(OuterGrid, InnerCapsule.GetBounds());
+		}
+
+		template<typename RealType, typename GridType>
+		bool IsInside(const TTriLinearGridInterpolant<GridType>& OuterGrid, const TOrientedBox3<RealType>& InnerBox)
+		{
+			TAxisAlignedBox3<RealType> InnerAABB;
+			InnerBox.EnumerateCorners([&InnerAABB](const FVector3d& CornerPt)
+			{
+				InnerAABB.Contain(CornerPt);
+			});
+
+			return IsInside(OuterGrid, InnerAABB);
+		}
+
+		template<typename RealType, typename GridType, typename EnumerablePointsType, typename E>
+		bool IsInside(const TTriLinearGridInterpolant<GridType>& OuterGrid, EnumerablePointsType EnumerablePts)
+		{
+			TAxisAlignedBox3<RealType> InnerAABB;
+			for (TVector<RealType> Point : EnumerablePts)
+			{
+				InnerAABB.Contain(Point);
+			}
+			return IsInside(OuterGrid, InnerAABB);
+		}
+
+	}
+}
