Render Capture Baking: Add a method exposing the evaluator data layout

This will be useful for the (upcoming) RenderCaptureMapEvaluator as it allows a pointer to the FMeshMapEvaluator base class to be statically cast to the right type even though that Evaluator must be a template class. Without this method we can't cast the base class pointer since we don't know the template parameter, but we want the template so that EvaluateSample can be branch free which is desirable since its called in hot loops

#rnx

diff --git a/Engine/Plugins/Runtime/GeometryProcessing/Source/DynamicMesh/Private/Sampling/MeshConstantMapEvaluator.cpp b/Engine/Plugins/Runtime/GeometryProcessing/Source/DynamicMesh/Private/Sampling/MeshConstantMapEvaluator.cpp
--- a/Engine/Plugins/Runtime/GeometryProcessing/Source/DynamicMesh/Private/Sampling/MeshConstantMapEvaluator.cpp
+++ b/Engine/Plugins/Runtime/GeometryProcessing/Source/DynamicMesh/Private/Sampling/MeshConstantMapEvaluator.cpp
@@ -18,7 +18,13 @@ void FMeshConstantMapEvaluator::Setup(const FMeshBaseBaker& Baker, FEvaluationCo
 	Context.EvaluateColor = &EvaluateColor;
 	Context.EvalData = this;
 	Context.AccumulateMode = EAccumulateMode::Overwrite;
-	Context.DataLayout = { EComponents::Float1 };
+	Context.DataLayout = DataLayout();
+}
+
+const TArray<FMeshMapEvaluator::EComponents>& FMeshConstantMapEvaluator::DataLayout() const
+{
+	static const TArray<EComponents> Layout{ EComponents::Float1 };
+	return Layout;
 }
 
 void FMeshConstantMapEvaluator::EvaluateSample(float*& Out, const FCorrespondenceSample& Sample, void* EvalData)
diff --git a/Engine/Plugins/Runtime/GeometryProcessing/Source/DynamicMesh/Private/Sampling/MeshCurvatureMapEvaluator.cpp b/Engine/Plugins/Runtime/GeometryProcessing/Source/DynamicMesh/Private/Sampling/MeshCurvatureMapEvaluator.cpp
--- a/Engine/Plugins/Runtime/GeometryProcessing/Source/DynamicMesh/Private/Sampling/MeshCurvatureMapEvaluator.cpp
+++ b/Engine/Plugins/Runtime/GeometryProcessing/Source/DynamicMesh/Private/Sampling/MeshCurvatureMapEvaluator.cpp
@@ -27,7 +27,7 @@ void FMeshCurvatureMapEvaluator::Setup(const FMeshBaseBaker& Baker, FEvaluationC
 	Context.EvaluateColor = &EvaluateColor;
 	Context.EvalData = this;
 	Context.AccumulateMode = EAccumulateMode::Add;
-	Context.DataLayout = { EComponents::Float1 };
+	Context.DataLayout = DataLayout();
 
 	// Cache data from the baker
 	DetailSampler = Baker.GetDetailSampler();
@@ -74,6 +74,12 @@ void FMeshCurvatureMapEvaluator::Setup(const FMeshBaseBaker& Baker, FEvaluationC
 	GetColorMapRange(NegativeColor, ZeroColor, PositiveColor);
 }
 
+const TArray<FMeshMapEvaluator::EComponents>& FMeshCurvatureMapEvaluator::DataLayout() const
+{
+	static const TArray<EComponents> Layout{ EComponents::Float1 };
+	return Layout;
+}
+
 void FMeshCurvatureMapEvaluator::EvaluateSample(float*& Out, const FCorrespondenceSample& Sample, void* EvalData)
 {
 	const FMeshCurvatureMapEvaluator* Eval = static_cast<FMeshCurvatureMapEvaluator*>(EvalData);
diff --git a/Engine/Plugins/Runtime/GeometryProcessing/Source/DynamicMesh/Private/Sampling/MeshNormalMapEvaluator.cpp b/Engine/Plugins/Runtime/GeometryProcessing/Source/DynamicMesh/Private/Sampling/MeshNormalMapEvaluator.cpp
--- a/Engine/Plugins/Runtime/GeometryProcessing/Source/DynamicMesh/Private/Sampling/MeshNormalMapEvaluator.cpp
+++ b/Engine/Plugins/Runtime/GeometryProcessing/Source/DynamicMesh/Private/Sampling/MeshNormalMapEvaluator.cpp
@@ -36,7 +36,13 @@ void FMeshNormalMapEvaluator::Setup(const FMeshBaseBaker& Baker, FEvaluationCont
 	Context.EvaluateColor = &EvaluateColor;
 	Context.EvalData = this;
 	Context.AccumulateMode = EAccumulateMode::Add;
-	Context.DataLayout = { EComponents::Float3 };
+	Context.DataLayout = DataLayout();
+}
+
+const TArray<FMeshMapEvaluator::EComponents>& FMeshNormalMapEvaluator::DataLayout() const
+{
+	static const TArray<EComponents> Layout{ EComponents::Float3 };
+	return Layout;
 }
 
 template <bool bUseDetailNormalMap>
diff --git a/Engine/Plugins/Runtime/GeometryProcessing/Source/DynamicMesh/Private/Sampling/MeshOcclusionMapEvaluator.cpp b/Engine/Plugins/Runtime/GeometryProcessing/Source/DynamicMesh/Private/Sampling/MeshOcclusionMapEvaluator.cpp
--- a/Engine/Plugins/Runtime/GeometryProcessing/Source/DynamicMesh/Private/Sampling/MeshOcclusionMapEvaluator.cpp
+++ b/Engine/Plugins/Runtime/GeometryProcessing/Source/DynamicMesh/Private/Sampling/MeshOcclusionMapEvaluator.cpp
@@ -23,6 +23,8 @@ const FVector3d FMeshOcclusionMapEvaluator::DefaultObjectNormal = FVector3d::Zer
 
 void FMeshOcclusionMapEvaluator::Setup(const FMeshBaseBaker& Baker, FEvaluationContext& Context)
 {
+	Context.DataLayout = DataLayout();
+
 	if (WANT_ALL(OcclusionType))
 	{
 		if (NormalSpace == ESpace::Tangent)
@@ -39,7 +41,6 @@ void FMeshOcclusionMapEvaluator::Setup(const FMeshBaseBaker& Baker, FEvaluationC
 		}
 		Context.EvalData = this;
 		Context.AccumulateMode = EAccumulateMode::Add;
-		Context.DataLayout = { EComponents::Float1, EComponents::Float3 };
 	}
 	else if (WANT_AMBIENT_OCCLUSION(OcclusionType))
 	{
@@ -48,7 +49,6 @@ void FMeshOcclusionMapEvaluator::Setup(const FMeshBaseBaker& Baker, FEvaluationC
 		Context.EvaluateColor = &EvaluateColor<EMeshOcclusionMapType::AmbientOcclusion, ESpace::Tangent>;
 		Context.EvalData = this;
 		Context.AccumulateMode = EAccumulateMode::Add;
-		Context.DataLayout = { EComponents::Float1 };
 	}
 	else if (WANT_BENT_NORMAL(OcclusionType))
 	{
@@ -66,7 +66,6 @@ void FMeshOcclusionMapEvaluator::Setup(const FMeshBaseBaker& Baker, FEvaluationC
 		}
 		Context.EvalData = this;
 		Context.AccumulateMode = EAccumulateMode::Add;
-		Context.DataLayout = { EComponents::Float3 };
 	}
 	else
 	{
@@ -119,6 +118,33 @@ void FMeshOcclusionMapEvaluator::Setup(const FMeshBaseBaker& Baker, FEvaluationC
 	}
 }
 
+const TArray<FMeshMapEvaluator::EComponents>& FMeshOcclusionMapEvaluator::DataLayout() const
+{
+	if (WANT_ALL(OcclusionType))
+	{
+		static const TArray<EComponents> Layout{ EComponents::Float1, EComponents::Float3 };
+		return Layout;
+	}
+
+	if (WANT_AMBIENT_OCCLUSION(OcclusionType))
+	{
+		static const TArray<EComponents> Layout{ EComponents::Float1 };
+		return Layout;
+	}
+
+	if (WANT_BENT_NORMAL(OcclusionType))
+	{
+		static const TArray<EComponents> Layout{ EComponents::Float3 };
+		return Layout;
+	}
+
+	// TODO: Support error case in the baker to skip over invalid eval configs?
+	checkSlow(false);
+
+	static const TArray<EComponents> Layout{ EComponents::Float1 };
+	return Layout;
+}
+
 template <EMeshOcclusionMapType ComputeType, FMeshOcclusionMapEvaluator::ESpace ComputeSpace>
 void FMeshOcclusionMapEvaluator::EvaluateSample(float*& Out, const FCorrespondenceSample& Sample, void* EvalData)
 {
diff --git a/Engine/Plugins/Runtime/GeometryProcessing/Source/DynamicMesh/Private/Sampling/MeshPropertyMapEvaluator.cpp b/Engine/Plugins/Runtime/GeometryProcessing/Source/DynamicMesh/Private/Sampling/MeshPropertyMapEvaluator.cpp
--- a/Engine/Plugins/Runtime/GeometryProcessing/Source/DynamicMesh/Private/Sampling/MeshPropertyMapEvaluator.cpp
+++ b/Engine/Plugins/Runtime/GeometryProcessing/Source/DynamicMesh/Private/Sampling/MeshPropertyMapEvaluator.cpp
@@ -34,7 +34,7 @@ void FMeshPropertyMapEvaluator::Setup(const FMeshBaseBaker& Baker, FEvaluationCo
 	Context.EvaluateColor = &EvaluateColor;
 	Context.EvalData = this;
 	Context.AccumulateMode = EAccumulateMode::Add;
-	Context.DataLayout = { EComponents::Float3 };
+	Context.DataLayout = DataLayout();
 
 	Bounds = DetailSampler->GetBounds();
 	for (int32 j = 0; j < 3; ++j)
@@ -71,6 +71,12 @@ void FMeshPropertyMapEvaluator::Setup(const FMeshBaseBaker& Baker, FEvaluationCo
 	}
 }
 
+const TArray<FMeshMapEvaluator::EComponents>& FMeshPropertyMapEvaluator::DataLayout() const
+{
+	static const TArray<EComponents> Layout{ EComponents::Float3 };
+	return Layout;
+}
+
 template <bool bUseDetailNormalMap>
 void FMeshPropertyMapEvaluator::EvaluateSample(float*& Out, const FCorrespondenceSample& Sample, void* EvalData)
 {
diff --git a/Engine/Plugins/Runtime/GeometryProcessing/Source/DynamicMesh/Private/Sampling/MeshResampleImageEvaluator.cpp b/Engine/Plugins/Runtime/GeometryProcessing/Source/DynamicMesh/Private/Sampling/MeshResampleImageEvaluator.cpp
--- a/Engine/Plugins/Runtime/GeometryProcessing/Source/DynamicMesh/Private/Sampling/MeshResampleImageEvaluator.cpp
+++ b/Engine/Plugins/Runtime/GeometryProcessing/Source/DynamicMesh/Private/Sampling/MeshResampleImageEvaluator.cpp
@@ -16,12 +16,18 @@ void FMeshResampleImageEvaluator::Setup(const FMeshBaseBaker& Baker, FEvaluation
 	Context.EvaluateColor = &EvaluateColor;
 	Context.EvalData = this;
 	Context.AccumulateMode = EAccumulateMode::Add;
-	Context.DataLayout = { EComponents::Float4 };
+	Context.DataLayout = DataLayout();
 
 	// Cache data from the baker
 	DetailSampler = Baker.GetDetailSampler();
 }
 
+const TArray<FMeshMapEvaluator::EComponents>& FMeshResampleImageEvaluator::DataLayout() const
+{
+	static const TArray<FMeshMapEvaluator::EComponents> Layout{ EComponents::Float4 };
+	return Layout;
+}
+
 void FMeshResampleImageEvaluator::EvaluateSample(float*& Out, const FCorrespondenceSample& Sample, void* EvalData)
 {
 	const FMeshResampleImageEvaluator* Eval = static_cast<FMeshResampleImageEvaluator*>(EvalData);
@@ -73,7 +79,7 @@ void FMeshMultiResampleImageEvaluator::Setup(const FMeshBaseBaker& Baker, FEvalu
 	Context.EvaluateColor = &EvaluateColor;
 	Context.EvalData = this;
 	Context.AccumulateMode = EAccumulateMode::Add;
-	Context.DataLayout = { EComponents::Float4 };
+	Context.DataLayout = DataLayout();
 
 	// Cache data from baker
 	DetailSampler = Baker.GetDetailSampler();
diff --git a/Engine/Plugins/Runtime/GeometryProcessing/Source/DynamicMesh/Public/Sampling/GenericEvaluator.h b/Engine/Plugins/Runtime/GeometryProcessing/Source/DynamicMesh/Public/Sampling/GenericEvaluator.h
--- a/Engine/Plugins/Runtime/GeometryProcessing/Source/DynamicMesh/Public/Sampling/GenericEvaluator.h
+++ b/Engine/Plugins/Runtime/GeometryProcessing/Source/DynamicMesh/Public/Sampling/GenericEvaluator.h
@@ -32,7 +32,13 @@ public:
 		Context.EvaluateColor = &EvaluateColor;
 		Context.EvalData = this;
 		Context.AccumulateMode = Mode;
-		Context.DataLayout = { NumComponents };
+		Context.DataLayout = DataLayout();
+	}
+
+	virtual const TArray<EComponents>& DataLayout() const override
+	{
+		static const TArray<EComponents> Layout{ NumComponents };
+		return Layout;
 	}
 
 	virtual EMeshMapEvaluatorType Type() const override { return EMeshMapEvaluatorType::Generic; }
diff --git a/Engine/Plugins/Runtime/GeometryProcessing/Source/DynamicMesh/Public/Sampling/MeshConstantMapEvaluator.h b/Engine/Plugins/Runtime/GeometryProcessing/Source/DynamicMesh/Public/Sampling/MeshConstantMapEvaluator.h
--- a/Engine/Plugins/Runtime/GeometryProcessing/Source/DynamicMesh/Public/Sampling/MeshConstantMapEvaluator.h
+++ b/Engine/Plugins/Runtime/GeometryProcessing/Source/DynamicMesh/Public/Sampling/MeshConstantMapEvaluator.h
@@ -23,6 +23,8 @@ public:
 	// Begin FMeshMapEvaluator interface
 	virtual void Setup(const FMeshBaseBaker& Baker, FEvaluationContext& Context) override;
 
+	virtual const TArray<EComponents>& DataLayout() const override;
+
 	virtual EMeshMapEvaluatorType Type() const override { return EMeshMapEvaluatorType::Constant; }
 	// End FMeshMapEvaluator interface
 
diff --git a/Engine/Plugins/Runtime/GeometryProcessing/Source/DynamicMesh/Public/Sampling/MeshCurvatureMapEvaluator.h b/Engine/Plugins/Runtime/GeometryProcessing/Source/DynamicMesh/Public/Sampling/MeshCurvatureMapEvaluator.h
--- a/Engine/Plugins/Runtime/GeometryProcessing/Source/DynamicMesh/Public/Sampling/MeshCurvatureMapEvaluator.h
+++ b/Engine/Plugins/Runtime/GeometryProcessing/Source/DynamicMesh/Public/Sampling/MeshCurvatureMapEvaluator.h
@@ -57,6 +57,8 @@ public:
 	// Begin FMeshMapEvaluator interface
 	virtual void Setup(const FMeshBaseBaker& Baker, FEvaluationContext& Context) override;
 
+	virtual const TArray<EComponents>& DataLayout() const override;
+
 	virtual EMeshMapEvaluatorType Type() const override { return EMeshMapEvaluatorType::Curvature; }
 	// End FMeshMapEvaluator interface
 
diff --git a/Engine/Plugins/Runtime/GeometryProcessing/Source/DynamicMesh/Public/Sampling/MeshMapEvaluator.h b/Engine/Plugins/Runtime/GeometryProcessing/Source/DynamicMesh/Public/Sampling/MeshMapEvaluator.h
--- a/Engine/Plugins/Runtime/GeometryProcessing/Source/DynamicMesh/Public/Sampling/MeshMapEvaluator.h
+++ b/Engine/Plugins/Runtime/GeometryProcessing/Source/DynamicMesh/Public/Sampling/MeshMapEvaluator.h
@@ -146,6 +146,9 @@ public:
 	 */
 	virtual void Setup(const FMeshBaseBaker& Baker, FEvaluationContext& Context) = 0;
 
+	/** @return the data layout of the evaluator */
+	virtual const TArray<EComponents>& DataLayout() const = 0;
+
 	/** @return the type of evaluator. */
 	virtual EMeshMapEvaluatorType Type() const = 0;
 
diff --git a/Engine/Plugins/Runtime/GeometryProcessing/Source/DynamicMesh/Public/Sampling/MeshNormalMapEvaluator.h b/Engine/Plugins/Runtime/GeometryProcessing/Source/DynamicMesh/Public/Sampling/MeshNormalMapEvaluator.h
--- a/Engine/Plugins/Runtime/GeometryProcessing/Source/DynamicMesh/Public/Sampling/MeshNormalMapEvaluator.h
+++ b/Engine/Plugins/Runtime/GeometryProcessing/Source/DynamicMesh/Public/Sampling/MeshNormalMapEvaluator.h
@@ -28,6 +28,8 @@ public:
 	// Begin FMeshMapEvaluator interface
 	virtual void Setup(const FMeshBaseBaker& Baker, FEvaluationContext& Context) override;
 
+	virtual const TArray<EComponents>& DataLayout() const override;
+
 	virtual EMeshMapEvaluatorType Type() const override { return EMeshMapEvaluatorType::Normal; }
 	// End FMeshMapEvaluator interface
 
diff --git a/Engine/Plugins/Runtime/GeometryProcessing/Source/DynamicMesh/Public/Sampling/MeshOcclusionMapEvaluator.h b/Engine/Plugins/Runtime/GeometryProcessing/Source/DynamicMesh/Public/Sampling/MeshOcclusionMapEvaluator.h
--- a/Engine/Plugins/Runtime/GeometryProcessing/Source/DynamicMesh/Public/Sampling/MeshOcclusionMapEvaluator.h
+++ b/Engine/Plugins/Runtime/GeometryProcessing/Source/DynamicMesh/Public/Sampling/MeshOcclusionMapEvaluator.h
@@ -59,6 +59,8 @@ public:
 	// Begin FMeshMapEvaluator interface
 	virtual void Setup(const FMeshBaseBaker& Baker, FEvaluationContext& Context) override;
 
+	virtual const TArray<EComponents>& DataLayout() const override;
+
 	virtual EMeshMapEvaluatorType Type() const override { return EMeshMapEvaluatorType::Occlusion; }
 	// End FMeshMapEvaluator interface
 
diff --git a/Engine/Plugins/Runtime/GeometryProcessing/Source/DynamicMesh/Public/Sampling/MeshPropertyMapEvaluator.h b/Engine/Plugins/Runtime/GeometryProcessing/Source/DynamicMesh/Public/Sampling/MeshPropertyMapEvaluator.h
--- a/Engine/Plugins/Runtime/GeometryProcessing/Source/DynamicMesh/Public/Sampling/MeshPropertyMapEvaluator.h
+++ b/Engine/Plugins/Runtime/GeometryProcessing/Source/DynamicMesh/Public/Sampling/MeshPropertyMapEvaluator.h
@@ -41,6 +41,8 @@ public:
 	// Begin FMeshMapEvaluator interface
 	virtual void Setup(const FMeshBaseBaker& Baker, FEvaluationContext& Context) override;
 
+	virtual const TArray<EComponents>& DataLayout() const override;
+
 	virtual EMeshMapEvaluatorType Type() const override { return EMeshMapEvaluatorType::Property; }
 	// End FMeshMapEvaluator interface
 
diff --git a/Engine/Plugins/Runtime/GeometryProcessing/Source/DynamicMesh/Public/Sampling/MeshResampleImageEvaluator.h b/Engine/Plugins/Runtime/GeometryProcessing/Source/DynamicMesh/Public/Sampling/MeshResampleImageEvaluator.h
--- a/Engine/Plugins/Runtime/GeometryProcessing/Source/DynamicMesh/Public/Sampling/MeshResampleImageEvaluator.h
+++ b/Engine/Plugins/Runtime/GeometryProcessing/Source/DynamicMesh/Public/Sampling/MeshResampleImageEvaluator.h
@@ -25,6 +25,8 @@ public:
 	// Begin FMeshMapEvaluator interface
 	virtual void Setup(const FMeshBaseBaker& Baker, FEvaluationContext& Context) override;
 
+	virtual const TArray<EComponents>& DataLayout() const override;
+
 	virtual EMeshMapEvaluatorType Type() const override { return EMeshMapEvaluatorType::ResampleImage; }
 	// End FMeshMapEvaluator interface
 
