New MeshPaths class. Given a set of edges on a mesh, finds the topology of simply connected edges and collects them into Spans.

diff --git a/Engine/Plugins/Runtime/GeometryProcessing/Source/DynamicMesh/Private/MeshPaths.cpp b/Engine/Plugins/Runtime/GeometryProcessing/Source/DynamicMesh/Private/MeshPaths.cpp
new file mode 100644
--- /dev/null
+++ b/Engine/Plugins/Runtime/GeometryProcessing/Source/DynamicMesh/Private/MeshPaths.cpp
@@ -0,0 +1,100 @@
+// Copyright Epic Games, Inc. All Rights Reserved.
+
+
+#include "MeshPaths.h"
+
+using namespace UE::Geometry;
+
+int FMeshPaths::GetNextEdge(int CurrentEdge, int FromVertex, int& CrossedVertex) const
+{
+	check(!UnvisitedEdges.Contains(CurrentEdge));
+
+	FIndex2i EdgeV = Mesh->GetEdgeV(CurrentEdge);
+
+	CrossedVertex = (EdgeV[0] == FromVertex) ? EdgeV[1] : EdgeV[0];
+
+	int AdjacentUnvisitedEdge = IndexConstants::InvalidID;
+	int NumAdjacentEdgesForThisVertex = 0;
+
+	for (int EdgeID : Mesh->VtxEdgesItr(CrossedVertex))
+	{
+		if (EdgeID == CurrentEdge)
+		{
+			continue;
+		}
+
+		if (Edges.Contains(EdgeID))
+		{
+			++NumAdjacentEdgesForThisVertex;
+
+			if (UnvisitedEdges.Contains(EdgeID))
+			{
+				AdjacentUnvisitedEdge = EdgeID;
+			}
+		}
+	}
+
+	if (NumAdjacentEdgesForThisVertex == 1)	 // Only consider simply connected edges
+	{
+		return AdjacentUnvisitedEdge;
+	}
+
+	return IndexConstants::InvalidID;
+}
+
+
+
+bool FMeshPaths::Compute()
+{
+	if (!Mesh)
+	{
+		return false;
+	}
+
+	UnvisitedEdges = Edges;
+
+	while (!UnvisitedEdges.IsEmpty())
+	{
+		// get any unvisited edge, then find the span it's connected to
+		const int InitialEdge = *UnvisitedEdges.CreateConstIterator();
+		
+		const FIndex2i InitialEdgeVertices = Mesh->GetEdgeV(InitialEdge);
+
+		// Compute the span by traversing connected edges in both directions
+		TArray<int> LeftSpan, RightSpan;
+
+		for (int Direction = 0; Direction < 2; ++Direction)
+		{
+			int CurrentEdge = InitialEdge;
+			int FromVertex = InitialEdgeVertices[Direction];
+			int CrossedVertex = IndexConstants::InvalidID;
+
+			while (CurrentEdge != IndexConstants::InvalidID)
+			{
+				if (Direction == 0)
+				{
+					LeftSpan.Add(CurrentEdge);
+				}
+				else if (CurrentEdge != InitialEdge)
+				{
+					RightSpan.Add(CurrentEdge);
+				}
+				
+				UnvisitedEdges.Remove(CurrentEdge);
+				CurrentEdge = GetNextEdge(CurrentEdge, FromVertex, CrossedVertex);
+				FromVertex = CrossedVertex;
+			}
+		}
+
+		Algo::Reverse(LeftSpan);
+		LeftSpan.Append(RightSpan);
+
+		FEdgeSpan NewSpan;
+		NewSpan.InitializeFromEdges(Mesh, LeftSpan);
+		checkSlow(NewSpan.CheckValidity());
+		Spans.Emplace(NewSpan);
+	}
+
+	return true;
+}
+
diff --git a/Engine/Plugins/Runtime/GeometryProcessing/Source/DynamicMesh/Public/MeshPaths.h b/Engine/Plugins/Runtime/GeometryProcessing/Source/DynamicMesh/Public/MeshPaths.h
new file mode 100644
--- /dev/null
+++ b/Engine/Plugins/Runtime/GeometryProcessing/Source/DynamicMesh/Public/MeshPaths.h
@@ -0,0 +1,53 @@
+// Copyright Epic Games, Inc. All Rights Reserved.
+
+#pragma once
+
+#include "DynamicMesh/DynamicMesh3.h"
+#include "EdgeSpan.h"
+
+namespace UE
+{
+namespace Geometry
+{
+
+/** 
+  * Given a mesh and a set of edges, walk connected edges to form a set of paths.
+  * We do not distinguish betweens spans and loops -- if the input set of edges forms a loop, it will create an FEdgeSpan with those edges.
+  * If the input edges do not form loops, then paths will begin/end at vertices incident on 1 or >2 of the given edges (e.g. "dead ends" or "intersections" of paths.)
+  */
+class DYNAMICMESH_API FMeshPaths
+{
+public:
+
+	FMeshPaths(const FDynamicMesh3* Mesh, const TSet<int>& Edges ) :
+		Mesh(Mesh),
+		Edges(Edges)
+	{}
+
+	// Inputs
+	const FDynamicMesh3* Mesh = nullptr;
+	const TSet<int> Edges;
+
+	// Outputs
+	TArray<FEdgeSpan> Spans;
+
+	// Find connected simple paths made up of edges in the input Edges structure
+	bool Compute();
+
+private:
+
+	TSet<int> UnvisitedEdges;
+
+	// Helper for the Compute function
+
+	/*
+	 * Find an edge adjacent to CurrentEdge which is in both the Edges and UnvisitedEdges set.
+	 * Consider only edges incident on CrossedVertex, which is the vertex belonging to CurrentEdge that is not FromVertex.
+	 * Do not consider an adjacent edge if CrossedVertex is incident on more than 2 input Edges (i.e. we've hit a corner).
+	 */
+	int GetNextEdge(int CurrentEdge, int FromVertex, int& CrossedVertex) const;
+
+};
+
+} // end namespace UE::Geometry
+} // end namespace UE
