Initial implementation of the adaptive tessellation operator with one pattern implemented.
Currently only supports the generation of the new geometry without interpolating attributes.
Added GeometryScript function for adaptive tessellation.

diff --git a/Engine/Plugins/Experimental/GeometryScripting/Source/GeometryScriptingCore/Private/MeshSubdivideFunctions.cpp b/Engine/Plugins/Experimental/GeometryScripting/Source/GeometryScriptingCore/Private/MeshSubdivideFunctions.cpp
--- a/Engine/Plugins/Experimental/GeometryScripting/Source/GeometryScriptingCore/Private/MeshSubdivideFunctions.cpp
+++ b/Engine/Plugins/Experimental/GeometryScripting/Source/GeometryScriptingCore/Private/MeshSubdivideFunctions.cpp
@@ -7,6 +7,7 @@
 #include "DynamicMesh/MeshNormals.h"
 #include "Operations/PNTriangles.h"
 #include "Operations/UniformTessellate.h"
+#include "Operations/AdaptiveTessellate.h"
 #include "UDynamicMesh.h"
 #include "Math/UnrealMathUtility.h"
 
@@ -127,4 +128,109 @@ UDynamicMesh* UGeometryScriptLibrary_MeshSubdivideFunctions::ApplyUniformTessell
 	return TargetMesh;
 }
 
+
+UDynamicMesh* UGeometryScriptLibrary_MeshSubdivideFunctions::ApplyAdaptiveTessellation(
+	UDynamicMesh* TargetMesh,
+	FGeometryScriptAdaptiveTessellateOptions Options,
+	FGeometryScriptIndexList IndexList,
+	int TessellationLevel,
+	EAdaptiveTessellatePatternType PatternType,
+	UGeometryScriptDebug* Debug)
+{
+	if (TargetMesh == nullptr)
+	{
+		UE::Geometry::AppendError(Debug, EGeometryScriptErrorType::InvalidInputs, LOCTEXT("ApplyAdapativeTessellation_InvalidInput", "ApplyAdapativeTessellation: TargetMesh is Null"));
+		return TargetMesh;
+	}
+	if (TessellationLevel <= 0)
+	{
+		return TargetMesh;
+	}
+	if (PatternType != EAdaptiveTessellatePatternType::ConcentricRings)
+	{
+		UE::Geometry::AppendError(Debug, EGeometryScriptErrorType::InvalidInputs, LOCTEXT("ApplyAdapativeTessellation_Error", "Only ConcentricRings pattern is currently supported."));
+		return TargetMesh;
+	} 
+	if (IndexList.IndexType != EGeometryScriptIndexType::Triangle)
+	{
+		UE::Geometry::AppendError(Debug, EGeometryScriptErrorType::InvalidInputs, LOCTEXT("ApplyAdapativeTessellation_Error", "Only Triangle selection is currently supported."));
+		return TargetMesh;
+	} 
+
+	TargetMesh->EditMesh([&](FDynamicMesh3& EditMesh) 
+	{
+		FAdaptiveTessellate Tessellator(&EditMesh);
+
+		TSharedPtr<TArray<int>> List = IndexList.List;
+		TUniquePtr<FTessellationPattern> Pattern; 
+		if (List == nullptr) //if list is not provided then tessellate the whole mesh
+		{	
+			switch(PatternType) 
+			{
+				case EAdaptiveTessellatePatternType::Uniform:
+				{
+					Pattern = nullptr; //TODO: implement
+					break;
+				}
+				case EAdaptiveTessellatePatternType::InnerUniform: 
+				{
+					Pattern = nullptr; //TODO: implement
+					break;
+				}
+				case EAdaptiveTessellatePatternType::ConcentricRings:
+				{
+					Pattern = FAdaptiveTessellate::CreateConcentricRingsTessellationPattern(&EditMesh, TessellationLevel);
+					break;
+				} 
+				default:
+				{
+					checkSlow(false);
+				}
+			}
+		}
+		else
+		{
+			switch(PatternType) 
+			{
+				case EAdaptiveTessellatePatternType::Uniform:
+				{
+					Pattern = nullptr; //TODO: implement
+					break;
+				}
+				case EAdaptiveTessellatePatternType::InnerUniform: 
+				{
+					Pattern = nullptr; //TODO: implement
+					break;
+				}
+				case EAdaptiveTessellatePatternType::ConcentricRings:
+				{
+					Pattern = FAdaptiveTessellate::CreateConcentricRingsTessellationPattern(&EditMesh, TessellationLevel, *List);
+					break;
+				} 
+				default:
+				{
+					checkSlow(false);
+				}
+			}
+		}
+
+		Tessellator.SetPattern(Pattern.Get());
+		Tessellator.bUseParallel = Options.bEnableMultithreading;
+		
+		if (Tessellator.Validate() != EOperationValidationResult::Ok)
+		{
+			UE::Geometry::AppendError(Debug, EGeometryScriptErrorType::InvalidInputs, LOCTEXT("ApplyAdapativeTessellation_Error", "ApplyAdapativeTessellation: The inputs are invalid."));
+			return;
+		} 
+
+		if (Tessellator.Compute() == false)
+		{
+			UE::Geometry::AppendError(Debug, EGeometryScriptErrorType::OperationFailed, LOCTEXT("ApplyAdapativeTessellate_Failed", "ApplyAdapativeTessellate: Tessellation failed.")); 
+		}
+
+	}, EDynamicMeshChangeType::GeneralEdit, EDynamicMeshAttributeChangeFlags::Unknown, false);
+
+	return TargetMesh;
+}
+
 #undef LOCTEXT_NAMESPACE
diff --git a/Engine/Plugins/Experimental/GeometryScripting/Source/GeometryScriptingCore/Public/GeometryScript/MeshSubdivideFunctions.h b/Engine/Plugins/Experimental/GeometryScripting/Source/GeometryScriptingCore/Public/GeometryScript/MeshSubdivideFunctions.h
--- a/Engine/Plugins/Experimental/GeometryScripting/Source/GeometryScriptingCore/Public/GeometryScript/MeshSubdivideFunctions.h
+++ b/Engine/Plugins/Experimental/GeometryScripting/Source/GeometryScriptingCore/Public/GeometryScript/MeshSubdivideFunctions.h
@@ -9,8 +9,9 @@
 
 class UDynamicMesh;
 
-
-
+//
+// PN Tessellate options
+//
 USTRUCT(BlueprintType, meta = (DisplayName = "PN Tessellate Options"))
 struct GEOMETRYSCRIPTINGCORE_API FGeometryScriptPNTessellateOptions
 {
@@ -20,6 +21,29 @@ public:
 	bool bRecomputeNormals = true;
 };
 
+
+//
+// Adaptive Tessellate options
+//
+UENUM(BlueprintType)
+enum class EAdaptiveTessellatePatternType : uint8
+{
+	Uniform,
+	InnerUniform,
+	ConcentricRings
+};
+
+
+USTRUCT(BlueprintType, meta = (DisplayName = "Additional Adaptive Tessellate Options"))
+struct GEOMETRYSCRIPTINGCORE_API FGeometryScriptAdaptiveTessellateOptions
+{
+	GENERATED_BODY()
+public:
+	UPROPERTY(BlueprintReadWrite, Category = Options)
+	bool bEnableMultithreading = true;
+};
+
+
 UCLASS(meta = (ScriptName = "GeometryScript_MeshSubdivide"))
 class GEOMETRYSCRIPTINGCORE_API UGeometryScriptLibrary_MeshSubdivideFunctions : public UBlueprintFunctionLibrary
 {
@@ -49,4 +73,15 @@ public:
 		UDynamicMesh* TargetMesh,
 		int TessellationLevel = 3,
 		UGeometryScriptDebug* Debug = nullptr );
+
+	UFUNCTION(BlueprintCallable, Category = "GeometryScript|Subdivide", meta=(ScriptMethod))
+	static UPARAM(DisplayName = "Target Mesh") UDynamicMesh* 
+	ApplyAdaptiveTessellation(
+		UDynamicMesh* TargetMesh,
+		FGeometryScriptAdaptiveTessellateOptions Options,
+		FGeometryScriptIndexList IndexList,
+		int TessellationLevel = 1,
+		EAdaptiveTessellatePatternType PatternType = EAdaptiveTessellatePatternType::ConcentricRings,
+		UGeometryScriptDebug* Debug = nullptr);
+
 };
\ No newline at end of file
diff --git a/Engine/Plugins/Runtime/GeometryProcessing/Source/DynamicMesh/Private/Operations/AdaptiveTessellate.cpp b/Engine/Plugins/Runtime/GeometryProcessing/Source/DynamicMesh/Private/Operations/AdaptiveTessellate.cpp
new file mode 100644
--- /dev/null
+++ b/Engine/Plugins/Runtime/GeometryProcessing/Source/DynamicMesh/Private/Operations/AdaptiveTessellate.cpp
@@ -0,0 +1,1071 @@
+// Copyright Epic Games, Inc. All Rights Reserved.
+
+#include "Operations/AdaptiveTessellate.h"
+#include "VectorTypes.h"
+#include "DynamicMesh/DynamicMesh3.h"
+#include "Async/ParallelFor.h"
+#include "Util/ProgressCancel.h"
+#include "Distance/DistLine3Line3.h"
+#include "Misc/AssertionMacros.h"
+#include "Util/CompactMaps.h"
+
+using namespace UE::Geometry;
+
+namespace AdaptiveTessellateLocals 
+{
+	/**
+     * Manage data containing tessellation information. This includes barycentric coordinates of the new vertices,
+     * new vertex IDs and triangle indices. 
+     *
+     * Given an Edge or Triangle ID, allows to create a TArrayView for reading and writing to the data block 
+	 * containing barycentric coordinates, vertex IDs or triangle data.
+	 */
+	class FTessellationData
+	{
+	public:
+
+		// Track all of the edges and triangles that we need to tessellate. Allows for fast Contains() queries.
+		TSet<int> EdgesToTessellate;
+		TSet<int> TrianglesToTessellate;
+
+        // Array of all the new IDs for all the vertices added along the tessellated edges. Different edges 
+        // can have a different number of new vertices added. Use the EdgeCoordOffsets to figure out the starting point 
+        // and the length of the data block containing all of the vertices added along the edge. Similar logic applies 
+		// to EdgeCoord, InnerVIDs, InnerCoord, Triangles.
+		//
+		//            Edge1                Edge2              Edge50   Edge51
+		//            ---------------------------------------------------------------------
+		// EdgeVIDs   | VID1  VID2  VID3 | VID4  VID5 | ... | VID100 | VID101 VID102 |....
+		//            ---------------------------------------------------------------------
+		//           /                                      /
+		//          /                                       EdgeCoordOffsets[Edge50][0]
+		//         EdgeCoordOffsets[Edge1][0] 	
+
+		TArray<int> EdgeVIDs;     // Vertex IDs of the vertices added inside tessellated edges 
+		TArray<double> EdgeCoord; // Lerp cooefficients of the vertices added along tessellated edges		
+		TMap<int, FIndex2i> EdgeCoordOffsets;  // Maps Edge ID to tuple of (offset, length) numbers used to access a data 
+											   // block in the EdgeVIDs and EdgeCoord arrays.
+
+
+		TArray<int> InnerVIDs; 	      // Vertex IDs of the vertices added inside tessellated triangles 
+		TArray<FVector3d> InnerCoord; // Barycentric coordinates of the vertices added inside tessellated triangles 
+		TMap<int, FIndex2i> InnerCoordOffsets; // Maps Triangle ID to tuple of offset and length numbers used to access 
+											   // a data block in the InnerVIDs and InnerCoord arrays.
+
+
+		TArray<FIndex3i> Triangles; 		  // Indices of the new triangles
+		TMap<int, FIndex2i> TrianglesOffsets; // Maps Triangle ID to tuple of offset and length numbers used to access 
+											  // a data block in the Triangles array.
+		
+	protected:
+
+		const FDynamicMesh3* Mesh;
+
+	public:
+
+		FTessellationData(const FDynamicMesh3* InMesh) 
+		:
+		Mesh(InMesh)
+		{
+		}
+
+
+		void Init(const FTessellationPattern* Pattern) 
+		{
+			this->InitEdgeVertexBuffers(Pattern);
+			this->InitTriVertexBuffers(Pattern); 
+			this->InitTrianglesBuffers(Pattern);
+		}
+
+
+		/** 
+		 * @return Array view of the data block containing the barycentric coordinates for the new vertices added along 
+		 * the edge. 
+		 */
+		TArrayView<double> MapEdgeCoordBufferBlock(const int EdgeID) 
+		{
+			TArrayView<double> ArrayView(EdgeCoord);
+			const FIndex2i OffsetLength = EdgeCoordOffsets[EdgeID]; 
+			return ArrayView.Slice(OffsetLength[0], OffsetLength[1]);
+		}
+
+
+		/** @return Array view of the data block containing the vertex ids for the new vertices added along the edge. */
+		TArrayView<int> MapEdgeVIDBufferBlock(const int EdgeID) 
+		{
+			TArrayView<int> ArrayView(EdgeVIDs);
+			const FIndex2i OffsetLength = EdgeCoordOffsets[EdgeID]; 
+			return ArrayView.Slice(OffsetLength[0], OffsetLength[1]);
+		}
+
+
+		/** @return Array view of the data block containing the barycentric coordinates for new vertices added inside the triangle. */
+		TArrayView<FVector3d> MapTriangleCoordBufferBlock(const int TriangleID)
+		{
+			TArrayView<FVector3d> ArrayView(InnerCoord);
+			const FIndex2i OffsetLength = InnerCoordOffsets[TriangleID]; 
+			return ArrayView.Slice(OffsetLength[0], OffsetLength[1]);
+		}
+
+
+		/** @return Array view of a data block containing the vertex ids for new vertices added inside the triangle. */
+		TArrayView<int> MapTriangleVIDBufferBlock(const int TriangleID)
+		{
+			TArrayView<int> ArrayView(InnerVIDs);
+			const FIndex2i OffsetLength = InnerCoordOffsets[TriangleID]; 
+			return ArrayView.Slice(OffsetLength[0], OffsetLength[1]);
+		}
+
+
+		/** @return Array view of a data block containing the triangle indices for new triangles. */
+		TArrayView<FIndex3i> MapTrianglesBufferBlock(const int TriangleID)
+		{
+			TArrayView<FIndex3i> ArrayView(Triangles);
+			FIndex2i OffsetLengthPair = TrianglesOffsets[TriangleID]; 
+			return ArrayView.Slice(OffsetLengthPair[0], OffsetLengthPair[1]);
+		}
+
+	protected:
+
+		void InitEdgeVertexBuffers(const FTessellationPattern* Pattern) 
+		{
+			int BufferNumVertices = 0;
+			for (const int EdgeID : Mesh->EdgeIndicesItr())
+			{
+				const int NumNewVertices = Pattern->GetNumberOfNewVerticesForEdgePatch(EdgeID); 
+				if (NumNewVertices > 0)
+				{
+					EdgesToTessellate.Add(EdgeID);
+					EdgeCoordOffsets.Add(EdgeID, FIndex2i(BufferNumVertices, NumNewVertices));
+					BufferNumVertices += NumNewVertices;
+				}
+			}
+
+			EdgeVIDs.SetNum(BufferNumVertices);
+			
+			// IDs of the new vertices added along edges start with the max vertex ID of the input mesh.
+			const int MaxVertexID = Mesh->MaxVertexID();
+			
+			for (int Index = 0; Index < EdgeVIDs.Num(); ++Index)
+			{
+				EdgeVIDs[Index] = Index + MaxVertexID;
+			}
+			
+			// Pre-allocate memory for barycentric coordinates we will be computing
+			EdgeCoord.SetNum(BufferNumVertices);
+		}
+
+		
+		void InitTriVertexBuffers(const FTessellationPattern* Pattern) 
+		{
+			int BufferNumVertices = 0;
+			for (const int TriangleID : Mesh->TriangleIndicesItr())
+			{
+				const int NumNewVertices = Pattern->GetNumberOfNewVerticesForTrianglePatch(TriangleID); 
+				if (NumNewVertices > 0) 
+				{
+					InnerCoordOffsets.Add(TriangleID, FIndex2i(BufferNumVertices, NumNewVertices));
+					BufferNumVertices += NumNewVertices;
+				}
+			}
+
+			InnerVIDs.SetNum(BufferNumVertices);
+			
+			// IDs of the new vertices added inside triangles start with the last vertex id added along edges. It is 
+			// possible that none of the edges were tessellated in which case we start with the max vertex id of the 
+			// input mesh.
+			const int LastEdgeVIDs = EdgeVIDs.Num() > 0 ? EdgeVIDs.Last() + 1 : Mesh->MaxVertexID();
+			for (int Index = 0; Index < InnerVIDs.Num(); ++Index)
+			{
+				InnerVIDs[Index] = Index + LastEdgeVIDs;
+			}
+
+			// Pre-allocate memory for barycentric coordinates we will be computing
+			InnerCoord.SetNum(BufferNumVertices);
+		}
+
+
+		void InitTrianglesBuffers(const FTessellationPattern* Pattern) 
+		{
+			int BufferNumTriangles = 0;
+			for (const int TriangleID : Mesh->TriangleIndicesItr())
+			{
+				const int NumNewTriangles = Pattern->GetNumberOfPatchTriangles(TriangleID); 
+				if (NumNewTriangles > 0) 
+				{	
+					TrianglesToTessellate.Add(TriangleID);
+					TrianglesOffsets.Add(TriangleID, FIndex2i(BufferNumTriangles, NumNewTriangles));
+					BufferNumTriangles += NumNewTriangles;
+				}
+			}
+
+			// Pre-allocate memory for triangles we will be computing
+			Triangles.SetNum(BufferNumTriangles);
+		}
+	}; 
+
+
+
+
+	/**
+     * Pattern where the inner area is tessellated using the style of the OpenGL tessellation shader. The inner area 
+	 * consists of multiple inner rings. The original triangle we are tessellating is called the outer ring. Areas 
+	 * between rings are tessellated with the new triangles.
+	 * 
+     *                 O
+     *                / \
+     *               O. .O
+     *              /  O  \
+     *             O. / \ .O
+     *            /  O. .O  \
+     *           /  /  O  \  \
+     *          /  /  / \  \  \
+     *         O. /  /   \  \ .O
+     *        /  O. /     \ .O  \
+     *       /  /  O-------O  \  \
+     *      O. / inner ring 2  \ .O
+     *     /  O----O-------O----O  \
+     *    /   .  inner ring 1   .   \
+     *   O----O----O-------O----O----O 
+     *			 outer ring
+	 *
+     */
+	class FConcentricRingsTessellationPattern : public FTessellationPattern
+	{
+	public:
+
+		/**
+		 * FRing edge is a collection of vertices that includes two corner vertices (marked + below) plus all the 
+		 * vertices in between (marked O below).
+		 * 
+		 *                 +  
+		 *                / \ 
+		 *               O   \  
+		 *              /     \  
+		 * ring edge 3 /       O  ring edge 1
+		 *            /         \ 
+		 *           O           \  
+		 *          /             \  
+		 *         /               \ 
+		 *        +---O----O----O---+  
+		 * 			  ring edge 2
+		 
+		*/
+			
+		class FRingEdge 
+		{
+
+		public:
+			FRingEdge() 
+			{
+			}
+
+			/** FRing edge can be a single point. */
+			FRingEdge(int V1) 
+			:
+			V1(V1), VertNum(1)
+			{
+			}
+			
+			FRingEdge(int V1, int V2, TArrayView<int> InInner, bool bReverse = false) 
+			:
+			V1(V1), V2(V2), Inner(InInner), VertNum(InInner.Num() + 2), bReverse(bReverse)
+			{
+			}
+
+			FRingEdge(int V1, int V2, EdgePatch Patch)
+			:
+			FRingEdge(V1, V2, Patch.VIDs, Patch.bIsReversed)
+			{
+			}
+
+			inline int Num() const 
+			{
+				return VertNum;
+			}
+
+			int operator[](int Index) const
+			{
+				check(Index >= 0 && Index < VertNum);
+
+				if (Index == 0) 
+				{
+					return V1;
+				} 
+				else if (Index == VertNum - 1) 
+				{
+					return V2;
+				}
+				else 
+				{
+					int InnerIndex = Index - 1; // index into the Inner array
+					int ActualIndex = bReverse ? Inner.Num() - InnerIndex - 1: InnerIndex; // potentially reverse index
+					check(ActualIndex >= 0 && ActualIndex < VertNum);
+					return Inner[ActualIndex];
+				}
+			}
+
+		private:
+			
+			int V1 = IndexConstants::InvalidID;
+			int V2 = IndexConstants::InvalidID;
+			TArrayView<int> Inner;
+			int VertNum = 0; // Number of elements
+			bool bReverse = false;
+		};
+
+		/* A ring consists of 3 ring edges */
+		struct FRing 
+		{
+			FRingEdge UV;
+			FRingEdge VW;
+			FRingEdge UW;
+		};
+
+
+		FConcentricRingsTessellationPattern(const FDynamicMesh3* InMesh) 
+		:
+		FTessellationPattern(InMesh)
+		{
+		}
+
+
+		virtual ~FConcentricRingsTessellationPattern() 
+		{
+		}
+
+		
+		void Init(const TArray<int>& InEdgeTessLevels, const TArray<int>& InInnerTessLevels) 
+		{
+			EdgeTessLevels = InEdgeTessLevels;
+			InnerTessLevels = InInnerTessLevels;
+
+			// If the inner area of the triangle patch is not being tessellated but at least one of its edges is, then 
+			// we insert a single vertex in the middle and connect it to all of the edge vertices. Therefore, we find 
+			// those triangles and set their inner tessellation level to 1. This would tell the TessellatePatch function 
+			// to insert a vertex and generate the triangle fan.
+			for (const int TriangleID : Mesh->TriangleIndicesItr()) 
+			{
+				if (InnerTessLevels[TriangleID] <= 0)
+				{	
+					const FIndex3i TriEdges = Mesh->GetTriEdges(TriangleID);
+
+					if (EdgeTessLevels[TriEdges[0]] || EdgeTessLevels[TriEdges[1]] || EdgeTessLevels[TriEdges[2]])
+					{
+						InnerTessLevels[TriangleID] = 1;
+					}
+				}
+			}
+		}
+
+		/** 
+		 * Convenience method to convert an inner tessellation level from the number of vertices to the number of segments.
+		 * 
+		 * 		O---O---O  3 vertices => 2 segments
+		 */
+		inline int GetInnerLevelAsSegments(const int InTriangleID) const
+		{
+			return InnerTessLevels[InTriangleID] - 1;
+		}
+
+		virtual int GetNumberOfNewVerticesForEdgePatch(const int InEdgeID) const override 
+		{
+			if (Mesh->IsEdge(InEdgeID) == false) 
+			{
+				checkNoEntry();
+				return InvalidIndex;
+			}
+
+			return EdgeTessLevels[InEdgeID];
+		}
+
+
+		virtual int GetNumberOfNewVerticesForTrianglePatch(const int InTriangleID) const override 
+		{
+			if (Mesh->IsTriangle(InTriangleID) == false) 
+			{
+				checkNoEntry();
+				return InvalidIndex;
+			}
+
+			const int TessLevel = GetInnerLevelAsSegments(InTriangleID);
+			if (TessLevel < 0) 
+			{
+				return 0;
+			} 
+			else if (TessLevel == 0) 
+			{	
+				return 1; // one vertex in the middle
+			}
+
+			// We iterate through rings and count how many vertices we are introducing per ring
+			int NumNewTriangleVertices = 0;
+			for (int RingLevel = TessLevel; RingLevel > 0; RingLevel -= 2) // next inner ring has 2 less segments
+			{
+				// 3 corners plus TessLevel - 1 vertcies along each of the 3 edges.
+				NumNewTriangleVertices += 3 + 3 * (RingLevel - 1);
+				
+				// If the current ring only has 2 segments then we will be inserting one extra vertex in the middle
+				NumNewTriangleVertices += RingLevel == 2 ? 1 : 0;
+			}
+
+			return NumNewTriangleVertices;
+		}
+
+
+		virtual int GetNumberOfPatchTriangles(const int InTriangleID) const override 
+		{
+			if (Mesh->IsTriangle(InTriangleID) == false) 
+			{
+				checkNoEntry();
+				return InvalidIndex;
+			}
+		
+			const int TessLevel = GetInnerLevelAsSegments(InTriangleID);
+
+			// Count how many triangles are we generating when connecting outer ring with the first inner ring
+			int NumNewTrianglesPerFace = 0;
+			for (int EdgeIdx = 0; EdgeIdx < 3; ++EdgeIdx) 
+			{
+				const int EdgeID = Mesh->GetTriEdge(InTriangleID, EdgeIdx);
+				const int OuterEdgeTriangles = this->GetNumberOfNewVerticesForEdgePatch(EdgeID) + 1;
+				const int InnerEdgeTriangles = TessLevel;
+				NumNewTrianglesPerFace += OuterEdgeTriangles + InnerEdgeTriangles;
+			}
+
+			// Iterate over every pair of inner rings and count how many triangles we are generating between them
+			int RingLevel = TessLevel;
+			while (RingLevel > 0) 
+			{	
+				const int CurRingLevel = RingLevel;
+				const int NextRingLevel = RingLevel - 2;
+				
+				if (NextRingLevel < 0) 
+				{	
+					// no more rings left, so the current ring is a single triangle at the center
+					checkSlow(CurRingLevel == 1);
+					NumNewTrianglesPerFace += 1;
+				}
+				else 
+				{
+					NumNewTrianglesPerFace += 3 * (CurRingLevel + NextRingLevel);
+				}
+				
+				RingLevel -= 2;
+			}
+
+			return NumNewTrianglesPerFace;
+		}
+
+
+		virtual void TessellateEdgePatch(EdgePatch& EdgePatch) const override
+		{
+			const int EdgeTessLevel = EdgeTessLevels[EdgePatch.EdgeID];
+			
+			checkSlow(EdgeTessLevel == EdgePatch.LinearCoord.Num());
+			
+			const int NumSegments = EdgeTessLevel + 1;
+			const double Step = 1.0 / NumSegments;
+			
+			for (int Idx = 0; Idx < EdgeTessLevel; ++Idx)
+			{
+				double Value = (Idx + 1)*Step;
+				EdgePatch.LinearCoord[Idx] = Value;
+			}
+		}
+
+
+		virtual void TessellateTriPatch(TrianglePatch& TriPatch) const override
+		{   
+			FRing OuterRing;
+			OuterRing.UV = FRingEdge(TriPatch.UVWCorners[0], TriPatch.UVWCorners[1], TriPatch.UVEdge);
+			OuterRing.VW = FRingEdge(TriPatch.UVWCorners[1], TriPatch.UVWCorners[2], TriPatch.VWEdge);
+			OuterRing.UW = FRingEdge(TriPatch.UVWCorners[2], TriPatch.UVWCorners[0], TriPatch.UWEdge);
+
+			TArray<FRing> RingArray;
+			RingArray.Add(OuterRing);
+
+			// Track which vertex index we are currently working with
+			int InnerIdx = 0;
+		
+            // We are working with an abstract patch whose vertex coordinates are the same as their barycentric coordinates
+			FVector3d OuterU = FVector3d(1.0, 0.0, 0.0);
+			FVector3d OuterV = FVector3d(0.0, 1.0, 0.0);
+			FVector3d OuterW = FVector3d(0.0, 0.0, 1.0);
+
+		    // Generate inner rings. Each subsequent inner ring will have its level reduced by 2.
+            // The last inner ring can either be a single triangle or a single point.
+			FVector3d InnerU, InnerV, InnerW;
+			for (int TesLevel = GetInnerLevelAsSegments(TriPatch.TriangleID); TesLevel >= 0; TesLevel -= 2)
+			{
+				// Given the 3 corner vertices of the previous ring, compute the 3 corner vertices of this inner ring
+				this->ComputeInnerConcentricTriangle(OuterU, OuterV, OuterW, TesLevel + 1, InnerU, InnerV, InnerW);
+
+				FRing InnerRing;
+
+				int StartIdx = InnerIdx; // save the index of the U corner since we will wrap back to it
+				
+				const double Step = TesLevel == 0 ? 0.0 : 1.0 / TesLevel;
+
+				// Compute the barycentric coordinates for all vertices in the UV ring edge
+				{
+					TriPatch.BaryCoord[InnerIdx] = InnerU;
+					for (int Idx = 1; Idx < TesLevel; ++Idx)
+					{
+						TriPatch.BaryCoord[InnerIdx + Idx] = InnerU + Idx*Step*(InnerV - InnerU);
+					}
+					TriPatch.BaryCoord[InnerIdx + TesLevel] = InnerV;
+
+
+					if (TesLevel == 0) 
+					{
+						InnerRing.UV = FRingEdge(TriPatch.VIDs[InnerIdx]); // just a single vertex in the middle
+					}
+					else 
+					{
+						InnerRing.UV = FRingEdge(TriPatch.VIDs[InnerIdx], TriPatch.VIDs[InnerIdx + TesLevel], TriPatch.VIDs.Slice(InnerIdx + 1, TesLevel - 1));
+					}
+					
+					InnerIdx += TesLevel;
+				}
+
+				// Compute the barycentric coordinates for all vertices in the VW ring edge
+				{
+					TriPatch.BaryCoord[InnerIdx] = InnerV;
+					for (int Idx = 1; Idx < TesLevel; ++Idx)
+					{
+						TriPatch.BaryCoord[InnerIdx + Idx] = InnerV + Idx*Step*(InnerW - InnerV);
+					}
+					TriPatch.BaryCoord[InnerIdx + TesLevel] = InnerW;
+
+					if (TesLevel == 0) 
+					{
+						InnerRing.VW = FRingEdge(TriPatch.VIDs[InnerIdx]);
+					}
+					else 
+					{
+						InnerRing.VW = FRingEdge(TriPatch.VIDs[InnerIdx], TriPatch.VIDs[InnerIdx + TesLevel], TriPatch.VIDs.Slice(InnerIdx + 1, TesLevel - 1));
+					}
+
+					InnerIdx += TesLevel;
+				}
+
+				// Compute the barycentric coordinates for all vertices in the UW ring edge
+				{
+					TriPatch.BaryCoord[InnerIdx] = InnerW;
+					for (int Idx = 1; Idx < TesLevel; ++Idx)
+					{
+						TriPatch.BaryCoord[InnerIdx + Idx] = InnerW + Idx*Step*(InnerU - InnerW);
+					}
+
+					if (TesLevel == 0)
+					{
+						InnerRing.UW = FRingEdge(TriPatch.VIDs[InnerIdx]);
+					}
+					else 
+					{
+						InnerRing.UW = FRingEdge(TriPatch.VIDs[InnerIdx], TriPatch.VIDs[StartIdx], TriPatch.VIDs.Slice(InnerIdx + 1, TesLevel - 1));
+					}
+
+					InnerIdx += TesLevel;
+				}
+
+				OuterU = InnerU;
+				OuterV = InnerV;
+				OuterW = InnerW;
+
+				RingArray.Add(InnerRing);
+			}
+
+            // Iterate over pair of rings and tessellate the area between them
+			int TriangleIdx = 0; // Offset into TriPatch.Triangles where we will be adding new triangles
+			for (int RingIdx = 0; RingIdx < RingArray.Num() - 1; ++RingIdx)
+			{
+				TriangleIdx += this->StitchRingEdges(RingArray[RingIdx].UV, RingArray[RingIdx + 1].UV, TriangleIdx, TriPatch.Triangles);
+				TriangleIdx += this->StitchRingEdges(RingArray[RingIdx].VW, RingArray[RingIdx + 1].VW, TriangleIdx, TriPatch.Triangles);
+				TriangleIdx += this->StitchRingEdges(RingArray[RingIdx].UW, RingArray[RingIdx + 1].UW, TriangleIdx, TriPatch.Triangles);
+			}
+
+            // If the tessellation level is odd, then we need to create one more triangle with the last inner ring vertices
+			if (GetInnerLevelAsSegments(TriPatch.TriangleID)% 2 != 0) 
+			{	
+				const FRing& LastRing = RingArray.Last();
+				TriPatch.Triangles[TriangleIdx] = FIndex3i(LastRing.UV[0], LastRing.UV[1], LastRing.VW[1]);
+			}
+		}
+
+
+		/**
+		 * Given two ring edges, "stitch" them together by generating a sequence of triangles connecting all vertices.
+		 * 
+		 *    Outer ring edge         O  O                            O--O  
+		 *                                      	  ==>            /\  /\
+		 *                                                          /  \/  \
+		 *    Inner ring edge 1    O    O    O                     O----O---O 
+		 * 
+		 * @return the number of triangles generated
+		 */
+		int StitchRingEdges(const FRingEdge& InOuter, const FRingEdge& InInner, int Offset, TArrayView<FIndex3i>& OutTriangles) const 
+		{
+			const int OuterNum = InOuter.Num();
+			const int InnerNum = InInner.Num();
+			int TessDiff = InnerNum - OuterNum;
+			
+			const int NumTriangles = OuterNum - 1 + InnerNum - 1; // we know how many triangles to expect
+			
+            // Track which vertex are we currently at in the outer and inner vertex arrays
+			int OuterIdx = 0;
+			int InnerIdx = 0;
+
+		    // We move along the outer and inner vertices and generate triangles. For each triangle, we need to choose which 
+            // ring edge do we pick the third vertex from. Is it the OuterIdx + 1 or the InnerIdx + 1 vertex? We use the 
+            // updated TessDiff variable to make the choice. If it's negative we choose the outer vertex, otherwise the inner vertex.
+            // For more information see the Sec. 5.1 and Figure 8 in "Watertight Tessellation using Forward Differencing, H. Moreton, 2001."
+            //
+			//       OuterIdx	  
+			//         \
+			//          O  O     
+			//                          
+			//                          
+			//       O    O    O 
+			//      /
+			//   InnerIdx
+			
+			for (int TriIndex = 0; TriIndex < NumTriangles; ++TriIndex) 
+			{
+				// If TessDiff is negative and the next vertex along outer side is available or we simply run out of inner vertices,
+				// then use outer vertex as the third vertex to create a triangle 
+				if ((TessDiff < 0 && OuterIdx + 1 < OuterNum) || InnerIdx == InnerNum - 1) 
+				{
+					OutTriangles[Offset + TriIndex] = FIndex3i(InOuter[OuterIdx], InOuter[OuterIdx + 1], InInner[InnerIdx]);
+					TessDiff += 2*InnerNum;
+					OuterIdx += 1;
+				}
+				else
+				{
+					ensure((TessDiff >= 0 && InnerIdx + 1 < InnerNum) || OuterIdx == OuterNum - 1);
+					OutTriangles[Offset + TriIndex] = FIndex3i(InInner[InnerIdx], InOuter[OuterIdx], InInner[InnerIdx + 1]);
+					TessDiff -= 2*OuterNum;
+					InnerIdx += 1;
+				}
+			}
+			
+			return NumTriangles;
+		}
+	
+		/** An array of per triangle (inner) tessellation levels. The size must match the maximum triangle ID of the mesh. */
+		TArray<int> InnerTessLevels;
+
+		/** An array of per edge tessellation levels. The size must match the maximum edge ID of the mesh. */
+		TArray<int> EdgeTessLevels;
+	};
+
+
+
+
+	/** 
+	 * Use the tessellation pattern to tessellate the mesh and generate the tessellation data (barycentric coordinates 
+	 * of the new vertices and new triangles).
+	 */
+	void Tessellate(const FDynamicMesh3* Mesh, const FTessellationPattern* Pattern, const bool bUseParallel, FTessellationData& TessData) 
+	{	
+		// Tessellate edge patches
+		TArray<int> EdgesToTessellate = TessData.EdgesToTessellate.Array();
+		ParallelFor(EdgesToTessellate.Num(), [&](int32 Index)
+		{
+			FTessellationPattern::EdgePatch Patch;
+
+			const int EdgeID = EdgesToTessellate[Index];
+			Patch.EdgeID = EdgeID;
+			Patch.LinearCoord = TessData.MapEdgeCoordBufferBlock(EdgeID);
+			Patch.VIDs = TessData.MapEdgeVIDBufferBlock(EdgeID);
+			
+			Pattern->TessellateEdgePatch(Patch);
+		}, bUseParallel ? EParallelForFlags::None : EParallelForFlags::ForceSingleThread);
+
+		
+		// Tessellate triangle patches
+		TArray<int> TrianglesToTessellate = TessData.TrianglesToTessellate.Array();
+		ParallelFor(TrianglesToTessellate.Num(), [&](int32 Index)
+		{
+			FTessellationPattern::TrianglePatch Patch;
+			
+			const int TriangleID = TrianglesToTessellate[Index]; 
+			
+			Patch.TriangleID = TriangleID;
+			
+			FIndex3i TriVertices = Mesh->GetTriangle(TriangleID);
+			Patch.UVWCorners = TriVertices;
+			
+			const FIndex3i& TriEdges = Mesh->GetTriEdgesRef(TriangleID);
+			
+			// UV Edge
+			if (TessData.EdgesToTessellate.Contains(TriEdges[0]))
+			{
+				Patch.UVEdge.EdgeID = TriEdges[0];
+				Patch.UVEdge.LinearCoord = TessData.MapEdgeCoordBufferBlock(TriEdges[0]);
+				Patch.UVEdge.VIDs = TessData.MapEdgeVIDBufferBlock(TriEdges[0]);
+				Patch.UVEdge.bIsReversed = Mesh->GetEdgeV(TriEdges[0])[0] != TriVertices[0];
+			}
+
+			// VW Edge
+			if (TessData.EdgesToTessellate.Contains(TriEdges[1]))
+			{
+				Patch.VWEdge.EdgeID = TriEdges[1];
+				Patch.VWEdge.LinearCoord = TessData.MapEdgeCoordBufferBlock(TriEdges[1]);
+				Patch.VWEdge.VIDs = TessData.MapEdgeVIDBufferBlock(TriEdges[1]);
+				Patch.VWEdge.bIsReversed = Mesh->GetEdgeV(TriEdges[1])[0] != TriVertices[1];
+			}
+			
+			// UW Edge
+			if (TessData.EdgesToTessellate.Contains(TriEdges[2]))
+			{
+				Patch.UWEdge.EdgeID = TriEdges[2];
+				Patch.UWEdge.LinearCoord = TessData.MapEdgeCoordBufferBlock(TriEdges[2]);
+				Patch.UWEdge.VIDs = TessData.MapEdgeVIDBufferBlock(TriEdges[2]);
+				Patch.UWEdge.bIsReversed = Mesh->GetEdgeV(TriEdges[2])[0] != TriVertices[2];
+			}
+
+			//  Inner Triangle Vertices
+			Patch.BaryCoord = TessData.MapTriangleCoordBufferBlock(TriangleID);
+			Patch.VIDs = TessData.MapTriangleVIDBufferBlock(TriangleID);
+
+			// Inner Triangles
+			Patch.Triangles = TessData.MapTrianglesBufferBlock(TriangleID);
+
+			// Tessellate the patch, i.e. generate inner vertices and triangles
+			Pattern->TessellateTriPatch(Patch);
+		}, bUseParallel ? EParallelForFlags::None : EParallelForFlags::ForceSingleThread);
+	}
+
+
+	/**
+	 * Given the original FDynamicMesh3 and the tessellation data, construct new tessellated FDynamicMesh3 geometry and 
+	 * interpolate attributes.
+	 */
+	void ConstructTessellatedMesh(const FDynamicMesh3* Mesh, 
+								  const FTessellationData& TessData, 
+								  FDynamicMesh3* ResultMesh)  
+	{	
+		ResultMesh->Clear();
+
+		// Handle non-compact meshes
+		FCompactMaps CompactInfo;
+
+		const int ResultNumVertexIDs = Mesh->MaxVertexID() + TessData.EdgeVIDs.Num() + TessData.InnerVIDs.Num();
+		CompactInfo.ResetVertexMap(ResultNumVertexIDs, false);
+
+		for (int VertexID = 0; VertexID < Mesh->MaxVertexID(); ++VertexID)
+		{
+			if (Mesh->IsVertex(VertexID))
+			{
+				CompactInfo.SetVertexMapping(VertexID, ResultMesh->AppendVertex(Mesh->GetVertex(VertexID)));
+			}
+			else 
+			{
+				CompactInfo.SetVertexMapping(VertexID, FCompactMaps::InvalidID);
+			}
+		}
+
+		for (const int EdgeID : TessData.EdgesToTessellate)
+		{
+			const FIndex2i OffsetNum = TessData.EdgeCoordOffsets[EdgeID];
+
+			const FIndex2i EdgeV = Mesh->GetEdgeV(EdgeID);
+
+			for (int Index = 0; Index < OffsetNum[1]; ++Index)
+			{	
+				const int Offset = OffsetNum[0] + Index;
+				const double Alpha = TessData.EdgeCoord[Offset];
+				const int VertexID = TessData.EdgeVIDs[Offset]; 
+
+				const FVector3d Vertex = (1.0 - Alpha)*Mesh->GetVertex(EdgeV[0]) + Alpha*Mesh->GetVertex(EdgeV[1]);
+				
+				CompactInfo.SetVertexMapping(VertexID, ResultMesh->AppendVertex(Vertex));
+			}
+		}
+
+
+		for (const int TriangleID : TessData.TrianglesToTessellate)
+		{
+			const FIndex2i OffsetNum = TessData.InnerCoordOffsets[TriangleID];
+
+			const FIndex3i TriangleV = Mesh->GetTriangle(TriangleID);
+
+			for (int Index = 0; Index < OffsetNum[1]; ++Index)
+			{	
+				const int Offset = OffsetNum[0] + Index;	
+				
+				const FVector3d Bary = TessData.InnerCoord[Offset];
+				const int VertexID = TessData.InnerVIDs[Offset];
+
+				const FVector3d Vertex = Bary[0] * Mesh->GetVertex(TriangleV[0]) + 
+										 Bary[1] * Mesh->GetVertex(TriangleV[1]) + 
+										 Bary[2] * Mesh->GetVertex(TriangleV[2]);
+				
+				CompactInfo.SetVertexMapping(VertexID, ResultMesh->AppendVertex(Vertex));
+			}
+		}
+
+
+		for (const int TriangleID : Mesh->TriangleIndicesItr())
+		{		
+			if (TessData.TrianglesToTessellate.Contains(TriangleID)) 
+			{
+				const FIndex2i OffsetNum = TessData.TrianglesOffsets[TriangleID];
+				for (int Index = 0; Index < OffsetNum[1]; ++Index)
+				{	
+					const int Offset = OffsetNum[0] + Index;	
+					FIndex3i Tri = TessData.Triangles[Offset];
+					FIndex3i MappedTri = CompactInfo.GetVertexMapping(Tri);
+					ResultMesh->AppendTriangle(MappedTri);
+				}
+			}
+			else 
+			{	
+				FIndex3i Tri = Mesh->GetTriangle(TriangleID);
+				FIndex3i MappedTri = CompactInfo.GetVertexMapping(Tri);
+				ResultMesh->AppendTriangle(MappedTri);
+			}
+		}
+
+		//TODO: use barycentric data to interpolate attributes
+		
+		//TODO: instead of generating the mesh from scratch, we could remove triangles we are tessellating from the 
+		// input mesh and append the new ones. This can be an option when the number of triangles tessellated is small,
+		// compared to the total number of triangles.
+	}
+}
+
+
+
+
+//
+// GLSL pattern
+//
+
+TUniquePtr<FTessellationPattern> FAdaptiveTessellate::CreateConcentricRingsTessellationPattern(const FDynamicMesh3* InMesh,	
+												   											   const TArray<int>& InEdgeTessLevels,
+                                                   											   const TArray<int>& InInnerTessLevels)
+{
+	TUniquePtr<AdaptiveTessellateLocals::FConcentricRingsTessellationPattern> Pattern = MakeUnique<AdaptiveTessellateLocals::FConcentricRingsTessellationPattern>(InMesh);
+	Pattern->Init(InEdgeTessLevels, InInnerTessLevels);
+	
+	return Pattern;
+}
+
+
+TUniquePtr<FTessellationPattern> FAdaptiveTessellate::CreateConcentricRingsTessellationPattern(const FDynamicMesh3* InMesh,
+												   											   const int InTessellationLevel)
+{
+	TArray<int> InEdgeTessLevels;
+	InEdgeTessLevels.Init(InTessellationLevel, InMesh->EdgeCount());
+	
+	TArray<int> InInnerTessLevels;
+	InInnerTessLevels.Init(InTessellationLevel, InMesh->TriangleCount());
+	
+	return FAdaptiveTessellate::CreateConcentricRingsTessellationPattern(InMesh, InEdgeTessLevels, InInnerTessLevels);
+}
+
+ 
+TUniquePtr<FTessellationPattern> FAdaptiveTessellate::CreateConcentricRingsTessellationPattern(const FDynamicMesh3* InMesh,  
+								   							  					    		   const int InTessellationLevel,
+								   							  					    		   const TArray<int>& InTriangleList) 
+{
+	TArray<int> EdgeTessLevels;
+	EdgeTessLevels.Init(0, InMesh->MaxEdgeID());
+	for (const int TriangleID : InTriangleList) 
+	{
+		const FIndex3i EdgesIdx = InMesh->GetTriEdges(TriangleID);
+		EdgeTessLevels[EdgesIdx[0]] = InTessellationLevel;
+		EdgeTessLevels[EdgesIdx[1]] = InTessellationLevel;
+		EdgeTessLevels[EdgesIdx[2]] = InTessellationLevel;
+	}
+
+	TArray<int> TriangleTessLevels;
+	TriangleTessLevels.Init(0, InMesh->MaxTriangleID());
+
+	for (const int TriangleID : InTriangleList) 
+	{
+		TriangleTessLevels[TriangleID] = InTessellationLevel; 
+	}
+
+	return FAdaptiveTessellate::CreateConcentricRingsTessellationPattern(InMesh, EdgeTessLevels, TriangleTessLevels);	
+}
+
+
+TUniquePtr<FTessellationPattern> FAdaptiveTessellate::CreateConcentricRingsTessellationPattern(const TFunctionRef<int(const int EdgeID)> InEdgeFunc, 
+                                                               					    const TFunctionRef<int(const int TriangleID)> InTriFunc) 
+{
+
+	// TODO: not implemented yet
+	checkNoEntry();
+	return nullptr;
+}
+
+
+
+
+//
+// Inner uniform pattern
+//
+
+TUniquePtr<FTessellationPattern> FAdaptiveTessellate::CreateInnerUnifromTessellationPattern(const FDynamicMesh3* InMesh,
+													  	   						 		    const TArray<int>& InEdgeTessLevels,
+                                                      	   									const TArray<int>& InInnerTessLevels)
+{
+	// TODO: not implemented yet
+	checkNoEntry();
+	return nullptr;
+}
+
+
+
+
+//
+// Uniform pattern
+//
+TUniquePtr<FTessellationPattern> CreateUniformTessellationPattern(const FDynamicMesh3* InMesh,
+																  const int InTessellationLevel,
+																  const TArray<int>& InTriangleList) 
+{
+	// TODO: not implemented yet
+	checkNoEntry();
+	return nullptr;
+}
+
+
+
+
+void FTessellationPattern::ComputeInnerConcentricTriangle(const FVector3d& V1,
+                                        				  const FVector3d& V2,
+                                        				  const FVector3d& V3,
+                                        				  const int EdgeTessLevel,
+                                        				  FVector3d& InnerU,
+                                        				  FVector3d& InnerV,
+                                        				  FVector3d& InnerW) const
+{	
+	//TODO: Handle the case where V1, V2, V3 form a degenerate triangle 
+    
+	// Given two lines defined by a point and a normal find their intersection. It is guaranteed 
+    // that both lines are on the same plane and that they are not parallel.
+	auto IntersectLines = [] (const FVector3d& V1, const FVector3d& N1, const FVector3d& V2, const FVector3d& N2) 
+	{
+		FLine3d Line1(V1, N1);
+		FLine3d Line2(V2, N2);
+		FDistLine3Line3d LineDist(Line1, Line2); 
+		LineDist.ComputeResult();
+		FVector3d OffsetPoint = 0.5 * (LineDist.Line1ClosestPoint + LineDist.Line2ClosestPoint);
+		return OffsetPoint;
+	};
+
+
+	if (EdgeTessLevel < 0) 
+	{
+		checkSlow(false);
+		return;
+	}
+
+	if (EdgeTessLevel == 0) // edges are not tessellated
+	{
+		InnerU = V1;
+		InnerV = V2;
+		InnerW = V3;
+		return;
+	}
+	else if (EdgeTessLevel == 1) // just a vertex in the middle
+	{
+		FVector3d Center = (V1 + V2 + V3)/3.0;
+		InnerU = Center;
+		InnerV = Center;
+		InnerW = Center;
+		return;
+	}
+
+	// Inserted vertices on each subdivided edge closest to the edge corners
+	FVector3d EdgeDir1 = (V2 - V1) / (EdgeTessLevel + 1);
+	FVector3d UV1 = V1 + EdgeDir1;
+	FVector3d UV2 = V1 + EdgeTessLevel*EdgeDir1;
+
+	FVector3d EdgeDir2 = (V3 - V1) / (EdgeTessLevel + 1);
+	FVector3d UW1 = V1 + EdgeDir2;
+	FVector3d UW2 = V1 + EdgeTessLevel*EdgeDir2;
+
+	FVector3d EdgeDir3 = (V3 - V2) / (EdgeTessLevel + 1);
+	FVector3d VW1 = V2 + EdgeDir3;
+	FVector3d VW2 = V2 + EdgeTessLevel*EdgeDir3;
+	
+	// Normal vector of the plane formed by the triangle
+	FVector3d PlaneNormal = Normalized(EdgeDir1.Cross(EdgeDir2));
+	
+	// Edge normal vectors
+	FVector3d OuterUVNormal = Normalized(PlaneNormal.Cross(EdgeDir1));
+	FVector3d OuterUWNormal = Normalized(EdgeDir2.Cross(PlaneNormal));
+	FVector3d OuterVWNormal = Normalized(PlaneNormal.Cross(EdgeDir3));
+
+	// Vertices of the inner triangle
+	InnerU = IntersectLines(UV1, OuterUVNormal, UW1, OuterUWNormal);
+	InnerV = IntersectLines(UV2, OuterUVNormal, VW1, OuterVWNormal);
+	InnerW = IntersectLines(UW2, OuterUWNormal, VW2, OuterVWNormal);
+}
+
+
+bool FAdaptiveTessellate::Cancelled()
+{
+	return (Progress == nullptr) ? false : Progress->Cancelled();
+}
+
+
+bool FAdaptiveTessellate::Compute()
+{	
+	if (Validate() != EOperationValidationResult::Ok) 
+	{
+		return false;
+	}
+
+	// Check for the empty meshes
+	if (bInPlace && ResultMesh->TriangleCount() == 0) 
+	{
+		return true;
+	}
+
+	if (bInPlace == false && Mesh->TriangleCount() == 0) 
+	{
+		return true;
+	}
+
+	// Make a copy of the mesh since we are tessellating in place. The copy will be used to restore the mesh to its 
+	// original state in case the user cancelled the operation.
+	TUniquePtr<FDynamicMesh3> ResultMeshCopy = nullptr;
+	if (bInPlace)
+	{
+		ResultMeshCopy = MakeUnique<FDynamicMesh3>();
+		ResultMeshCopy->Copy(*ResultMesh);
+		Mesh = ResultMeshCopy.Get();
+	}
+
+	AdaptiveTessellateLocals::FTessellationData TessData(Mesh);
+	TessData.Init(Pattern);
+
+	AdaptiveTessellateLocals::Tessellate(Mesh, Pattern, bUseParallel, TessData);
+
+	AdaptiveTessellateLocals::ConstructTessellatedMesh(Mesh, TessData, ResultMesh);
+
+	if (Cancelled()) 
+	{
+		if (bInPlace) 
+		{ 
+			// Restore the input mesh
+			checkSlow(ResultMeshCopy != nullptr);
+			*ResultMesh = MoveTemp(*ResultMeshCopy);
+		}
+		else 
+		{
+			ResultMesh->Clear();
+		}
+
+		return false;
+	}
+	
+	return true;
+}
\ No newline at end of file
diff --git a/Engine/Plugins/Runtime/GeometryProcessing/Source/DynamicMesh/Public/Operations/AdaptiveTessellate.h b/Engine/Plugins/Runtime/GeometryProcessing/Source/DynamicMesh/Public/Operations/AdaptiveTessellate.h
new file mode 100644
--- /dev/null
+++ b/Engine/Plugins/Runtime/GeometryProcessing/Source/DynamicMesh/Public/Operations/AdaptiveTessellate.h
@@ -0,0 +1,447 @@
+// Copyright Epic Games, Inc. All Rights Reserved.
+
+#pragma once
+
+#include "GeometryTypes.h"
+#include "IndexTypes.h"
+#include "Containers/Array.h"
+#include "Containers/ArrayView.h"
+#include "Math/MathFwd.h"
+
+class FProgressCancel;
+
+namespace UE
+{
+namespace Geometry
+{
+
+class FDynamicMesh3;
+
+/**
+ * An abstract class for writing tessellation patterns. You can subclass it and implement its virtual methods to specify 
+ * a custom tessellation behavior. Original mesh edges and triangles are referred to as edge patches and triangle patches. 
+ * The edge patch handles all the vertices we are adding along the edge. The triangle patch handles all the vertices we 
+ * are adding inside of the patch (excluding the edges) and the new triangles we are splitting the patch into.
+ * 
+ * Patterns describe how individual patches are tessellated. Although you can access mesh data during the tessellation, 
+ * the patches do not generate actual vertex coordinates or other attribute data. Instead, new vertices are generated in 
+ * terms of linear or barycentric coefficients with respect to the edge or triangle vertices. The coefficient data can 
+ * later be used to interpolate the actual mesh data (vertex positions, colors, normals, etc). 
+ * 
+ * Patterns do not need to worry about generating the vertex IDs as those will be provided for each patch during the
+ * tessellation and need to simply be referred to when generating triangles.
+ * 
+ * The job of the tessellation pattern is:
+ * 
+ *      - Provide the number of new vertices introduced for an edge or a triangle patch.
+ *      
+ *      - Provide the number of new triangles introduced for a triangle patch.
+ *        
+ *      - Given a single edge patch with vertices [v1,v2], provide coordinates of the new vertices 
+ *        [a_i, a_(i+1), ...] inserted along the edge via a linear coeffcient [t_i, t_(i+1), ...] such that 
+ *        a_i = v1 + t_i*(v2-v1).
+ *        
+ *      - Given a triangle patch with vertices [v1,v2,v3], provide coordinates of the new vertices [b_i, b_(i+1), ...]  
+ *        inserted inside the triangle patch via barycentric coordinates [(u_i, v_i, w_i), (u_(i+1), v_(i+1), w_(i+1)), ...] 
+ *        such that b_i = u_i*v1 + v_i*v2 + w_i*v3.
+ *        
+ *      - Using the new vertices (along edges and inside the triangle), provide connectivity for a triangle patch by 
+ *        stitching those vertices together and creating triangles.
+ */
+class DYNAMICMESH_API FTessellationPattern
+{
+public:
+    
+    constexpr static int InvalidIndex = -1;
+
+protected:
+
+    const FDynamicMesh3* Mesh;
+
+public:
+
+    /**
+     * Represents an abstract edge as a parameterized line segment with two endpoints [v1,v2]. Contains information  
+     * about the linear interpolation parameter of the new vertices added along the edge and their vertex IDs.
+     * 
+     *      O----x----x----x----O           x  -  New vertices added. We only store their lerp coefficients (t1,t2,t3)  
+     *    v1     a1   a2   a3    v2               with respect to the end points, s.t. a_i = vi + t_i*(v2-v1).
+     */
+    struct EdgePatch 
+    {
+        int EdgeID = IndexConstants::InvalidID; 
+        TArrayView<int> VIDs;                   // Vertex IDs of the new vertices added.
+        TArrayView<double> LinearCoord;         // Lerp cooeficients of the new vertices added.
+        bool bIsReversed = false;               // If true, we need to traverse BaryCoord and VIDs in reverse order.    
+                                                // This happens when the directed edge of the triangle is reversed 
+                                                // with the respect to the edge direction stored in the FDynamicMesh3.
+    };
+
+    /**
+     * Represents an abstract triangle patch with corners [u,v,w], s.t. u = (1,0,0), v = (0,1,0), w = (0,0,1). 
+     * Contains information about the barycentric coordinates and IDs of the new vertices and the new triangles 
+     * generated during the tessellation. 
+     */
+    struct TrianglePatch 
+    {   
+        // ID of the triangle in the original mesh this patch represents
+        int TriangleID = IndexConstants::InvalidID;
+
+        // Vertex IDs of the 3 corners of the patch. These are the triangle vertex IDs in the original mesh.
+        FIndex3i UVWCorners = FIndex3i::Invalid();
+
+        // Edge patches
+        EdgePatch UVEdge;
+        EdgePatch VWEdge;
+        EdgePatch UWEdge;
+        
+        // Vertex IDs of the inner vertices added inside the triangle.
+        TArrayView<int> VIDs;
+
+        // Barycentric coordinates of the inner vertices.
+        TArrayView<FVector3d> BaryCoord;
+        
+        // Triangles the patch is split into
+        TArrayView<FIndex3i> Triangles;
+    };
+        
+    FTessellationPattern(const FDynamicMesh3* InMesh) 
+    : 
+    Mesh(InMesh)
+    {
+    }
+
+    virtual ~FTessellationPattern() 
+    {
+    }
+    
+    /** 
+     * @return the number of new vertices added along the edge patch. If InEdgeID is not valid then return 
+     * FTessellationPattern::InvalidIndex. 
+     */
+    virtual int GetNumberOfNewVerticesForEdgePatch(const int InEdgeID) const = 0;
+
+    /** 
+     * @return the number of new vertices added inside the triangle patch (i.e. exclude vertices added 
+     * along edges). If InTriangleID is not valid then return FTessellationPattern::InvalidIndex.
+     */
+    virtual int GetNumberOfNewVerticesForTrianglePatch(const int InTriangleID) const = 0;
+
+    /** 
+     * @return the number of triangles that will be generated after tessellating the patch. If InTriangleID is not 
+     * valid then return FTessellationPattern::InvalidIndex.
+     */
+    virtual int GetNumberOfPatchTriangles(const int InTriangleID) const = 0;
+
+    /** 
+     * Tessellate the edge patch by generating linear interpolation coefficients for each new vertex and storing the 
+     * result in the EdgePatch.LinearCoord.
+     * 
+     * @param EdgePatch Edge patch has EdgePatch.EdgeID and EdgePatch.VIDs fields set. The order of IDs is along the 
+     * edge from one corner to another. The edge direction is same as the direction stored by the FDynamicMesh3 class.
+     */
+    virtual void TessellateEdgePatch(EdgePatch& EdgePatch) const = 0;
+
+    /**
+     * Tessellate the triangle patch by generating new vertices inside the triangle. 
+     * Compute the barycentric coordinates of the new vertices and store them in the TriPatch.BaryCoord. 
+     * Using the vertices generated in the edge patches and inside the triangle patch, generate new triangles and 
+     * store the result in the TriPatch.Triangles. 
+     * 
+     * @param TriPatch Has all the fields set except the values in the TriPatch.BaryCoord and TriPatch.Triangles which 
+     *                 should be set by this method. Both arrays are already sized correctly to be able to hold all the 
+     *                 data needed as specified by the GetNumberOfNewVerticesForTrianglePatch() and GetNumberOfPatchTriangles().
+     */
+    virtual void TessellateTriPatch(TrianglePatch& TriPatch) const = 0;
+
+
+
+
+   //
+   // Common helper methods used during tessellation of various patterns. These might be helpful when writing custom 
+   // tessellation patterns.
+   //
+protected:
+
+    /**
+     * Takes a triangle defined by three vertices and a TessellationLevel value. For each corner of the outer triangle, 
+     * an inner triangle corner is produced at the intersection of two lines extended perpendicular to the corner's two 
+     * adjacent edges running through the vertex of the subdivided outer edge nearest to that corner.
+     *                            
+     *              U          
+     *              O             
+     *             / \            O - triangle corner
+     *            /   \           x - closest vertex to the corner inserted along the edge during tessellation
+     *           /     \          + - inner vertex resulting from intersection of two lines extened perpendicular to the edges
+     *          x       x         
+     *         /  .   .  \         
+     *        /     +     \       
+     *       /   innerU    \       
+     * 
+     * @note Vertices should be passed in CW order.
+     */
+    void ComputeInnerConcentricTriangle(const FVector3d& U,
+                                        const FVector3d& V,
+                                        const FVector3d& W,
+                                        const int EdgeTessellationLevel, 
+                                        FVector3d& InnerU,
+                                        FVector3d& InnerV,
+                                        FVector3d& InnerW) const;
+};
+
+
+
+
+/**
+ * Given an input mesh and a tessellation pattern, this operator generates a new tessellated mesh where the
+ * triangles affected are subtriangulated according to the rules specified by the pattern. Per-vertex normals, UVs, 
+ * colors, and extended per-vertex attributes are linearly interpolated to the new vertices. Per-triangle group 
+ * identifiers/materials and extended triangle attributes for the new triangles are inherited from the corresponding 
+ * input mesh triangles they replaced.
+ * 
+ * @note Currently does not support interpolation of the GenericAttributes besides Skin Weights.
+ */
+class DYNAMICMESH_API FAdaptiveTessellate
+{
+public:
+
+
+	//
+	// Inputs
+	//
+	
+	/** 
+     * Set this to be able to cancel running operation. 
+     * @note currently not implemented 
+     */
+	FProgressCancel* Progress = nullptr;
+
+	/** Should tessellation be multi-threaded. */
+	bool bUseParallel = true;
+
+    /** Tessellation pattern determines how we tessellate edges and triangles. */
+    FTessellationPattern* Pattern = nullptr;
+
+	//
+	// Input/Output
+	//
+
+	/** The tessellated mesh. */
+	FDynamicMesh3* ResultMesh = nullptr;
+	
+protected:
+	
+	/** 
+	 * Points to either the input mesh to be used to generate the new tessellated mesh or will point to a copy of the 
+     * mesh we are tessellating inplace.
+	 */
+	const FDynamicMesh3* Mesh = nullptr;
+
+	/** 
+	 * If true, the mesh to tessellate is contained in the ResultMesh and will be overwritten with the tessellation 
+     * result. 
+	 */
+	bool bInPlace = false;
+
+
+public:
+
+	/** 
+	 * Tessellate the mesh inplace. If the tessellation fails or the user cancelled the operation, OutMesh will not 
+	 * be changed.
+     * //TODO: Add an option to turn off making the backup copy of the input mesh in case the computation fails.
+	 */
+	FAdaptiveTessellate(FDynamicMesh3* OutMesh) 
+	:
+	ResultMesh(OutMesh), bInPlace(true)
+	{
+	}
+
+	/** 
+	 * Tessellate the mesh and write the result into another mesh. This will overwrite any data stored in the OutMesh. 
+	 */
+	FAdaptiveTessellate(const FDynamicMesh3* Mesh, FDynamicMesh3* OutMesh) 
+	:
+	ResultMesh(OutMesh), Mesh(Mesh), bInPlace(false)
+	{
+	}
+	
+	virtual ~FAdaptiveTessellate() 
+	{
+	}
+
+	/**
+	 * @return EOperationValidationResult::Ok if we can apply operation, or error code if we cannot.
+	 */
+	virtual EOperationValidationResult Validate()
+	{	
+		const bool bIsInvalid = bInPlace == false && (Mesh == nullptr || ResultMesh == nullptr || Pattern == nullptr);
+		const bool bIsInvalidInPlace = bInPlace == true && (ResultMesh == nullptr || Pattern == nullptr);
+		
+		if (bIsInvalid || bIsInvalidInPlace) 
+		{
+			return EOperationValidationResult::Failed_UnknownReason;
+		}
+
+		return EOperationValidationResult::Ok;
+	}
+
+    inline void SetPattern(FTessellationPattern* InPattern)
+    {
+        Pattern = InPattern;
+    }
+    
+	/**
+	 * Generate tessellated geometry.
+	 * @return true if the algorithm succeeds, false if it failed or was cancelled by the user.
+	 */
+	virtual bool Compute();
+
+
+
+
+    //
+    // Helper methods for setting up built-in tessellation patterns. Currently glsl, uniform and inner uniform 
+    // patterns are provided. You can create custom patterns by subclassing FTessellationPattern.
+    //
+public:
+
+
+    /**
+     * Pattern where the inner area is tessellated using the style of the OpenGL tessellation shader. The inner area
+     * of each tessellated triangle consists of multiple "rings". Areas between the rings are tessellated with the new 
+     * triangles.
+     * 
+     * You can specify per-edge and per-triangle tessellation levels. 
+     * Per-edge tessellation level determines how many new vertices are inserted along the new edge.
+     * Per-triangle tessellation level determines the number of new vertices along each of the edges of the first 
+     * inner ring. Each subsequent ring will have 2 fewer vertices.
+     * 
+     * Triangles with a zero inner tessellation level but with a non-zero edge tessellation level (common for those triangles 
+     * adjacent to the triangles marked for tessellation) will have a single vertex inserted in the middle and connected 
+     * to all edge vertices (triangle fan).
+     * 
+     *            
+     *                 O
+     *                / \
+     *               O. .O
+     *              /  O  \
+     *             O. / \ .O
+     *            /  O. .O  \
+     *           /  /  O  \  \
+     *          /  /  / \  \  \
+     *         O. /  /   \  \ .O
+     *        /  O. /     \ .O  \
+     *       /  /  O-------O  \  \
+     *      O. /   .       .   \ .O
+     *     /  O----O-------O----O  \
+     *    /   .    .       .    .   \
+     *   O----O----O-------O----O----O
+     *  
+     * @param InMesh Mesh we are tessellating.
+     * @param InEdgeTessLevels Array of InMesh->MaxEdgeID() size, where InEdgeTessLevels[EdgeID] is the number of new 
+     *                         vertices inserted along an EdgeID edge.
+     * @param InInnerTessLevels Array of InMesh->MaxTriangleID() size, where InInnerTessLevels[TriangleID] is the 
+     *                          number of new vertices along each of the edges of the first inner ring.
+     */
+    static TUniquePtr<FTessellationPattern> CreateConcentricRingsTessellationPattern(const FDynamicMesh3* InMesh,
+                                                                                     const TArray<int>& InEdgeTessLevels, 
+                                                                                     const TArray<int>& InInnerTessLevels);
+    
+    /** 
+     * Tessellate the whole mesh such that all the edge and inner tessellation levels are equal to InTessellationLevel. 
+     */
+    static TUniquePtr<FTessellationPattern> CreateConcentricRingsTessellationPattern(const FDynamicMesh3* InMesh,
+                                                                                     const int InTessellationLevel);
+
+    /** 
+     * Tessellate only selected triangles such that the edge and inner tessellation levels of those triangles are set to
+     * InTessellationLevel.
+     */
+    static TUniquePtr<FTessellationPattern> CreateConcentricRingsTessellationPattern(const FDynamicMesh3* InMesh,
+                                                                                     const int InTessellationLevel,
+                                                                                     const TArray<int>& InTriangleList);
+
+    /** 
+     * InEdgeFunc and InTriFunc functions will be evaluated on all edges and triangles respectively to determine their 
+     * desired tessellation levels. This can be helpful when the tessellation levels depend on some condition. For 
+     * example, how close/far the edge/triangle is from some point in space.
+     * 
+     * @note not yet implemented
+     */
+    static TUniquePtr<FTessellationPattern> CreateConcentricRingsTessellationPattern(const TFunctionRef<int(const int EdgeID)> InEdgeFunc, 
+                                                                                     const TFunctionRef<int(const int TriangleID)> InTriFunc);
+
+    
+    //TODO: add methods for automatic selection of triangles based on Material or Polygroup IDs.
+
+
+
+
+    /**
+     * Pattern where the inner area is tessellated using the uniform loop style subdivision. Each edge can have 
+     * different tessellation levels and edge vertices are connected to the inner uniformly tessellated sub-triangle.
+     * 
+     * Triangles with a zero inner tessellation level but with a non-zero edge tessellation level (common for those triangles 
+     * adjacent to the triangles marked for tessellation) will have a single vertex inserted in the middle and connected 
+     * to all edge vertices (triangle fan).
+     *  
+     *              o           
+     *             / \          
+     *            /   \         
+     *           o  o  o        
+     *          /  / \  \       
+     *         /  /   \  \      
+     *        /  o-----o  \     
+     *       o  / \   / \  o    
+     *      /  /   \ /   \  \   
+     *     /  o-----o-----o  \  
+     *    /                   \ 
+     *   o----------o----------o
+     *
+     * @note not yet implemented
+     */         
+
+    static TUniquePtr<FTessellationPattern> CreateInnerUnifromTessellationPattern(const FDynamicMesh3* InMesh,
+                                                                                  const TArray<int>& InEdgeTessLevels, 
+                                                                                  const TArray<int>& InInnerTessLevels);
+
+
+    //TODO: add more variation of different ways to provide tessellation levels similar to the ConcentricRings pattern.
+
+
+    /**
+     * Pattern where the whole triangle is tessellated using the loop style subdivision. Defined by a single 
+     * tessellation level, so you can not specify per-triangle or per-edge levels separately.
+     * 
+     * Triangles which are not marked for tessellation but which are adjacent to the triangles that are being tessellated 
+     * will have a single vertex inserted in the middle and connected to all edge vertices (triangle fan).
+     *   
+     *            O         
+     *           / \        
+     *          /   \       
+     *         O-----O      
+     *        / \   / \     
+     *       /   \ /   \    
+     *      O-----O-----O   
+     *     / \   / \   / \  
+     *    /   \ /   \ /   \ 
+     *   O-----O-----O-----O 
+     * 
+     * @note not yet implemented
+     */
+    
+    static TUniquePtr<FTessellationPattern> CreateUniformTessellationPattern(const FDynamicMesh3* InMesh,
+                                                                             const int InTessellationLevel,
+                                                                             const TArray<int>& InTriangleList);
+
+    //TODO: add more variation of different ways to provide tessellation levels similar to the ConcentricRings pattern.
+
+protected:
+
+	/** @return true if we need to abort the computation. */
+	virtual bool Cancelled();
+};
+
+} // end namespace UE::Geometry
+} // end namespace UE
