GeometryProcessing: fix bug in FDynamicMeshUVEditor::EstimateGeodesicCenterFrameVertex, it was not actually setting the results

diff --git a/Engine/Plugins/Runtime/GeometryProcessing/Source/DynamicMesh/Private/Parameterization/DynamicMeshUVEditor.cpp b/Engine/Plugins/Runtime/GeometryProcessing/Source/DynamicMesh/Private/Parameterization/DynamicMeshUVEditor.cpp
--- a/Engine/Plugins/Runtime/GeometryProcessing/Source/DynamicMesh/Private/Parameterization/DynamicMeshUVEditor.cpp
+++ b/Engine/Plugins/Runtime/GeometryProcessing/Source/DynamicMesh/Private/Parameterization/DynamicMeshUVEditor.cpp
@@ -350,7 +350,8 @@ bool FDynamicMeshUVEditor::EstimateGeodesicCenterFrameVertex(const FDynamicMesh3
 	{
 		return false;
 	}
-	FrameVertexID = SubmeshCalc.MapVertexToBaseMesh(FrameVertexID);
+	VertexIDOut = SubmeshCalc.MapVertexToBaseMesh(FrameVertexID);
+	FrameOut = SeedFrame;
 	return true;
 }
 
