Fixing a few unsafe typecasts

diff --git a/Engine/Plugins/Runtime/GeometryProcessing/Source/DynamicMesh/Private/Operations/FFDLattice.cpp b/Engine/Plugins/Runtime/GeometryProcessing/Source/DynamicMesh/Private/Operations/FFDLattice.cpp
--- a/Engine/Plugins/Runtime/GeometryProcessing/Source/DynamicMesh/Private/Operations/FFDLattice.cpp
+++ b/Engine/Plugins/Runtime/GeometryProcessing/Source/DynamicMesh/Private/Operations/FFDLattice.cpp
@@ -7,6 +7,57 @@
 
 using namespace UE::Geometry;
 
+namespace FFDLatticeHelpers
+{
+	double CubicBSplineKernel(double A)
+	{
+		// Using cubic kernel f(a) =
+		//	(4 - 6a^2 + 3|a|^3) / 6		for 0 <= |a| < 1
+		//	(2 - |a|)^3 / 6				for 1 <= |a| < 2
+		//	0							otherwise
+		//
+		// So at a = {-2, -1, 0, 1, 2}, f(a) = {0, 1/6, 4/6, 1/6, 0}, and is piecewise cubic in between.
+
+		double AbsA = FMath::Abs(A);
+		if (AbsA < 1.0f)
+		{
+			double ASquared = AbsA * AbsA;
+			return (4.0 - 6.0 * ASquared + 3.0 * ASquared * AbsA) / 6.0;
+		}
+		else if (AbsA < 2.0)
+		{
+			double TwoMinusAbsA = (2.0 - AbsA);
+			return TwoMinusAbsA * TwoMinusAbsA * TwoMinusAbsA / 6.0;
+		}
+
+		return 0.0;
+	}
+
+
+	double CubicBSplineKernelDerivative(double A)
+	{
+		// f'(a) =
+		//	a/2*(3|a| - 4)		            for 0 <= |a| < 1
+		//	-(a*(2-|a|)^2) / (2*|a|) 		for 1 <= |a| < 2
+		//	0								otherwise
+
+		double AbsA = FMath::Abs(A);
+		if (AbsA < 1.0)
+		{
+			return A / 2.0 * (3.0 * AbsA - 4.0);
+		}
+		else if (AbsA < 2.0)
+		{
+			double TwoMinusAbsA = (2.0 - AbsA);
+			return -(A * TwoMinusAbsA * TwoMinusAbsA) / (2.0 * AbsA);
+		}
+
+		return 0.0;
+	}
+
+} // namespace FFDLatticeHelpers
+
+
 FFFDLattice::FFFDLattice(const FVector3i& InDims, const FDynamicMesh3& Mesh, float Padding) :
 	Dimensions(InDims)
 {
@@ -276,57 +327,11 @@ void FFFDLattice::GetValuePair(int I, int J, int K, FVector3d& A, FVector3d& B,
 }
 
 
-static float CubicBSplineKernel(float A)
-{
-	// Using cubic kernel f(a) =
-	//	(4 - 6a^2 + 3|a|^3) / 6		for 0 <= |a| < 1
-	//	(2 - |a|)^3 / 6				for 1 <= |a| < 2
-	//	0							otherwise
-	//
-	// So at a = {-2, -1, 0, 1, 2}, f(a) = {0, 1/6, 4/6, 1/6, 0}, and is piecewise cubic in between.
-
-	float AbsA = FMath::Abs(A);
-	if (AbsA < 1.0f)
-	{
-		float ASquared = AbsA * AbsA;
-		return (4.0f - 6.0f * ASquared + 3.0f * ASquared * AbsA) / 6.0f;
-	}
-	else if (AbsA < 2.0f)
-	{
-		float TwoMinusAbsA = (2.0f - AbsA);
-		return TwoMinusAbsA * TwoMinusAbsA * TwoMinusAbsA / 6.0f;
-	}
-
-	return 0.0f;
-}
-
-
-static float CubicBSplineKernelDerivative(float A)
-{
-	// f'(a) =
-	//	a/2*(3|a| - 4)		            for 0 <= |a| < 1
-	//	-(a*(2-|a|)^2) / (2*|a|) 		for 1 <= |a| < 2
-	//	0								otherwise
-
-	float AbsA = FMath::Abs(A);
-	if (AbsA < 1.0f)
-	{
-		return A / 2.0f * (3.0f * AbsA - 4.0f);
-	}
-	else if (AbsA < 2.0f)
-	{
-		float TwoMinusAbsA = (2.0f - AbsA);
-		return -(A * TwoMinusAbsA * TwoMinusAbsA) / (2.0f * AbsA);
-	}
-
-	return 0.0f;
-}
-
-
-
 FVector3d FFFDLattice::InterpolatedPositionCubic(const FEmbedding& VertexEmbedding, 
 												 const TArray<FVector3d>& LatticeControlPoints) const
 {
+	using FFDLatticeHelpers::CubicBSplineKernel;
+
 	double T = VertexEmbedding.CellWeighting.X;
 	double U = VertexEmbedding.CellWeighting.Y;
 	double V = VertexEmbedding.CellWeighting.Z;
@@ -472,6 +477,9 @@ FMatrix3d FFFDLattice::LinearInterpolationJacobian(const FEmbedding& VertexEmbed
 
 FMatrix3d FFFDLattice::CubicInterpolationJacobian(const FEmbedding& VertexEmbedding, const TArray<FVector3d>& LatticeControlPoints) const
 {
+	using FFDLatticeHelpers::CubicBSplineKernel;
+	using FFDLatticeHelpers::CubicBSplineKernelDerivative;
+
 	// TODO: This was written for clarity and correctness. Could definitely be faster.
 
 	double T = VertexEmbedding.CellWeighting.X;
diff --git a/Engine/Plugins/Runtime/GeometryProcessing/Source/GeometryAlgorithms/Private/XAtlasWrapper.cpp b/Engine/Plugins/Runtime/GeometryProcessing/Source/GeometryAlgorithms/Private/XAtlasWrapper.cpp
--- a/Engine/Plugins/Runtime/GeometryProcessing/Source/GeometryAlgorithms/Private/XAtlasWrapper.cpp
+++ b/Engine/Plugins/Runtime/GeometryProcessing/Source/GeometryAlgorithms/Private/XAtlasWrapper.cpp
@@ -88,8 +88,8 @@ namespace XAtlasWrapper
 		UVIndexBuffer.Empty();
 		VertexRemapArray.Empty();
 
-		const float W = Atlas->width;
-		const float H = Atlas->height;
+		const double W = Atlas->width;
+		const double H = Atlas->height;
 
 		for (uint32 OutMeshID = 0; OutMeshID < Atlas->meshCount; ++OutMeshID)
 		{
diff --git a/Engine/Source/Runtime/GeometryCore/Public/Sampling/Gaussians.h b/Engine/Source/Runtime/GeometryCore/Public/Sampling/Gaussians.h
--- a/Engine/Source/Runtime/GeometryCore/Public/Sampling/Gaussians.h
+++ b/Engine/Source/Runtime/GeometryCore/Public/Sampling/Gaussians.h
@@ -109,8 +109,8 @@ public:
 		int32 N = KernelSize * KernelSize;
 		KernelOut.SetNum(N);
 
-		RealType InvTwoSigmaSqr = 1.0 / (2.0 * Sigma * Sigma);
-		RealType NormalizeFactor = (bNormalized) ? 1.0 : (TMathUtil<RealType>::InvPi * InvTwoSigmaSqr);
+		RealType InvTwoSigmaSqr = (RealType)1 / ((RealType)2 * Sigma * Sigma);
+		RealType NormalizeFactor = (bNormalized) ? (RealType)1 : (TMathUtil<RealType>::InvPi * InvTwoSigmaSqr);
 
 		FVector2i CenterIdx(IntRadius, IntRadius);
 		RealType KernelSum = 0;
@@ -137,7 +137,7 @@ public:
 	static void MakeKernelFromRadius(RealType Radius, TArray<RealType>& KernelOut, int32& IntRadiusOut, bool bNormalized = true)
 	{
 		IntRadiusOut = (int32)TMathUtil<RealType>::Ceil(Radius);
-		RealType Sigma = Radius / 2.0;			// ??
+		RealType Sigma = Radius / (RealType)2;			// ??
 		MakeKernelFromWidth(IntRadiusOut, Sigma, KernelOut, bNormalized);
 	}
 
@@ -145,7 +145,7 @@ public:
 	static void MakeKernelFromRadius(RealType Radius, TDiscreteKernel2<RealType>& KernelOut, bool bNormalized = true)
 	{
 		KernelOut.IntRadius = (int32)TMathUtil<RealType>::Ceil(Radius);
-		RealType Sigma = Radius / 2.0;			// ??
+		RealType Sigma = Radius / (RealType)2;			// ??
 		MakeKernelFromWidth(KernelOut.IntRadius, Sigma, KernelOut.Kernel, bNormalized);
 	}
 
