GeometryProcessing: handle an unexpected case in MeshBevel where edges that should be disconnected, aren't. Change the way that checkSlow/ensure is used so that they are not hit w/o changing the code, as they are problematic for geometry script bevel usage.

diff --git a/Engine/Plugins/Runtime/GeometryProcessing/Source/DynamicMesh/Private/Operations/MeshBevel.cpp b/Engine/Plugins/Runtime/GeometryProcessing/Source/DynamicMesh/Private/Operations/MeshBevel.cpp
--- a/Engine/Plugins/Runtime/GeometryProcessing/Source/DynamicMesh/Private/Operations/MeshBevel.cpp
+++ b/Engine/Plugins/Runtime/GeometryProcessing/Source/DynamicMesh/Private/Operations/MeshBevel.cpp
@@ -17,6 +17,21 @@
 
 using namespace UE::Geometry;
 
+// Uncomment to enable various checks and ensures in the Bevel code.
+// We cannot default-enable these checks even in Debug builds, because 
+// Bevel still frequently hits some of these cases, and those checks/ensures 
+// make it very difficult to debug a game using Bevel in Geometry Script (eg Lyra sample)
+//#define ENABLE_MESH_BEVEL_DEBUG
+
+#ifdef ENABLE_MESH_BEVEL_DEBUG
+#define MESH_BEVEL_DEBUG_CHECK(Expr) checkSlow(Expr)
+#define MESH_BEVEL_DEBUG_ENSURE(Expr) ensure(Expr)
+#else
+#define MESH_BEVEL_DEBUG_CHECK(Expr) 
+#define MESH_BEVEL_DEBUG_ENSURE(Expr) (Expr)
+#endif
+
+
 namespace UELocal
 {
 	static void QuadsToTris(const FDynamicMesh3& Mesh, const TArray<FIndex2i>& Quads, TArray<int32>& TrisOut, bool bReset)
@@ -354,7 +369,7 @@ void FMeshBevel::BuildVertexSets(const FDynamicMesh3& Mesh)
 		}
 
 
-		checkSlow(Vertex.IncomingBevelMeshEdges.Num() != 0);		// shouldn't ever happen
+		MESH_BEVEL_DEBUG_CHECK(Vertex.IncomingBevelMeshEdges.Num() != 0);		// shouldn't ever happen
 		if (Vertex.IncomingBevelMeshEdges.Num() == 1)
 		{
 			BuildTerminatorVertex(Vertex, Mesh);
@@ -408,7 +423,7 @@ void FMeshBevel::BuildJunctionVertex(FBevelVertex& Vertex, const FDynamicMesh3&
 		int32 CurTri = Vertex.SortedTriangles[CurTriIndex % NT];
 		int32 NextTri = Vertex.SortedTriangles[(CurTriIndex + 1) % NT];
 		int32 SharedEdge = FindSharedEdgeInTriangles(Mesh, CurTri, NextTri);
-		checkSlow(SharedEdge != -1);
+		MESH_BEVEL_DEBUG_CHECK(SharedEdge != -1);
 		if (Vertex.IncomingBevelMeshEdges.Contains(SharedEdge))
 		{
 			// if we found a bevel-edge, close the current wedge and start a new one
@@ -443,7 +458,7 @@ void FMeshBevel::BuildJunctionVertex(FBevelVertex& Vertex, const FDynamicMesh3&
 void FMeshBevel::BuildTerminatorVertex(FBevelVertex& Vertex, const FDynamicMesh3& Mesh)
 {
 	Vertex.VertexType = EBevelVertexType::Unknown;
-	if (ensure(Vertex.IncomingBevelMeshEdges.Num() == 1) == false)
+	if (MESH_BEVEL_DEBUG_ENSURE(Vertex.IncomingBevelMeshEdges.Num() == 1) == false)
 	{
 		return;
 	}
@@ -697,7 +712,7 @@ void FMeshBevel::UnlinkBevelEdgeInterior(
 		{
 			FDynamicMesh3::FVertexSplitInfo SplitInfo;
 			EMeshResult Result = Mesh.SplitVertex(Split.VertexID, Split.TriSets[0], SplitInfo);
-			if (ensure(Result == EMeshResult::Ok))
+			if (MESH_BEVEL_DEBUG_ENSURE(Result == EMeshResult::Ok))
 			{
 				BevelEdge.NewMeshVertices.Add(SplitInfo.NewVertex);
 				bDone = true;
@@ -711,13 +726,13 @@ void FMeshBevel::UnlinkBevelEdgeInterior(
 
 	// now build edge correspondences
 	N = BevelEdge.MeshVertices.Num();
-	checkSlow(N == BevelEdge.NewMeshVertices.Num());
+	MESH_BEVEL_DEBUG_CHECK(N == BevelEdge.NewMeshVertices.Num());
 	for (int32 k = 0; k < N-1; ++k)
 	{
 		int32 Edge0 = BevelEdge.MeshEdges[k];
 		int32 Edge1 = Mesh.FindEdge(BevelEdge.NewMeshVertices[k], BevelEdge.NewMeshVertices[k + 1]);
 		BevelEdge.NewMeshEdges.Add(Edge1);
-		checkSlow(Edge1 >= 0);
+		MESH_BEVEL_DEBUG_CHECK(Edge1 >= 0);
 		if ( Mesh.IsEdge(Edge1) && Edge0 != Edge1 && MeshEdgePairs.Contains(Edge0) == false )
 		{
 			MeshEdgePairs.Add(Edge0, Edge1);
@@ -772,7 +787,7 @@ void FMeshBevel::UnlinkBevelLoop(FDynamicMesh3& Mesh, FBevelLoop& BevelLoop, FDy
 		{
 			FDynamicMesh3::FVertexSplitInfo SplitInfo;
 			EMeshResult Result = Mesh.SplitVertex(Split.VertexID, Split.TriSets[1], SplitInfo);
-			if (Result == EMeshResult::Ok)
+			if (MESH_BEVEL_DEBUG_ENSURE(Result == EMeshResult::Ok))
 			{
 				BevelLoop.NewMeshVertices.Add(SplitInfo.NewVertex);
 				bDone = true;
@@ -786,13 +801,13 @@ void FMeshBevel::UnlinkBevelLoop(FDynamicMesh3& Mesh, FBevelLoop& BevelLoop, FDy
 
 	// now build edge correspondences
 	N = BevelLoop.MeshVertices.Num();
-	checkSlow(N == BevelLoop.NewMeshVertices.Num());
+	MESH_BEVEL_DEBUG_CHECK(N == BevelLoop.NewMeshVertices.Num());
 	for (int32 k = 0; k < N; ++k)
 	{
 		int32 Edge0 = BevelLoop.MeshEdges[k];
 		int32 Edge1 = Mesh.FindEdge(BevelLoop.NewMeshVertices[k], BevelLoop.NewMeshVertices[(k + 1)%N]);
 		BevelLoop.NewMeshEdges.Add(Edge1);
-		checkSlow(Edge1 >= 0);
+		MESH_BEVEL_DEBUG_CHECK(Edge1 >= 0);
 		if (Mesh.IsEdge(Edge1) && Edge0 != Edge1 && MeshEdgePairs.Contains(Edge0) == false )
 		{
 			MeshEdgePairs.Add(Edge0, Edge1);
@@ -837,7 +852,7 @@ void FMeshBevel::UnlinkVertices(FDynamicMesh3& Mesh, FDynamicMeshChangeTracker*
 
 void FMeshBevel::UnlinkJunctionVertex(FDynamicMesh3& Mesh, FBevelVertex& Vertex, FDynamicMeshChangeTracker* ChangeTracker)
 {
-	check(Vertex.VertexType == EBevelVertexType::JunctionVertex);
+	MESH_BEVEL_DEBUG_CHECK(Vertex.VertexType == EBevelVertexType::JunctionVertex);
 
 	if (ChangeTracker)
 	{
@@ -845,7 +860,7 @@ void FMeshBevel::UnlinkJunctionVertex(FDynamicMesh3& Mesh, FBevelVertex& Vertex,
 	}
 
 	int32 NumWedges = Vertex.Wedges.Num();
-	checkSlow(NumWedges > 1);
+	MESH_BEVEL_DEBUG_CHECK(NumWedges > 1);
 
 	// Split triangles around vertex into separate tri-sets based on wedges.
 	// This will create a new vertex for each wedge.
@@ -855,7 +870,7 @@ void FMeshBevel::UnlinkJunctionVertex(FDynamicMesh3& Mesh, FBevelVertex& Vertex,
 
 		FDynamicMesh3::FVertexSplitInfo SplitInfo;
 		EMeshResult Result = Mesh.SplitVertex(Vertex.VertexID, Wedge.Triangles, SplitInfo);
-		if (Result == EMeshResult::Ok)
+		if (MESH_BEVEL_DEBUG_ENSURE(Result == EMeshResult::Ok))
 		{
 			Wedge.WedgeVertex = SplitInfo.NewVertex;
 		}
@@ -891,8 +906,8 @@ void FMeshBevel::UnlinkJunctionVertex(FDynamicMesh3& Mesh, FBevelVertex& Vertex,
 
 void FMeshBevel::UnlinkTerminatorVertex(FDynamicMesh3& Mesh, FBevelVertex& BevelVertex, FDynamicMeshChangeTracker* ChangeTracker)
 {
-	check(BevelVertex.VertexType == EBevelVertexType::TerminatorVertex);
-	ensure(BevelVertex.Wedges.Num() == 2);
+	MESH_BEVEL_DEBUG_CHECK(BevelVertex.VertexType == EBevelVertexType::TerminatorVertex);
+	MESH_BEVEL_DEBUG_CHECK(BevelVertex.Wedges.Num() == 2);
 
 	if (ChangeTracker)
 	{
@@ -902,7 +917,7 @@ void FMeshBevel::UnlinkTerminatorVertex(FDynamicMesh3& Mesh, FBevelVertex& Bevel
 	// split the vertex
 	FDynamicMesh3::FVertexSplitInfo SplitInfo;
 	EMeshResult Result = Mesh.SplitVertex(BevelVertex.VertexID, BevelVertex.Wedges[1].Triangles, SplitInfo);
-	if (Result == EMeshResult::Ok)
+	if (MESH_BEVEL_DEBUG_ENSURE(Result == EMeshResult::Ok))
 	{
 		BevelVertex.Wedges[1].WedgeVertex = SplitInfo.NewVertex;
 
@@ -949,7 +964,7 @@ void FMeshBevel::FixUpUnlinkedBevelEdges(FDynamicMesh3& Mesh)
 		if (Edge.MeshEdges.Num() == 1)
 		{
 			int32* FoundOtherEdge = MeshEdgePairs.Find(Edge.MeshEdges[0]);
-			ensure(FoundOtherEdge != nullptr);
+			MESH_BEVEL_DEBUG_CHECK(FoundOtherEdge != nullptr);
 			if (FoundOtherEdge != nullptr)
 			{
 				Edge.NewMeshEdges[0] = *FoundOtherEdge;
@@ -977,14 +992,14 @@ void FMeshBevel::FixUpUnlinkedBevelEdges(FDynamicMesh3& Mesh)
 					FIndex3i TriEdges = Mesh.GetTriEdges(tid);
 					if (TriEdges.Contains(E0) && bFoundV0 == false)
 					{
-						checkSlow( Mesh.GetEdgeV(E0).Contains(Wedge.WedgeVertex) );
+						MESH_BEVEL_DEBUG_CHECK( Mesh.GetEdgeV(E0).Contains(Wedge.WedgeVertex) );
 						V0 = Wedge.WedgeVertex;
 						bFoundV0 = true;
 						break;
 					}
 					else if (TriEdges.Contains(E1) && bFoundV1 == false)
 					{
-						checkSlow( Mesh.GetEdgeV(E1).Contains(Wedge.WedgeVertex) );
+						MESH_BEVEL_DEBUG_CHECK( Mesh.GetEdgeV(E1).Contains(Wedge.WedgeVertex) );
 						V1 = Wedge.WedgeVertex;
 						bFoundV1 = true;
 						break;
@@ -1146,7 +1161,7 @@ void FMeshBevel::DisplaceVertices(FDynamicMesh3& Mesh, double Distance)
 
 void FMeshBevel::AppendJunctionVertexPolygon(FDynamicMesh3& Mesh, FBevelVertex& Vertex)
 {
-	check(Vertex.VertexType == EBevelVertexType::JunctionVertex);
+	MESH_BEVEL_DEBUG_CHECK(Vertex.VertexType == EBevelVertexType::JunctionVertex);
 
 	// UnlinkJunctionVertex() split the terminator vertex into N vertices, one for each
 	// (now disconnected) triangle-wedge. The wedges are ordered such that their wedge-vertices
@@ -1177,7 +1192,7 @@ void FMeshBevel::AppendJunctionVertexPolygon(FDynamicMesh3& Mesh, FBevelVertex&
 
 void FMeshBevel::AppendTerminatorVertexTriangle(FDynamicMesh3& Mesh, FBevelVertex& Vertex)
 {
-	check(Vertex.VertexType == EBevelVertexType::TerminatorVertex);
+	MESH_BEVEL_DEBUG_CHECK(Vertex.VertexType == EBevelVertexType::TerminatorVertex);
 
 	// UnlinkTerminatorVertex() opened up a triangle-shaped hole adjacent to the incoming edge quad-strip
 	// at the terminator vertex. The Wedges of the terminator vertex contain the vertex IDs of the two
@@ -1244,8 +1259,21 @@ void FMeshBevel::AppendEdgeQuads(FDynamicMesh3& Mesh, FBevelEdge& Edge)
 		{
 			FIndex2i EdgeV0 = Mesh.GetOrientedBoundaryEdgeV(EdgeID0);
 			FIndex2i EdgeV1 = Mesh.GetOrientedBoundaryEdgeV(EdgeID1);
-			QuadTris.A = Mesh.AppendTriangle(EdgeV0.B, EdgeV0.A, EdgeV1.B, Edge.NewGroupID);
-			QuadTris.B = Mesh.AppendTriangle(EdgeV1.B, EdgeV1.A, EdgeV0.B, Edge.NewGroupID);
+			if (EdgeV0.Contains(EdgeV1.A) || EdgeV0.Contains(EdgeV1.B))
+			{
+				// If we hit this case, it means that Edge0 and Edge1 are still connected at one end, and
+				// so cannot be connected by a Quad, they can only be connected by a single triangle.
+				// It is unclear how we end up in this situation, it does occur somewhat regularly in complex
+				// geometry scripts though (eg see UE-157531 for a potential repro).
+				int32 OtherV = EdgeV0.Contains(EdgeV1.A) ? EdgeV1.B : EdgeV1.A;
+				QuadTris.A = Mesh.AppendTriangle(EdgeV0.B, EdgeV0.A, OtherV, Edge.NewGroupID);
+				QuadTris.B = IndexConstants::InvalidID;
+			}
+			else
+			{
+				QuadTris.A = Mesh.AppendTriangle(EdgeV0.B, EdgeV0.A, EdgeV1.B, Edge.NewGroupID);
+				QuadTris.B = Mesh.AppendTriangle(EdgeV1.B, EdgeV1.A, EdgeV0.B, Edge.NewGroupID);
+			}
 		}
 		Edge.StripQuads.Add(QuadTris);
 	}
