Add Group Edge Welding: Add group edge welding by adding extra vertices to edge that doesn't have enough, then welding edge sequences

diff --git a/Engine/Plugins/Runtime/GeometryProcessing/Source/DynamicMesh/Private/Operations/WeldEdgeSequence.cpp b/Engine/Plugins/Runtime/GeometryProcessing/Source/DynamicMesh/Private/Operations/WeldEdgeSequence.cpp
new file mode 100644
--- /dev/null
+++ b/Engine/Plugins/Runtime/GeometryProcessing/Source/DynamicMesh/Private/Operations/WeldEdgeSequence.cpp
@@ -0,0 +1,221 @@
+// Copyright Epic Games, Inc. All Rights Reserved.
+
+#include "Operations/WeldEdgeSequence.h"
+
+#include "DynamicMesh/DynamicMesh3.h"
+#include "DynamicMesh/DynamicMeshAttributeSet.h"
+#include "DynamicMeshEditor.h"
+
+using namespace UE::Geometry;
+
+FWeldEdgeSequence::EWeldResult FWeldEdgeSequence::Weld()
+{
+	// As soon as any of these helper functions return a non-OK value,
+	// we want to forward that to the user and stop operating.
+
+	EWeldResult Result = EWeldResult::Ok;
+
+	if ((Result = CheckInput()) != EWeldResult::Ok)
+	{
+		return Result;
+	}
+
+	if ((Result = SplitSmallerSpan()) != EWeldResult::Ok)
+	{
+		return Result;
+	}
+
+	if ((Result = CheckForAndCollapseSideTriangles()) != EWeldResult::Ok)
+	{
+		return Result;
+	}
+
+	if ((Result = WeldEdgeSequence()) != EWeldResult::Ok)
+	{
+		return Result;
+	}
+
+	return Result;
+}
+
+
+
+/** Protected Functions */
+
+FWeldEdgeSequence::EWeldResult FWeldEdgeSequence::CheckInput()
+{
+	// Selected edges must be boundary edges
+	for (int Edge : EdgeSpanToDiscard.Edges)
+	{
+		if (Mesh->IsBoundaryEdge(Edge) == false)
+		{
+			return EWeldResult::Failed_EdgesNotBoundaryEdges;
+		}
+	}
+
+	for (int Edge : EdgeSpanToKeep.Edges)
+	{
+		if (Mesh->IsBoundaryEdge(Edge) == false)
+		{
+			return EWeldResult::Failed_EdgesNotBoundaryEdges;
+		}
+	}
+
+	// Ensure that the two input spans are oriented according to mesh boundary
+	// Guaranteed to be on boundary after two for loops above
+	EdgeSpanToDiscard.SetCorrectOrientation();
+	EdgeSpanToKeep.SetCorrectOrientation();
+
+	return EWeldResult::Ok;
+}
+
+FWeldEdgeSequence::EWeldResult FWeldEdgeSequence::SplitSmallerSpan()
+{
+	// For each new vertex that must be created:
+	// The longest simple edge is found and split
+	// The newly generated vertex is inserted into the span
+	// The newly generated edge is inserted into the span
+
+	FEdgeSpan& SpanToSplit = (EdgeSpanToDiscard.Vertices.Num() < EdgeSpanToKeep.Vertices.Num()) ? EdgeSpanToDiscard : EdgeSpanToKeep;
+	int TotalSplits = FMath::Abs(EdgeSpanToDiscard.Vertices.Num() - EdgeSpanToKeep.Vertices.Num());
+	for (int SplitCount = 0; SplitCount < TotalSplits; ++SplitCount)
+	{
+		double MaxLength = 0.0;
+		int LongestEID = -1;
+		int LongestIndex = -1;
+
+		// Find longest edge and store length, ID, and index
+		for (int EdgeIndex = 0; EdgeIndex < SpanToSplit.Edges.Num(); ++EdgeIndex)
+		{
+			FVector3d VertA = Mesh->GetVertex(Mesh->GetEdge(SpanToSplit.Edges[EdgeIndex]).Vert.A);
+			FVector3d VertB = Mesh->GetVertex(Mesh->GetEdge(SpanToSplit.Edges[EdgeIndex]).Vert.B);
+			double EdgeLength = DistanceSquared(VertA, VertB);
+
+			if (MaxLength < EdgeLength)
+			{
+				MaxLength = EdgeLength;
+				LongestEID = SpanToSplit.Edges[EdgeIndex];
+				LongestIndex = EdgeIndex;
+			}
+		}
+
+		// Split longest edge
+		FDynamicMesh3::FEdgeSplitInfo SplitInfo;
+		EMeshResult Result = Mesh->SplitEdge(LongestEID, SplitInfo);
+		if (Result != EMeshResult::Ok)
+		{
+			return EWeldResult::Failed_CannotSplitEdge;
+		}
+
+		// Correctly insert new vertex (between vertices of split edge)
+		SpanToSplit.Vertices.Insert(SplitInfo.NewVertex, LongestIndex + 1);
+
+		// Correctly insert new edge
+		// OriginalVertices.B is the non-new vertex of the newly inserted edge- use this
+		// to determine whether the edge goes before or after the original in our span
+		if (SplitInfo.OriginalVertices.B == SpanToSplit.Vertices[LongestIndex])
+		{
+			SpanToSplit.Edges.Insert(SplitInfo.NewEdges.A, LongestIndex);
+		}
+		else
+		{
+			SpanToSplit.Edges.Insert(SplitInfo.NewEdges.A, LongestIndex + 1);
+		}
+	}
+
+	return EWeldResult::Ok;
+}
+
+FWeldEdgeSequence::EWeldResult FWeldEdgeSequence::CheckForAndCollapseSideTriangles()
+{
+	// Checks for and deletes edge between VertA and VertB
+	auto CheckForAndHandleEdge = [this](int VertA, int VertB)
+	{
+		int Edge = Mesh->FindEdge(VertA, VertB);
+		if (Edge != IndexConstants::InvalidID)
+		{
+			if (bAllowIntermediateTriangleDeletion == false)
+			{
+				return EWeldResult::Failed_TriangleDeletionDisabled;
+			}
+
+			FIndex2i TrianglePair = Mesh->GetEdgeT(Edge);
+			EMeshResult Result = Mesh->RemoveTriangle(TrianglePair.A);
+			if (Result != EMeshResult::Ok)
+			{
+				return EWeldResult::Failed_CannotDeleteTriangle;
+			}
+
+			if (Mesh->IsTriangle(TrianglePair.B))
+			{
+				Result = Mesh->RemoveTriangle(TrianglePair.B);
+				if (Result != EMeshResult::Ok)
+				{
+					return EWeldResult::Failed_CannotDeleteTriangle;
+				}
+			}
+		}
+
+		return EWeldResult::Ok;
+	};
+
+	EWeldResult Result = CheckForAndHandleEdge(EdgeSpanToDiscard.Vertices[0], EdgeSpanToKeep.Vertices.Last());
+	if (Result != EWeldResult::Ok)
+	{
+		return Result;
+	}
+	
+	Result = CheckForAndHandleEdge(EdgeSpanToDiscard.Vertices.Last(), EdgeSpanToKeep.Vertices[0]);
+	if (Result != EWeldResult::Ok)
+	{
+		return Result;
+	}
+
+	return EWeldResult::Ok;
+}
+
+FWeldEdgeSequence::EWeldResult FWeldEdgeSequence::WeldEdgeSequence()
+{
+	if (!ensure(EdgeSpanToDiscard.Edges.Num() == EdgeSpanToKeep.Edges.Num()))
+	{
+		return EWeldResult::Failed_Other;
+	}
+
+	for (int EdgePairIndex = 0; EdgePairIndex < EdgeSpanToDiscard.Edges.Num(); ++EdgePairIndex)
+	{
+		int EdgeA = EdgeSpanToDiscard.Edges[EdgeSpanToDiscard.Edges.Num() - 1 - EdgePairIndex];
+		int EdgeB = EdgeSpanToKeep.Edges[EdgePairIndex];
+
+		// We need to check this because MergeEdges() can implicitly
+		// merge adjacent edges when adjacent edges share a vertex.
+		if (Mesh->IsEdge(EdgeA) && Mesh->IsEdge(EdgeB))
+		{
+			FDynamicMesh3::FMergeEdgesInfo MergeInfo;
+			EMeshResult Result = Mesh->MergeEdges(EdgeB, EdgeA, MergeInfo);
+			if (Result != EMeshResult::Ok)
+			{
+				if (Result == EMeshResult::Failed_InvalidNeighbourhood && bAllowFailedMerge == true)
+				{
+					// Get properly oriented vertex pairs for edges which can't be merged
+					FIndex2i EdgeAVerts = Mesh->GetEdgeV(EdgeA);
+					FIndex2i EdgeBVerts = Mesh->GetEdgeV(EdgeB);
+					IndexUtil::OrientTriEdge(EdgeAVerts.A, EdgeAVerts.B, Mesh->GetTriangle(Mesh->GetEdgeT(EdgeA).A));
+					IndexUtil::OrientTriEdge(EdgeBVerts.A, EdgeBVerts.B, Mesh->GetTriangle(Mesh->GetEdgeT(EdgeB).A));
+
+					// Move verts from EdgeA to EdgeB, leaving an invisible seam
+					Mesh->SetVertex(EdgeAVerts.A, Mesh->GetVertex(EdgeBVerts.B));
+					Mesh->SetVertex(EdgeAVerts.B, Mesh->GetVertex(EdgeBVerts.A));
+
+					// Add eids to UnmergedEdgePairsOut so user can easily identify offending edges
+					UnmergedEdgePairsOut.Add(TPair<int, int>(EdgeA, EdgeB));
+				}
+				else
+				{
+					return EWeldResult::Failed_Other;
+				}
+			}
+		}
+	}
+
+	return EWeldResult::Ok;
+}
diff --git a/Engine/Plugins/Runtime/GeometryProcessing/Source/DynamicMesh/Public/Operations/WeldEdgeSequence.h b/Engine/Plugins/Runtime/GeometryProcessing/Source/DynamicMesh/Public/Operations/WeldEdgeSequence.h
new file mode 100644
--- /dev/null
+++ b/Engine/Plugins/Runtime/GeometryProcessing/Source/DynamicMesh/Public/Operations/WeldEdgeSequence.h
@@ -0,0 +1,109 @@
+// Copyright Epic Games, Inc. All Rights Reserved.
+#pragma once
+
+#include "DynamicMesh/Public/EdgeSpan.h"
+
+namespace UE
+{
+namespace Geometry
+{
+
+class FDynamicMesh3;
+
+/**
+* Weld a pair of group edges.
+* 
+* User can optionally allow triangle deletion which handles cases
+* where the group edges are connected by an edge at the end points.
+*/
+
+class DYNAMICMESH_API FWeldEdgeSequence
+{
+public:
+	enum class EWeldResult
+	{
+		Ok								= 0,
+
+		Failed_EdgesNotBoundaryEdges	= 10,		// Occurs when any edge in either input span isn't a boundary edge.
+
+		Failed_CannotSplitEdge			= 21,		// Occurs when SplitEdge() fails, haven't encountered this and I'm not sure what would cause it
+		Failed_TriangleDeletionDisabled = 22,		// Occurs when bAllowIntermediateTriangleDeletion is false and edge spans are connected by an edge
+		Failed_CannotDeleteTriangle		= 23,		// Occurs when bAllowIntermediateTriangleDeletion is true, edge spans are connected, but edge deletion fails
+
+		Failed_Other					= 100,		// Catch all for general failure
+	};
+
+public:
+	/**
+	* Inputs/Outputs
+	*/
+	FDynamicMesh3* Mesh;
+	FEdgeSpan EdgeSpanToDiscard;	// This data is junk once Weld() is called
+	FEdgeSpan EdgeSpanToKeep;		// This is the updated edge span which can be used once Weld() is called
+
+	// Whether triangle deletion is allowed in order to merge edges which are connected by a different edge
+	// This specifically determines how CheckForAndCollapseSideTriangles() behaves
+	bool bAllowIntermediateTriangleDeletion = false;
+
+	// When true, failed calls to MergeEdges() will be handled by moving the edges without merging such
+	// that the final result appears to be welded but has invisible seam(s) instead of just failing.
+	bool bAllowFailedMerge = false;
+
+	// This is populated with pairs of eids which were not able to be merged.
+	// Only valid when bAllowFailedMerge is true
+	TArray<TPair<int, int>> UnmergedEdgePairsOut;
+
+public:
+	FWeldEdgeSequence(FDynamicMesh3* Mesh, FEdgeSpan SpanDiscard, FEdgeSpan SpanKeep)
+		: Mesh(Mesh)
+		, EdgeSpanToDiscard(SpanDiscard)
+		, EdgeSpanToKeep(SpanKeep)
+	{
+	}
+	virtual ~FWeldEdgeSequence() {}
+
+
+	/**
+	* Alters the existing mesh by welding two edge sequences, preserving sequence A.
+	* Conditions the mesh by splitting edges and optionally deleting triangles.
+	*
+	* @return EWeldResult::OK on success
+	*/
+	EWeldResult Weld();
+
+protected:
+	/**
+	* Verifies validity of input edges by ensuring they are
+	* correctly-oriented boundary edges
+	*
+	* @return EWeldResult::OK on success
+	*/
+	EWeldResult CheckInput();
+
+	/**
+	* Splits largest edges of the span with fewest vertices so that
+	* both input spans have an equal number of vertices and edges after 
+	* 
+	* @return EWeldResult::OK on success
+	*/
+	EWeldResult SplitSmallerSpan();
+
+	/**
+	* Checks for edges between terminating vertices of the input spans and attempts to delete them if possible
+	* by deleting the adjacent triangle. These cases would prevent MergeEdges() from succeeding because it won't
+	* implicitly collapse/delete triangles which connect the two edges being merged.
+	* 
+	* @return EWeldResult::OK on success
+	*/
+	EWeldResult CheckForAndCollapseSideTriangles();
+
+	/**
+	* Welds edge sequence together
+	*
+	* @return EWeldResult::OK on success
+	*/
+	EWeldResult WeldEdgeSequence();
+};
+
+} // end namespace UE::Geometry
+} // end namespace UE
\ No newline at end of file
diff --git a/Engine/Plugins/Runtime/MeshModelingToolset/Source/MeshModelingTools/Private/EditMeshPolygonsTool.cpp b/Engine/Plugins/Runtime/MeshModelingToolset/Source/MeshModelingTools/Private/EditMeshPolygonsTool.cpp
--- a/Engine/Plugins/Runtime/MeshModelingToolset/Source/MeshModelingTools/Private/EditMeshPolygonsTool.cpp
+++ b/Engine/Plugins/Runtime/MeshModelingToolset/Source/MeshModelingTools/Private/EditMeshPolygonsTool.cpp
@@ -28,6 +28,7 @@
 #include "Operations/SimpleHoleFiller.h"
 #include "Operations/MinimalHoleFiller.h"
 #include "Operations/PolygroupRemesh.h"
+#include "Operations/WeldEdgeSequence.h"
 #include "Selection/PersistentMeshSelection.h"
 #include "Selection/StoredMeshSelectionUtil.h"
 #include "Selection/PolygonSelectionMechanic.h"
@@ -1655,45 +1656,110 @@ void UEditMeshPolygonsTool::ApplyCollapseEdge()
 
 
 
-
 void UEditMeshPolygonsTool::ApplyWeldEdges()
 {
-	bool bValidInput = SelectionMechanic->GetActiveSelection().SelectedEdgeIDs.Num() == 2 && BeginMeshBoundaryEdgeEditChange(true);
-	bValidInput = bValidInput && ActiveEdgeSelection.Num() == 2;		// one of the initial edges may not have been valid
-	if ( bValidInput == false )
+	if (SelectionMechanic->GetActiveSelection().SelectedEdgeIDs.Num() != 2)
 	{
-		GetToolManager()->DisplayMessage( LOCTEXT("OnWeldEdgesFailed", "Cannot Weld current selection"), EToolMessageLevel::UserWarning);
+		GetToolManager()->DisplayMessage(
+			LOCTEXT("OnWeldEdgesFailedEdgeCount", "Cannot Weld current selection, selection must be exactly 2 edges."),
+			EToolMessageLevel::UserWarning);
 		return;
 	}
 
 	FDynamicMesh3* Mesh = CurrentMesh.Get();
+	FDynamicMesh3 MeshCopy(*Mesh);				// We are going to operate on this copy and, if successful, copy the changes back into Mesh
+	
+	FGroupTopologySelection CurrentSelection = SelectionMechanic->GetActiveSelection();
+	TArray<int32> SelectedEdgeIDs = CurrentSelection.SelectedEdgeIDs.Array();
+	FEdgeSpan& SpanA = Topology->Edges[SelectedEdgeIDs[0]].Span;
+	FEdgeSpan& SpanB = Topology->Edges[SelectedEdgeIDs[1]].Span;
+	
+	if (SpanA.Vertices[0] == SpanA.Vertices.Last() || SpanB.Vertices[0] == SpanB.Vertices.Last())
+	{
+		GetToolManager()->DisplayMessage(
+			LOCTEXT("OnWeldEdgesFailedEdgesAreLoops", "Cannot Weld current selection, selected edges must not be loops."),
+			EToolMessageLevel::UserWarning);
+		return;
+	}
 
 	FDynamicMeshChangeTracker ChangeTracker(Mesh);
 	ChangeTracker.BeginChange();
 
-	int32 EdgeIDA = Topology->GetGroupEdgeEdges(ActiveEdgeSelection[0].EdgeTopoID)[0];
-	int32 EdgeIDB = Topology->GetGroupEdgeEdges(ActiveEdgeSelection[1].EdgeTopoID)[0];
-	FIndex2i EdgeVerts[2] = { Mesh->GetEdgeV(EdgeIDA), Mesh->GetEdgeV(EdgeIDB) };
-	for (int j = 0; j < 2; ++j)
+	// Save one ring tri's for vertices along first edge
+	for (int Vert : SpanA.Vertices)
 	{
-		ChangeTracker.SaveVertexOneRingTriangles(EdgeVerts[j].A, true);
-		ChangeTracker.SaveVertexOneRingTriangles(EdgeVerts[j].B, true);
+		ChangeTracker.SaveVertexOneRingTriangles(Vert, true);
 	}
 
-	FDynamicMesh3::FMergeEdgesInfo MergeInfo;
-	EMeshResult Result = Mesh->MergeEdges(EdgeIDB, EdgeIDA, MergeInfo);
-	if (Result != EMeshResult::Ok)
+	// Save one ring tri's for vertices along second edge
+	for (int Vert : SpanB.Vertices)
 	{
-		GetToolManager()->DisplayMessage( LOCTEXT("OnWeldEdgesFailed", "Cannot Weld current selection"), EToolMessageLevel::UserWarning);
+		ChangeTracker.SaveVertexOneRingTriangles(Vert, true);
+	}
+
+	FWeldEdgeSequence EdgeWelder(&MeshCopy, SpanA, SpanB);
+	EdgeWelder.bAllowIntermediateTriangleDeletion = true;
+	EdgeWelder.bAllowFailedMerge = true;
+	
+	FWeldEdgeSequence::EWeldResult Result = EdgeWelder.Weld();
+	if (Result != FWeldEdgeSequence::EWeldResult::Ok)
+	{
+		switch (Result)
+		{
+		case FWeldEdgeSequence::EWeldResult::Failed_EdgesNotBoundaryEdges:
+			GetToolManager()->DisplayMessage(
+				LOCTEXT("OnWeldEdgesFailedBoundary", "Cannot Weld current selection, selected edges must be boundary edges."),
+				EToolMessageLevel::UserWarning);
+			break;
+
+		case FWeldEdgeSequence::EWeldResult::Failed_CannotSplitEdge:
+			GetToolManager()->DisplayMessage(
+				LOCTEXT("OnWeldEdgesFailedSplitEdge", "Cannot Weld current selection, failed to insert vertex."), 
+				EToolMessageLevel::UserWarning);
+			break;
+
+		case FWeldEdgeSequence::EWeldResult::Failed_TriangleDeletionDisabled:
+			GetToolManager()->DisplayMessage(
+				LOCTEXT("OnWeldEdgesFailedTriDeleteDisabled", "Cannot Weld current selection, deletion of edges connecting selected edges is disabled."), 
+				EToolMessageLevel::UserWarning);
+			break;
+
+		case FWeldEdgeSequence::EWeldResult::Failed_CannotDeleteTriangle:
+			GetToolManager()->DisplayMessage(
+				LOCTEXT("OnWeldEdgesFailedTriDeleteFailed", "Cannot Weld current selection, failed to delete edge connecting selected edges."), 
+				EToolMessageLevel::UserWarning);
+			break;
+
+		case FWeldEdgeSequence::EWeldResult::Failed_Other:
+		default:
+			GetToolManager()->DisplayMessage(
+				LOCTEXT("OnWeldEdgesFailedOther", "Cannot Weld current selection, bad geometry."), 
+				EToolMessageLevel::UserWarning);
+			break;
+		}
+
 		return;
 	}
+	else
+	{
+		// On success, apply the result by copying over the existing mesh
+		*Mesh = MeshCopy;
+
+		if (EdgeWelder.UnmergedEdgePairsOut.Num() != 0)
+		{
+			GetToolManager()->DisplayMessage(
+				LOCTEXT("OnWeldEdgesCompletedSeamsRemain", "Warning: welding incomplete because it would create "
+					"invalid geometry (attached non manifold edge or duplicate triangle). Seam still exists at weld "
+					"location. Modify attached triangles and retry, or undo."),
+				EToolMessageLevel::UserWarning);
+		}
+	}
 
 	FGroupTopologySelection NewSelection;
 	EmitCurrentMeshChangeAndUpdate(LOCTEXT("PolyMeshWeldEdgeChange", "Weld Edges"),
 		ChangeTracker.EndChange(), NewSelection);
 }
 
-
 void UEditMeshPolygonsTool::ApplyStraightenEdges()
 {
 	if (BeginMeshEdgeEditChange() == false)
