GeometryProcessing: ShapeApproximation, fixing implicit type conversion warnings.

diff --git a/Engine/Plugins/Runtime/GeometryProcessing/Source/DynamicMesh/Private/ShapeApproximation/MeshSimpleShapeApproximation.cpp b/Engine/Plugins/Runtime/GeometryProcessing/Source/DynamicMesh/Private/ShapeApproximation/MeshSimpleShapeApproximation.cpp
--- a/Engine/Plugins/Runtime/GeometryProcessing/Source/DynamicMesh/Private/ShapeApproximation/MeshSimpleShapeApproximation.cpp
+++ b/Engine/Plugins/Runtime/GeometryProcessing/Source/DynamicMesh/Private/ShapeApproximation/MeshSimpleShapeApproximation.cpp
@@ -325,8 +325,8 @@ void FMeshSimpleShapeApproximation::Generate_ConvexHullDecompositions(FSimpleSha
 		const FDynamicMesh3& SourceMesh = *SourceMeshes[idx];
 		// TODO: if (bSimplifyHulls), also consider simplifying the input?
 		FConvexDecomposition3 Decomposition(SourceMesh);
-
-		Decomposition.Compute(ConvexDecompositionMaxPieces, ConvexDecompositionMaxPieces * ConvexDecompositionSearchFactor, ConvexDecompositionErrorTolerance, ConvexDecompositionMinPartThickness);
+		const int32 NumAdditionalSplits = FMath::FloorToInt32(ConvexDecompositionMaxPieces * ConvexDecompositionSearchFactor);
+		Decomposition.Compute(ConvexDecompositionMaxPieces, NumAdditionalSplits, ConvexDecompositionErrorTolerance, ConvexDecompositionMinPartThickness);
 
 		for (int32 HullIdx = 0; HullIdx < Decomposition.NumHulls(); HullIdx++)
 		{
@@ -431,9 +431,9 @@ void FMeshSimpleShapeApproximation::Generate_LevelSets(FSimpleShapeSet3d& ShapeS
 		SDF.Mesh = SourceMeshes[MeshIndex];
 		SDF.Spatial = &Spatial;
 		SDF.ComputeMode = TSweepingMeshSDF<FDynamicMesh3>::EComputeModes::NarrowBand_SpatialFloodFill;
-		SDF.CellSize = CellSize;
+		SDF.CellSize = (float)CellSize;
 		SDF.NarrowBandMaxDistance = 2.0 * CellSize;
-		SDF.ExactBandWidth = FMath::CeilToInt(SDF.NarrowBandMaxDistance / CellSize);
+		SDF.ExactBandWidth = FMath::CeilToInt32(SDF.NarrowBandMaxDistance / CellSize);
 		SDF.ExpandBounds = 2.0 * CellSize * FVector3d::One();
 
 		if (SDF.Compute(Bounds))
diff --git a/Engine/Plugins/Runtime/GeometryProcessing/Source/DynamicMesh/Private/ShapeApproximation/ShapeDetection3.cpp b/Engine/Plugins/Runtime/GeometryProcessing/Source/DynamicMesh/Private/ShapeApproximation/ShapeDetection3.cpp
--- a/Engine/Plugins/Runtime/GeometryProcessing/Source/DynamicMesh/Private/ShapeApproximation/ShapeDetection3.cpp
+++ b/Engine/Plugins/Runtime/GeometryProcessing/Source/DynamicMesh/Private/ShapeApproximation/ShapeDetection3.cpp
@@ -5,6 +5,7 @@
 #include "DynamicMesh/DynamicMesh3.h"
 #include "MeshQueries.h"
 #include "FitCapsule3.h"
+#include "Math/UnrealMathUtility.h"
 
 using namespace UE::Geometry;
 
@@ -109,7 +110,7 @@ bool UE::Geometry::IsBoxMesh(const FDynamicMesh3& Mesh, FOrientedBox3d& BoxOut,
 	// that are not ~0 or ~90 degrees, if so this is not a box. This will immediately reject 
 	// most non-box meshes without having to do expensive normal clustering first
 	int32 NumEdges = Mesh.EdgeCount();
-	int32 Step = FMathd::Max(Mesh.EdgeCount() / 10, 2);
+	int32 Step = FMath::Max(Mesh.EdgeCount() / 10, 2);
 	for (int32 k = 0; k < NumEdges; k += Step)
 	{
 		if (Mesh.IsEdge(k))
@@ -170,7 +171,7 @@ bool UE::Geometry::IsBoxMesh(const FDynamicMesh3& Mesh, FOrientedBox3d& BoxOut,
 		// accumulate unique axis if we haven't seen this one and we haven't found 3 already
 		if (UniqueCount < 3)
 		{
-			int32 UniqueIdx = (ParallelPair == -1) ? k : FMathd::Min(k, ParallelPair);
+			int32 UniqueIdx = (ParallelPair == -1) ? k : FMath::Min(k, ParallelPair);
 			UniqueAxes[UniqueCount++] = UniqueIdx;
 		}
 	}
diff --git a/Engine/Plugins/Runtime/GeometryProcessing/Source/DynamicMesh/Public/ShapeApproximation/SimpleShapeSet3.h b/Engine/Plugins/Runtime/GeometryProcessing/Source/DynamicMesh/Public/ShapeApproximation/SimpleShapeSet3.h
--- a/Engine/Plugins/Runtime/GeometryProcessing/Source/DynamicMesh/Public/ShapeApproximation/SimpleShapeSet3.h
+++ b/Engine/Plugins/Runtime/GeometryProcessing/Source/DynamicMesh/Public/ShapeApproximation/SimpleShapeSet3.h
@@ -138,7 +138,8 @@ struct DYNAMICMESH_API FLevelSetShape3d
 
 	FOrientedBox3d GetGridBox() const
 	{
-		const FVector3d BoxExtents = FVector3d(0.5 * CellSize * Grid.GetDimensions());
+		const UE::Geometry::FVector3i& Dims = Grid.GetDimensions();
+		const FVector3d BoxExtents = (0.5 * CellSize) * FVector3d(Dims[0], Dims[1], Dims[2]);
 		const FFrame3d BoxFrame(GridTransform.GetTranslation(), GridTransform.GetRotation());
 		return FOrientedBox3d(BoxFrame, BoxExtents);
 	}
