GeometryProcessing: Sampling - fix implicit type conversion warnings.

diff --git a/Engine/Plugins/Runtime/GeometryProcessing/Source/DynamicMesh/Private/Sampling/MeshCurvatureMapBaker.cpp b/Engine/Plugins/Runtime/GeometryProcessing/Source/DynamicMesh/Private/Sampling/MeshCurvatureMapBaker.cpp
--- a/Engine/Plugins/Runtime/GeometryProcessing/Source/DynamicMesh/Private/Sampling/MeshCurvatureMapBaker.cpp
+++ b/Engine/Plugins/Runtime/GeometryProcessing/Source/DynamicMesh/Private/Sampling/MeshCurvatureMapBaker.cpp
@@ -50,7 +50,7 @@ void FMeshCurvatureMapBaker::Bake()
 	if (BlurRadius > 0.01)
 	{
 		TDiscreteKernel2f BlurKernel2d;
-		TGaussian2f::MakeKernelFromRadius(BlurRadius, BlurKernel2d);
+		TGaussian2f::MakeKernelFromRadius((float)BlurRadius, BlurKernel2d);
 		TArray<FVector3f> AOBlurBuffer;
 		Occupancy.ParallelProcessingPass<FVector3f>(
 			[&](int64 Index) { return FVector3f::Zero(); },
diff --git a/Engine/Plugins/Runtime/GeometryProcessing/Source/DynamicMesh/Private/Sampling/MeshCurvatureMapEvaluator.cpp b/Engine/Plugins/Runtime/GeometryProcessing/Source/DynamicMesh/Private/Sampling/MeshCurvatureMapEvaluator.cpp
--- a/Engine/Plugins/Runtime/GeometryProcessing/Source/DynamicMesh/Private/Sampling/MeshCurvatureMapEvaluator.cpp
+++ b/Engine/Plugins/Runtime/GeometryProcessing/Source/DynamicMesh/Private/Sampling/MeshCurvatureMapEvaluator.cpp
@@ -97,7 +97,7 @@ void FMeshCurvatureMapEvaluator::EvaluateColor(const int DataIdx, float*& In, FV
 	const FMeshCurvatureMapEvaluator* Eval = static_cast<FMeshCurvatureMapEvaluator*>(EvalData);
 	const float Curvature = In[0];
 	const float Sign = FMathf::Sign(Curvature);
-	const float T = Eval->ClampRange.GetT(FMathf::Abs(Curvature));
+	const float T = (float)Eval->ClampRange.GetT(FMathf::Abs(Curvature));
 	FVector3f CurvatureColor = (Sign < 0) ?
 		Lerp(Eval->ZeroColor, Eval->NegativeColor, T) : Lerp(Eval->ZeroColor, Eval->PositiveColor, T);
 	Out = FVector4f(CurvatureColor[0], CurvatureColor[1], CurvatureColor[2], 1.0f);
diff --git a/Engine/Plugins/Runtime/GeometryProcessing/Source/DynamicMesh/Private/Sampling/MeshMapBaker.cpp b/Engine/Plugins/Runtime/GeometryProcessing/Source/DynamicMesh/Private/Sampling/MeshMapBaker.cpp
--- a/Engine/Plugins/Runtime/GeometryProcessing/Source/DynamicMesh/Private/Sampling/MeshMapBaker.cpp
+++ b/Engine/Plugins/Runtime/GeometryProcessing/Source/DynamicMesh/Private/Sampling/MeshMapBaker.cpp
@@ -525,7 +525,7 @@ void FMeshMapBaker::Bake()
 				return; // ParallelFor
 			}
 
-			const int NumGutter = GutterTexelsPerTile[TileIdx].Num();
+			const int64 NumGutter = GutterTexelsPerTile[TileIdx].Num();
 			for (int64 GutterIdx = 0; GutterIdx < NumGutter; ++GutterIdx)
 			{
 				int64 GutterPixelTo;
diff --git a/Engine/Plugins/Runtime/GeometryProcessing/Source/DynamicMesh/Private/Sampling/MeshOcclusionMapBaker.cpp b/Engine/Plugins/Runtime/GeometryProcessing/Source/DynamicMesh/Private/Sampling/MeshOcclusionMapBaker.cpp
--- a/Engine/Plugins/Runtime/GeometryProcessing/Source/DynamicMesh/Private/Sampling/MeshOcclusionMapBaker.cpp
+++ b/Engine/Plugins/Runtime/GeometryProcessing/Source/DynamicMesh/Private/Sampling/MeshOcclusionMapBaker.cpp
@@ -224,7 +224,7 @@ void FMeshOcclusionMapBaker::Bake()
 		}
 		if (WantBentNormal())
 		{
-			FVector3f NormalColor = (FVector3f(BentNormal.X, BentNormal.Y, BentNormal.Z) + FVector3f::One()) * 0.5;
+			FVector3f NormalColor = ((FVector3f)BentNormal + FVector3f::One()) * 0.5f;
 			NormalBuilder->SetPixel(Coords, NormalColor);
 		}
 	});
@@ -246,7 +246,7 @@ void FMeshOcclusionMapBaker::Bake()
 	if (WantAmbientOcclusion() && BlurRadius > 0.01)
 	{
 		TDiscreteKernel2f BlurKernel2d;
-		TGaussian2f::MakeKernelFromRadius(BlurRadius, BlurKernel2d);
+		TGaussian2f::MakeKernelFromRadius((float)BlurRadius, BlurKernel2d);
 		TArray<float> AOBlurBuffer;
 		Occupancy.ParallelProcessingPass<float>(
 			[&](int64 Index) { return 0.0f; },
diff --git a/Engine/Plugins/Runtime/GeometryProcessing/Source/DynamicMesh/Private/Sampling/MeshOcclusionMapEvaluator.cpp b/Engine/Plugins/Runtime/GeometryProcessing/Source/DynamicMesh/Private/Sampling/MeshOcclusionMapEvaluator.cpp
--- a/Engine/Plugins/Runtime/GeometryProcessing/Source/DynamicMesh/Private/Sampling/MeshOcclusionMapEvaluator.cpp
+++ b/Engine/Plugins/Runtime/GeometryProcessing/Source/DynamicMesh/Private/Sampling/MeshOcclusionMapEvaluator.cpp
@@ -157,7 +157,7 @@ void FMeshOcclusionMapEvaluator::EvaluateSample(float*& Out, const FCorresponden
 
 	if constexpr (WANT_AMBIENT_OCCLUSION(ComputeType))
 	{
-		WriteToBuffer(Out, FMathd::Clamp(1.0f - (float)Occlusion, 0.0f, 1.0f));
+		WriteToBuffer(Out, FMathf::Clamp(1.0f - (float)Occlusion, 0.0f, 1.0f));
 	}
 
 	if constexpr (WANT_BENT_NORMAL(ComputeType))
@@ -223,7 +223,7 @@ FMeshOcclusionMapEvaluator::FOcclusionTuple FMeshOcclusionMapEvaluator::SampleFu
 	if (!DetailTriNormal.Normalize())
 	{
 		// degenerate triangle normal
-		return FOcclusionTuple(0.0, DefaultNormal);
+		return FOcclusionTuple(0.0f, (FVector3f)DefaultNormal);
 	}
 
 	FVector3d BaseTangentX, BaseTangentY;
@@ -305,7 +305,7 @@ FMeshOcclusionMapEvaluator::FOcclusionTuple FMeshOcclusionMapEvaluator::SampleFu
 	{
 		AccumNormal = (AccumNormal.Length() > 0.0) ? Normalized(AccumNormal) : DefaultNormal;
 	}
-	return FOcclusionTuple(AccumOcclusion, AccumNormal);
+	return FOcclusionTuple( (float)AccumOcclusion, (FVector3f)AccumNormal);
 }
 
 #undef WANT_AMBIENT_OCCLUSION
