GeometryProcessing: add FMeshPlanarSymmetry for detecting and applying planar symmetry to mesh vertices

diff --git a/Engine/Plugins/Runtime/GeometryProcessing/Source/DynamicMesh/Private/Parameterization/MeshPlanarSymmetry.cpp b/Engine/Plugins/Runtime/GeometryProcessing/Source/DynamicMesh/Private/Parameterization/MeshPlanarSymmetry.cpp
new file mode 100644
--- /dev/null
+++ b/Engine/Plugins/Runtime/GeometryProcessing/Source/DynamicMesh/Private/Parameterization/MeshPlanarSymmetry.cpp
@@ -0,0 +1,244 @@
+// Copyright Epic Games, Inc. All Rights Reserved.
+
+#include "Parameterization/MeshPlanarSymmetry.h"
+#include "Async/ParallelFor.h"
+
+using namespace UE::Geometry;
+
+
+bool FMeshPlanarSymmetry::Initialize(FDynamicMesh3* Mesh, FDynamicMeshAABBTree3* Spatial, FFrame3d SymmetryFrameIn)
+{
+	this->TargetMesh = Mesh;
+	this->SymmetryFrame = SymmetryFrameIn;
+	this->CachedSymmetryAxis = SymmetryFrameIn.Z();
+
+	Vertices.Init(FSymmetryVertex(), TargetMesh->MaxVertexID());
+
+	TArray<int32> PositiveVerts;
+	PositiveVerts.Reserve(Vertices.Num());
+
+	for (int32 vid : Mesh->VertexIndicesItr())
+	{
+		FVector3d Position = Mesh->GetVertex(vid);
+		Vertices[vid].PlaneSignedDistance = (Position - SymmetryFrame.Origin).Dot(CachedSymmetryAxis);
+		if (Vertices[vid].PlaneSignedDistance > FMathd::ZeroTolerance)
+		{
+			PositiveVerts.Add(vid);
+			Vertices[vid].bIsSourceVertex = true;
+		}
+		else
+		{
+			Vertices[vid].bIsSourceVertex = false;
+		}
+	}
+
+	// Compute unique matching vertices on mirror side (negative side of plane) for positive-side vertices
+	ParallelFor(PositiveVerts.Num(), [&](int32 k)
+	{
+		int32 VertexID = PositiveVerts[k];
+
+		// compute min edge length in one-ring
+		double MinEdgeLenSqr = TNumericLimits<double>::Max();
+		Mesh->EnumerateVertexVertices(VertexID, [&](int32 NbrVertexID) {
+			double DistSqr = DistanceSquared(Mesh->GetVertex(VertexID), Mesh->GetVertex(NbrVertexID));
+			MinEdgeLenSqr = FMath::Min(DistSqr, MinEdgeLenSqr);
+		});
+		double MinEdgeLen = FMathd::Sqrt(MinEdgeLenSqr);
+
+		FVector3d Position = Mesh->GetVertex(VertexID);
+		FVector3d MirrorPosition = GetMirroredPosition(Position);
+
+		// find vertex that is closest to mirrored position
+		double NearestDistSqr = 0;
+		int32 NearestVID = Spatial->FindNearestVertex(MirrorPosition, NearestDistSqr, (0.5*MinEdgeLen) + (double)FMathf::ZeroTolerance * 10.0);
+		if (Mesh->IsVertex(NearestVID) 
+			&& (NearestVID != VertexID)
+			&& (Vertices[NearestVID].bIsSourceVertex == false) )
+		{
+			// only match vertex if it is within half min edge length, otherwise it is ambiguous...
+			if ( NearestDistSqr < (0.25*MinEdgeLenSqr) )
+			{
+				Vertices[NearestVID].PairedVertex = VertexID;
+				Vertices[VertexID].PairedVertex = NearestVID;
+			}
+		}
+	});
+
+	// Point must be within this distance from the symmetry plane to be considered "on" the plane.
+	// (may in future be desirable to relax/ignore this tolerance for unmatched vertices that connect two
+	//  matched symmetric vertices)
+	double OnPlaneTolerance = (double)FMathf::ZeroTolerance * 10.0;
+
+	// If a vertex has no match and is within the symmetry-plane tolerance band, consider it a plane vertex.
+	// Otherwise it is a failed match
+	int32 NumFailedMatches = 0;
+	for (int32 VertexID : Mesh->VertexIndicesItr())
+	{
+		if ( Vertices[VertexID].PairedVertex < 0 )
+		{
+			if (FMathd::Abs(Vertices[VertexID].PlaneSignedDistance) < OnPlaneTolerance)
+			{
+				Vertices[VertexID].bIsSourceVertex = false;
+				Vertices[VertexID].bOnPlane = true;
+			}
+			else
+			{
+				NumFailedMatches++;
+			}
+		}
+	}
+
+	return (NumFailedMatches == 0);
+}
+
+
+
+FFrame3d FMeshPlanarSymmetry::GetPositiveSideFrame(FFrame3d FromFrame) const
+{
+	FVector3d DeltaVec = (FromFrame.Origin - SymmetryFrame.Origin);
+	double SignedDistance = DeltaVec.Dot(CachedSymmetryAxis);
+	if (SignedDistance < 0)
+	{
+		FromFrame.Origin = GetMirroredPosition(FromFrame.Origin);
+		FromFrame.Rotation = GetMirroredOrientation(FromFrame.Rotation);
+	}
+	return FromFrame;
+}
+
+
+
+FVector3d FMeshPlanarSymmetry::GetMirroredPosition(const FVector3d& Position) const
+{
+	FVector3d DeltaVec = (Position - SymmetryFrame.Origin);
+	double SignedDistance = DeltaVec.Dot(CachedSymmetryAxis);
+	FVector3d MirrorPosition = Position - (2 * SignedDistance * CachedSymmetryAxis);
+	return MirrorPosition;
+}
+
+FVector3d FMeshPlanarSymmetry::GetMirroredAxis(const FVector3d& Axis) const
+{
+	double SignedDistance = Axis.Dot(CachedSymmetryAxis);
+	FVector3d MirrorAxis = Axis - (2 * SignedDistance * CachedSymmetryAxis);
+	return MirrorAxis;
+}
+
+FQuaterniond FMeshPlanarSymmetry::GetMirroredOrientation(const FQuaterniond& Orientation) const
+{
+	FVector3d Axis(Orientation.X, Orientation.Y, Orientation.Z);
+	Axis = GetMirroredAxis(Axis);
+	return FQuaterniond(Axis.X, Axis.Y, Axis.Z, -Orientation.W);
+}
+
+
+void FMeshPlanarSymmetry::GetMirrorVertexROI(const TArray<int>& VertexROI, TArray<int>& MirrorVertexROI, bool bForceSameSizeWithGaps) const
+{
+	MirrorVertexROI.Reset();
+	if (bForceSameSizeWithGaps)
+	{
+		MirrorVertexROI.Reserve(VertexROI.Num());
+	}
+
+	for (int32 VertexID : VertexROI)
+	{
+		const FSymmetryVertex& Vertex = Vertices[VertexID];
+		if (Vertex.bIsSourceVertex)
+		{
+			int32 MirrorVID = Vertex.PairedVertex;
+			MirrorVertexROI.Add(MirrorVID);
+		}
+		else if (Vertex.bIsSourceVertex == false && Vertex.bOnPlane == false && Vertex.PairedVertex >= 0)
+		{
+			// The ROI already contains the mirror-vertex, we will keep it so that it can be updated later
+			MirrorVertexROI.Add(VertexID);
+		}
+		else if (bForceSameSizeWithGaps)
+		{
+			MirrorVertexROI.Add(-1);
+		}
+	}
+}
+
+
+
+void FMeshPlanarSymmetry::ApplySymmetryPlaneConstraints(const TArray<int>& VertexIndices, TArray<FVector3d>& VertexPositions) const
+{
+	int32 N = VertexIndices.Num();
+	check(N == VertexPositions.Num());
+	for (int32 k = 0; k < N; ++k)
+	{
+		int32 VertexID = VertexIndices[k];
+		const FSymmetryVertex& Vertex = Vertices[VertexID];
+		if (Vertex.bOnPlane)
+		{
+			VertexPositions[k] = SymmetryFrame.ToPlane(VertexPositions[k]);
+		}
+	}
+}
+
+
+void FMeshPlanarSymmetry::ComputeSymmetryConstrainedPositions(
+	const TArray<int>& SourceVertexROI,
+	const TArray<int>& MirrorVertexROI,
+	const TArray<FVector3d>& SourceVertexPositions,
+	TArray<FVector3d>& MirrorVertexPositionsOut) const
+{
+	int32 NumV = SourceVertexROI.Num();
+	check(MirrorVertexROI.Num() == NumV);
+	check(SourceVertexPositions.Num() == NumV);
+
+	MirrorVertexPositionsOut.SetNum(NumV, false);
+
+	for (int32 k = 0; k < NumV; ++k)
+	{
+		int32 MirrorVertexID = MirrorVertexROI[k];
+		if (MirrorVertexID >= 0)
+		{
+			int32 SourceVertexID = SourceVertexROI[k];
+			if (SourceVertexID == MirrorVertexID)
+			{
+				SourceVertexID = Vertices[MirrorVertexID].PairedVertex;
+				MirrorVertexPositionsOut[k] = GetMirroredPosition(TargetMesh->GetVertex(SourceVertexID));
+			}
+			else
+			{
+				check(Vertices[SourceVertexID].PairedVertex == MirrorVertexID);
+				MirrorVertexPositionsOut[k] = GetMirroredPosition(SourceVertexPositions[k]);
+			}
+		}
+	}
+}
+
+
+void FMeshPlanarSymmetry::FullSymmetryUpdate()
+{
+	for (int32 vid : TargetMesh->VertexIndicesItr())
+	{
+		if (Vertices[vid].bIsSourceVertex)
+		{
+			UpdateSourceVertex(vid);
+		}
+		else if (Vertices[vid].bOnPlane)
+		{
+			UpdatePlaneVertex(vid);
+		}
+	}
+}
+
+
+void FMeshPlanarSymmetry::UpdateSourceVertex(int32 VertexID)
+{
+	FVector3d CurPosition = TargetMesh->GetVertex(VertexID);
+	int32 MirrorVID = Vertices[VertexID].PairedVertex;
+	check(TargetMesh->IsVertex(MirrorVID));
+	FVector3d MirrorPosition = GetMirroredPosition(CurPosition);
+	TargetMesh->SetVertex(MirrorVID, MirrorPosition);
+}
+
+void FMeshPlanarSymmetry::UpdatePlaneVertex(int32 VertexID)
+{
+	FVector3d CurPosition = TargetMesh->GetVertex(VertexID);
+	FVector3d PlanePos = SymmetryFrame.ToPlane(CurPosition);
+	TargetMesh->SetVertex(VertexID, PlanePos);
+}
+
+
diff --git a/Engine/Plugins/Runtime/GeometryProcessing/Source/DynamicMesh/Public/Parameterization/MeshPlanarSymmetry.h b/Engine/Plugins/Runtime/GeometryProcessing/Source/DynamicMesh/Public/Parameterization/MeshPlanarSymmetry.h
new file mode 100644
--- /dev/null
+++ b/Engine/Plugins/Runtime/GeometryProcessing/Source/DynamicMesh/Public/Parameterization/MeshPlanarSymmetry.h
@@ -0,0 +1,125 @@
+// Copyright Epic Games, Inc. All Rights Reserved.
+
+#pragma once
+
+#include "DynamicMesh/DynamicMesh3.h"
+#include "DynamicMesh/DynamicMeshAABBTree3.h"
+#include "Quaternion.h"
+
+namespace UE
+{
+namespace Geometry
+{
+
+/**
+ * FMeshPlanarSymmetry detects pairwise symmetry relationships between vertices in a mesh, given a symmetry frame.
+ * Once those relationships are known, symmetry can be re-enforced after mesh edits.
+ * 
+ * Vertices within a tolerance band of the symmetry plane do not have a mirror and are considered "on plane"
+ * and will be snapped to the symmetry plane when enforcing symmetry after edits.
+ * 
+ * The Symmetry Plane is specified by a Frame3d where the XY axes define the plane and the Z axis is the plane normal.
+ * So references to "Positive" side below are relative to that Z axis/plane
+ */
+class DYNAMICMESH_API FMeshPlanarSymmetry
+{
+public:
+
+	/**
+	 * Given a Mesh, an AABBTree, and a Symmetry Plane/Frame, detect any pairs of vertices with
+	 * planar/mirror-symmetry relationships, as well as "on plane" vertices
+	 * @return false if any non-on-plane vertex fails to find a match. 
+	 */
+	bool Initialize(FDynamicMesh3* Mesh, FDynamicMeshAABBTree3* Spatial, FFrame3d SymmetryFrameIn);
+
+	/**
+	 * @return the input Point mirrored across the Symmetry plane
+	 */
+	FVector3d GetMirroredPosition(const FVector3d& Position) const;
+
+	/**
+	 * @return the input Vector/Axis mirrored across the Symmetry plane
+	 */
+	FVector3d GetMirroredAxis(const FVector3d& Axis) const;
+
+	/**
+	 * @return the input Quaternion mirrored across the Symmetry plane
+	 */
+	FQuaterniond GetMirroredOrientation(const FQuaterniond& Orientation) const;
+
+	/**
+	 * @return the input Frame mirrored to the "positive" side (unchanged if it is already on the positive side)
+	 */
+	FFrame3d GetPositiveSideFrame(FFrame3d FromFrame) const;
+
+
+	/**
+	 * Update all the symmetry vertex positions based on their Positive-side pair vertex.
+	 * Also snap all on-plane vertices to the symmetry plane
+	 */
+	void FullSymmetryUpdate();
+
+
+
+	//
+	// The set of functions below are intended to be used together in situations like 3D sculpting
+	// where we want to apply symmetry constraints within a "brush region of interest (ROI)"
+	// See UMeshVertexSculptTool for example usage
+	//
+
+
+	/**
+	 * Computes list of vertices that are mirror-constrained to vertices in VertexROI.
+	 * If bForceSameSizeWithGaps is true, MirrorVertexROIOut will be the same size as VertexROI, and -1 is stored 
+	 * for vertices in VertexROI that are source vertices or on the symmetry plane. Otherwise those vertices are skipped.
+	 */
+	void GetMirrorVertexROI(const TArray<int>& VertexROI, TArray<int>& MirrorVertexROIOut, bool bForceSameSizeWithGaps) const;
+
+	/**
+	 * For any vertices in VertexIndices that are on-plane, take the position in VertexPositionsInOut, apply the
+	 * on-plane constraint, and return the constrained position in VertexPositionsInOut  (ie VertexPositionsInOut is updated to new positions)
+	 */
+	void ApplySymmetryPlaneConstraints(const TArray<int>& VertexIndices, TArray<FVector3d>& VertexPositionsInOut) const;
+
+	/**
+	 * Given the pairing (SourceVertexROI, SourceVertexPositions), compute the symmetry-constrained vertex positions for
+	 * MirrorVertexROI and store in MirrorVertexPositionsOut. This function assumes that the MirrorVertexROI was computed
+	 * by calling GetMirrorVertexROI(SourceVertexROI, MirrorVertexROI, bForceSameSizeWithGaps=true), ie all the arrays
+	 * must be the same length
+	 */
+	void ComputeSymmetryConstrainedPositions(
+		const TArray<int>& SourceVertexROI,
+		const TArray<int>& MirrorVertexROI,
+		const TArray<FVector3d>& SourceVertexPositions,
+		TArray<FVector3d>& MirrorVertexPositionsOut) const;
+
+
+protected:
+	FDynamicMesh3* TargetMesh = nullptr;
+
+	FFrame3d SymmetryFrame;
+	FVector3d CachedSymmetryAxis;
+
+	struct FSymmetryVertex
+	{
+		// initial distance to symmetry plane
+		double PlaneSignedDistance = 0.0;
+		// true if this is a positive-side vertex
+		bool bIsSourceVertex = true;
+		// true if this is an on-plane vertex (todo: possibly combine with bIsSourceVertex as an enum)
+		bool bOnPlane = false;
+		// VertexID of paired vertex
+		int32 PairedVertex = -1;
+	};
+
+	// list of vertices, size == TargetMesh.MaxVertexID()
+	TArray<FSymmetryVertex> Vertices;
+
+	void UpdateSourceVertex(int32 VertexID);
+	void UpdatePlaneVertex(int32 VertexID);
+};
+
+
+
+} // end namespace UE::Geometry
+} // end namespace UE
\ No newline at end of file
