GeometryProcessing: Mesh Solvers.   Fix implicit cast warnings.  These are related the integer size type used by Eigen, and our comparatively  modest maxtrix  row/col max of Max_int32.

diff --git a/Engine/Plugins/Runtime/GeometryProcessing/Source/DynamicMesh/Private/Solvers/Internal/ConstrainedPoissonSolver.h b/Engine/Plugins/Runtime/GeometryProcessing/Source/DynamicMesh/Private/Solvers/Internal/ConstrainedPoissonSolver.h
--- a/Engine/Plugins/Runtime/GeometryProcessing/Source/DynamicMesh/Private/Solvers/Internal/ConstrainedPoissonSolver.h
+++ b/Engine/Plugins/Runtime/GeometryProcessing/Source/DynamicMesh/Private/Solvers/Internal/ConstrainedPoissonSolver.h
@@ -7,6 +7,7 @@
 
 #include "FSparseMatrixD.h"
 #include "MatrixSolver.h"
+#include "Math/NumericLimits.h"
 
 // only needed if you enable logging by changing the typedef for FConstrainedSolverTimeLogger
 //#include "ProfilingDebugging/ScopedTimers.h"
@@ -56,8 +57,9 @@ public:
 
 
 	FConstrainedSolver(TUniquePtr<FSparseMatrixD>& SymmatrixMatrixOperator, const EMatrixSolverType MatrixSolverType)
-		: ConstraintPositions(SymmatrixMatrixOperator->cols())
+		: ConstraintPositions((int32)SymmatrixMatrixOperator->cols())
 	{
+		checkSlow(SymmatrixMatrixOperator->cols() <= MAX_int32);
 		SymmetricMatrixPtr.Reset(SymmatrixMatrixOperator.Release());
 
 		bMatrixSolverDirty = true;
@@ -201,7 +203,7 @@ public:
 
 		FSparseMatrixD& SymmetricMatrix = *SymmetricMatrixPtr;
 		// Set up the source vector
-		FSOAPositions RHSVector(SymmetricMatrix.cols());
+		FSOAPositions RHSVector((int32)SymmetricMatrix.cols());
 		for (int32 Dir = 0; Dir < 3; ++Dir) 
 		{
 			RHSVector.Array(Dir) = SourceVector.Array(Dir) + WeightsSqrdMatrix * ConstraintPositions.Array(Dir);
@@ -226,7 +228,7 @@ public:
 		FConstrainedSolverTimeLogger Timmer(TEXT("Post-setup solve time"));
 
 		FSparseMatrixD& SymmetricMatrix = *SymmetricMatrixPtr;
-		FSOAPositions RHSVector(SymmetricMatrix.cols());
+		FSOAPositions RHSVector((int32)SymmetricMatrix.cols());
 		for (int32 Dir = 0; Dir < 3; ++Dir)
 		{
 			RHSVector.Array(Dir) = WeightsSqrdMatrix * ConstraintPositions.Array(Dir);
@@ -255,7 +257,7 @@ public:
 			FConstrainedSolverTimeLogger Timmer(TEXT("Post-setup solve time"));
 
 			FSparseMatrixD& SymmetricMatrix = *SymmetricMatrixPtr;
-			FSOAPositions RHSVector(SymmetricMatrix.cols());
+			FSOAPositions RHSVector((int32)SymmetricMatrix.cols());
 			for (int32 Dir = 0; Dir < 3; ++Dir)
 			{
 				RHSVector.Array(Dir) = WeightsSqrdMatrix * ConstraintPositions.Array(Dir);
@@ -290,7 +292,7 @@ public:
 
 			FSparseMatrixD& SymmetricMatrix = *SymmetricMatrixPtr;
 			// Set up the source vector
-			FSOAPositions RHSVector(SymmetricMatrix.cols());
+			FSOAPositions RHSVector((int32)SymmetricMatrix.cols());
 			for (int32 Dir = 0; Dir < 3; ++Dir)
 			{
 				RHSVector.Array(Dir) = SourceVector.Array(Dir) + WeightsSqrdMatrix * ConstraintPositions.Array(Dir);
@@ -324,7 +326,7 @@ protected:
 	void ClearConstraintPositions()
 	{
 		const FSparseMatrixD& SymmetricMatrix = *SymmetricMatrixPtr;
-		const int32 NumColumns = SymmetricMatrix.cols();
+		const int32 NumColumns = (int32)SymmetricMatrix.cols();
 
 		ConstraintPositions.SetZero(NumColumns);
 
diff --git a/Engine/Plugins/Runtime/GeometryProcessing/Source/DynamicMesh/Private/Solvers/Internal/LaplacianOperators.cpp b/Engine/Plugins/Runtime/GeometryProcessing/Source/DynamicMesh/Private/Solvers/Internal/LaplacianOperators.cpp
--- a/Engine/Plugins/Runtime/GeometryProcessing/Source/DynamicMesh/Private/Solvers/Internal/LaplacianOperators.cpp
+++ b/Engine/Plugins/Runtime/GeometryProcessing/Source/DynamicMesh/Private/Solvers/Internal/LaplacianOperators.cpp
@@ -138,7 +138,7 @@ double ConstructScaledCotangentLaplacian(const FDynamicMesh3& DynamicMesh, FVert
 	ConstructCotangentLaplacian(DynamicMesh, VertexMap, AreaMatrix, CotangentInterior, CotangentBoundary);
 
 	// Find average entry in the area matrix
-	const int32 Rank = AreaMatrix.cols();
+	const int32 Rank = (int32)AreaMatrix.cols();
 	double AveArea = 0.;
 	for (int32 i = 0; i < Rank; ++i)
 	{
diff --git a/Engine/Plugins/Runtime/GeometryProcessing/Source/DynamicMesh/Private/Solvers/Internal/MeshDiffusionIntegrator.cpp b/Engine/Plugins/Runtime/GeometryProcessing/Source/DynamicMesh/Private/Solvers/Internal/MeshDiffusionIntegrator.cpp
--- a/Engine/Plugins/Runtime/GeometryProcessing/Source/DynamicMesh/Private/Solvers/Internal/MeshDiffusionIntegrator.cpp
+++ b/Engine/Plugins/Runtime/GeometryProcessing/Source/DynamicMesh/Private/Solvers/Internal/MeshDiffusionIntegrator.cpp
@@ -54,7 +54,7 @@ void FMeshDiffusionIntegrator::Initialize(const FDynamicMesh3& DynamicMesh, cons
 	const FSparseMatrixD& M = DiffusionOperator;
 
 	// Find the min diagonal entry (all should be negative).
-	int32 Rank = M.rows();
+	int32 Rank = (int32)M.rows();
 	MinDiagonalValue = FSparseMatrixD::Scalar(0);
 	for (int32 i = 0; i < Rank; ++i)
 	{
diff --git a/Engine/Plugins/Runtime/GeometryProcessing/Source/DynamicMesh/Private/Solvers/Internal/PowerMethodSolver.cpp b/Engine/Plugins/Runtime/GeometryProcessing/Source/DynamicMesh/Private/Solvers/Internal/PowerMethodSolver.cpp
--- a/Engine/Plugins/Runtime/GeometryProcessing/Source/DynamicMesh/Private/Solvers/Internal/PowerMethodSolver.cpp
+++ b/Engine/Plugins/Runtime/GeometryProcessing/Source/DynamicMesh/Private/Solvers/Internal/PowerMethodSolver.cpp
@@ -2,6 +2,7 @@
 
 #include "PowerMethodSolver.h"
 #include "DenseMatrix.h"
+#include "Math/NumericLimits.h"
 
 using namespace UE::Geometry;
 
@@ -231,6 +232,8 @@ void FSparsePowerMethod::CreateSparseMatrixOperator(const FSparseMatrixD& InMatr
                                                     const FMatrixHints& InMatrixHints,
                                                     MatrixOperator& OutMatrixOp) 
 {
+	checkSlow(InMatrix.rows() <= MAX_int32);
+
     EMatrixSolverType MatrixSolverType = InMatrixHints.bIsPSD ? EMatrixSolverType::FastestPSD : EMatrixSolverType::LU;
     bool bIsSymmetric = InMatrixHints.bIsSymmetric || InMatrixHints.bIsPSD;
     
@@ -254,6 +257,6 @@ void FSparsePowerMethod::CreateSparseMatrixOperator(const FSparseMatrixD& InMatr
 
     OutMatrixOp.Rows = [&InMatrix]() 
     {
-        return InMatrix.rows();
+        return (int32)InMatrix.rows();
     };
 }
\ No newline at end of file
