Fix comment: FMeshMergeCoincidentEdges => FMergeCoincidentMeshEdges

diff --git a/Engine/Plugins/Runtime/GeometryProcessing/Source/DynamicMesh/Public/Operations/MeshResolveTJunctions.h b/Engine/Plugins/Runtime/GeometryProcessing/Source/DynamicMesh/Public/Operations/MeshResolveTJunctions.h
--- a/Engine/Plugins/Runtime/GeometryProcessing/Source/DynamicMesh/Public/Operations/MeshResolveTJunctions.h
+++ b/Engine/Plugins/Runtime/GeometryProcessing/Source/DynamicMesh/Public/Operations/MeshResolveTJunctions.h
@@ -30,7 +30,7 @@ class FDynamicMesh3;
  *                    G  C                         G  C
  *
  * If everything is within-tolerance, then after resolving all the T-Junctions, 
- * a FMeshMergeCoincidentEdges would successfully weld the new set of border 
+ * a FMergeCoincidentMeshEdges would successfully weld the new set of border 
  * edges back together (note: caller must do this, FMeshResolveTJunctions only splits, it does not weld!)
  * 
  * Caller can provide a subset of edges via BoundaryEdges, otherwise all boundary edges
