GeometryProcessing - Cleaning up implicit type cast warnings in DynamicMesh-based selections, simplification and other processing elements.

diff --git a/Engine/Plugins/Runtime/GeometryProcessing/Source/DynamicMesh/Private/DynamicMeshEditor.cpp b/Engine/Plugins/Runtime/GeometryProcessing/Source/DynamicMesh/Private/DynamicMeshEditor.cpp
--- a/Engine/Plugins/Runtime/GeometryProcessing/Source/DynamicMesh/Private/DynamicMeshEditor.cpp
+++ b/Engine/Plugins/Runtime/GeometryProcessing/Source/DynamicMesh/Private/DynamicMeshEditor.cpp
@@ -152,8 +152,8 @@ bool FDynamicMeshEditor::ConvertLoopToTriVidPairSequence(const FDynamicMesh3& Me
 		int32 SecondVid = VidLoop[(QuadIndex + 1) % VidLoop.Num()];
 
 		FIndex3i TriVids = Mesh.GetTriangle(Tid);
-		int8 SubIdx1 = IndexUtil::FindTriIndex(FirstVid, TriVids);
-		int8 SubIdx2 = IndexUtil::FindTriIndex(SecondVid, TriVids);
+		int8 SubIdx1 = (int8)IndexUtil::FindTriIndex(FirstVid, TriVids);
+		int8 SubIdx2 = (int8)IndexUtil::FindTriIndex(SecondVid, TriVids);
 		if (ensure(SubIdx1 >= 0 && SubIdx2 >= 0))
 		{
 			TriVertPairsOut.Emplace(Tid,
@@ -264,9 +264,9 @@ bool FDynamicMeshEditor::StitchSparselyCorrespondedVertexLoops(const TArray<int>
 		int Starts[2] { MatchedIndices1[i], MatchedIndices2[i] };
 		int Ends[2] { MatchedIndices1[(i + 1) % CorrespondN], MatchedIndices2[(i + 1) % CorrespondN] };
 
-		auto GetWrappedSpanLen = [](const FDynamicMesh3* M, const TArray<int>& VertexIDs, int StartInd, int EndInd)
+		auto GetWrappedSpanLen = [](const FDynamicMesh3* M, const TArray<int>& VertexIDs, int StartInd, int EndInd)->float
 		{
-			float LenTotal = 0;
+			double LenTotal = 0;
 			FVector3d V = M->GetVertex(VertexIDs[StartInd]);
 			for (int Ind = StartInd, IndNext; Ind != EndInd;)
 			{
@@ -276,7 +276,7 @@ bool FDynamicMeshEditor::StitchSparselyCorrespondedVertexLoops(const TArray<int>
 				Ind = IndNext;
 				V = VNext;
 			}
-			return LenTotal;
+			return (float)LenTotal;
 		};
 		float LenTotal[2] { GetWrappedSpanLen(Mesh, VertexIDs1, Starts[0], Ends[0]), GetWrappedSpanLen(Mesh, VertexIDs2, Starts[1], Ends[1]) };
 		float LenAlong[2] { FMathf::Epsilon, FMathf::Epsilon };
@@ -300,7 +300,7 @@ bool FDynamicMeshEditor::StitchSparselyCorrespondedVertexLoops(const TArray<int>
 				
 				Tri.C = VertexIDs1[Walks[0]];
 				FVector3d NextV = Mesh->GetVertex(Tri.C);
-				LenAlong[0] += Distance(NextV, Vertex[0]);
+				LenAlong[0] += (float)Distance(NextV, Vertex[0]);
 				Vertex[0] = NextV;
 			}
 			else
@@ -308,7 +308,7 @@ bool FDynamicMeshEditor::StitchSparselyCorrespondedVertexLoops(const TArray<int>
 				Walks[1] = (Walks[1] + 1) % VertexIDs2.Num();
 				Tri.C = VertexIDs2[Walks[1]];
 				FVector3d NextV = Mesh->GetVertex(Tri.C);
-				LenAlong[1] += Distance(NextV, Vertex[1]);
+				LenAlong[1] += (float)Distance(NextV, Vertex[1]);
 				Vertex[1] = NextV;
 			}
 			if (bReverseOrientation)
@@ -1067,7 +1067,7 @@ void FDynamicMeshEditor::SetGeneralTubeUVs(const TArray<int>& Triangles,
 	FVector3d RefPos = Mesh->GetVertex(VertexIDs1[0]);
 	auto GetUV = [this, &VDir, &UVScaleFactor, &UVTranslation, &RefPos](int MeshIdx, float UStart, float UEnd, float Param)
 	{
-		return FVector2f((Mesh->GetVertex(MeshIdx) - RefPos).Dot((FVector3d)VDir), FMath::Lerp(UStart, UEnd, Param)) * UVScaleFactor + UVTranslation;
+		return FVector2f(float( (Mesh->GetVertex(MeshIdx) - RefPos).Dot((FVector3d)VDir) ), FMath::Lerp(UStart, UEnd, Param)) * UVScaleFactor + UVTranslation;
 	};
 
 	TArray<FVector2f> VertUVs[2];
@@ -1107,7 +1107,7 @@ void FDynamicMeshEditor::SetGeneralTubeUVs(const TArray<int>& Triangles,
 				}
 				for (int InsideIdx = (Start + 1) % NumVertices; InsideIdx != End; InsideIdx = (InsideIdx + 1) % NumVertices)
 				{
-					double InterpT = (Mesh->GetVertex(VertexIDs[InsideIdx]) - StartPos).Dot(Along) / SepSq;
+					float InterpT = float( (Mesh->GetVertex(VertexIDs[InsideIdx]) - StartPos).Dot(Along) / SepSq );
 					VertUVs[Side][InsideIdx] = GetUV(VertexIDs[InsideIdx], UValues[Idx], UValues[NextIdx], InterpT);
 				}
 			}
@@ -1303,7 +1303,7 @@ void FDynamicMeshEditor::RescaleAttributeUVs(float UVScale, bool bWorldSpace, in
 		}
 		if (TotalEdgeUVLen > KINDA_SMALL_NUMBER)
 		{
-			float AvgUVScale = TotalEdgeLen / TotalEdgeUVLen;
+			float AvgUVScale = float (TotalEdgeLen / TotalEdgeUVLen);
 			UVScale *= AvgUVScale;
 		}
 	}
diff --git a/Engine/Plugins/Runtime/GeometryProcessing/Source/DynamicMesh/Private/GroupTopology.cpp b/Engine/Plugins/Runtime/GeometryProcessing/Source/DynamicMesh/Private/GroupTopology.cpp
--- a/Engine/Plugins/Runtime/GeometryProcessing/Source/DynamicMesh/Private/GroupTopology.cpp
+++ b/Engine/Plugins/Runtime/GeometryProcessing/Source/DynamicMesh/Private/GroupTopology.cpp
@@ -63,7 +63,7 @@ bool FGroupTopology::RebuildTopology()
 	GroupFaceCounts.Init(0, MaxGroupID);
 	for (int tid : Mesh->TriangleIndicesItr())
 	{
-		int GroupID = FMathd::Max(0, GetGroupID(tid));
+		int GroupID = FMath::Max(0, GetGroupID(tid));
 		if (GroupIDToGroupIndexMap[GroupID] == -1)
 		{
 			FGroup NewGroup;
@@ -81,7 +81,7 @@ bool FGroupTopology::RebuildTopology()
 	// sort faces into groups
 	for (int tid : Mesh->TriangleIndicesItr())
 	{
-		int GroupID = FMathd::Max(0, GetGroupID(tid));
+		int GroupID = FMath::Max(0, GetGroupID(tid));
 		Groups[GroupIDToGroupIndexMap[GroupID]].Triangles.Add(tid);
 	}
 
diff --git a/Engine/Plugins/Runtime/GeometryProcessing/Source/DynamicMesh/Private/MeshSimplification.cpp b/Engine/Plugins/Runtime/GeometryProcessing/Source/DynamicMesh/Private/MeshSimplification.cpp
--- a/Engine/Plugins/Runtime/GeometryProcessing/Source/DynamicMesh/Private/MeshSimplification.cpp
+++ b/Engine/Plugins/Runtime/GeometryProcessing/Source/DynamicMesh/Private/MeshSimplification.cpp
@@ -770,7 +770,7 @@ void TMeshSimplification<QuadricErrorType>::SimplifyToMaxError(double MaxError)
 	SimplifyMode = ETargetModes::MaxError;
 	TargetCount = 1;
 	MinEdgeLength = FMathd::MaxReal;
-	MaxErrorAllowed = MaxError;
+	MaxErrorAllowed = (float)MaxError;
 	DoSimplify();
 }
 
diff --git a/Engine/Plugins/Runtime/GeometryProcessing/Source/DynamicMesh/Private/Operations/GroupEdgeInserter.cpp b/Engine/Plugins/Runtime/GeometryProcessing/Source/DynamicMesh/Private/Operations/GroupEdgeInserter.cpp
--- a/Engine/Plugins/Runtime/GeometryProcessing/Source/DynamicMesh/Private/Operations/GroupEdgeInserter.cpp
+++ b/Engine/Plugins/Runtime/GeometryProcessing/Source/DynamicMesh/Private/Operations/GroupEdgeInserter.cpp
@@ -1361,8 +1361,8 @@ bool GetPlaneCutPath(const FDynamicMesh3& Mesh, int32 GroupID,
 		// could end up clamping to different endpoints as we connect multiple paths through the same start/end
 		// point (such as when following a loop)
 		FIndex2i EdgeVids = Mesh.GetEdgeV(StartPoint.ElementID);
-		CurrentEdgeVertPlaneDistances[0] = FVector::PointPlaneDist((FVector)Mesh.GetVertex(EdgeVids.A), CutPlaneOrigin, CutPlaneNormal);
-		CurrentEdgeVertPlaneDistances[1] = FVector::PointPlaneDist((FVector)Mesh.GetVertex(EdgeVids.B), CutPlaneOrigin, CutPlaneNormal);
+		CurrentEdgeVertPlaneDistances[0] = (float)FVector::PointPlaneDist((FVector)Mesh.GetVertex(EdgeVids.A), CutPlaneOrigin, CutPlaneNormal);
+		CurrentEdgeVertPlaneDistances[1] = (float)FVector::PointPlaneDist((FVector)Mesh.GetVertex(EdgeVids.B), CutPlaneOrigin, CutPlaneNormal);
 
 			OutputPath.Emplace(FMeshSurfacePoint(StartPoint.ElementID, StartPoint.EdgeTValue), FDynamicMesh3::InvalidID);
 		}
@@ -1447,8 +1447,8 @@ bool GetPlaneCutPath(const FDynamicMesh3& Mesh, int32 GroupID,
 				}
 
 				// See if one of the other vertices is on the plane (and is therefore the next destination)
-				float PlaneDistanceA = FVector::PointPlaneDist((FVector)Mesh.GetVertex(VertA), CutPlaneOrigin, CutPlaneNormal);
-				float PlaneDistanceB = FVector::PointPlaneDist((FVector)Mesh.GetVertex(VertB), CutPlaneOrigin, CutPlaneNormal);
+				float PlaneDistanceA = (float)FVector::PointPlaneDist((FVector)Mesh.GetVertex(VertA), CutPlaneOrigin, CutPlaneNormal);
+				float PlaneDistanceB = (float)FVector::PointPlaneDist((FVector)Mesh.GetVertex(VertB), CutPlaneOrigin, CutPlaneNormal);
 				bool bVertAIsOnPlane = abs(PlaneDistanceA) <= VertexCutTolerance;
 				bool bVertBIsOnPlane = abs(PlaneDistanceB) <= VertexCutTolerance;
 
@@ -1574,7 +1574,7 @@ bool GetPlaneCutPath(const FDynamicMesh3& Mesh, int32 GroupID,
 			}
 
 			// We'll keep going. Get the placement of the opposite vert relative to the plane.
-			float OppositeVertPlaneDistance = FVector::PointPlaneDist((FVector)Mesh.GetVertex(OppositeVert), CutPlaneOrigin, CutPlaneNormal);
+			float OppositeVertPlaneDistance = (float)FVector::PointPlaneDist((FVector)Mesh.GetVertex(OppositeVert), CutPlaneOrigin, CutPlaneNormal);
 			if (abs(OppositeVertPlaneDistance) <= VertexCutTolerance && !DisallowedVids.Contains(OppositeVert))
 			{
 				// We are cutting through a vertex
diff --git a/Engine/Plugins/Runtime/GeometryProcessing/Source/DynamicMesh/Private/Parameterization/DynamicMeshUVEditor.cpp b/Engine/Plugins/Runtime/GeometryProcessing/Source/DynamicMesh/Private/Parameterization/DynamicMeshUVEditor.cpp
--- a/Engine/Plugins/Runtime/GeometryProcessing/Source/DynamicMesh/Private/Parameterization/DynamicMeshUVEditor.cpp
+++ b/Engine/Plugins/Runtime/GeometryProcessing/Source/DynamicMesh/Private/Parameterization/DynamicMeshUVEditor.cpp
@@ -210,9 +210,9 @@ void InternalSetPerTriangleUVs(EnumerableType TriangleIDs, const FDynamicMesh3*
 		FIndex3i ElemTri;
 		for (int32 j = 0; j < 3; ++j)
 		{
-			FVector2f UV = (FVector2f)TriProjFrame.ToPlaneUV(Mesh->GetVertex(MeshTri[j]), 2);
+			FVector2d UV = TriProjFrame.ToPlaneUV(Mesh->GetVertex(MeshTri[j]), 2);
 			UV *= ScaleFactor;
-			ElemTri[j] = UVOverlay->AppendElement(UV);
+			ElemTri[j] = UVOverlay->AppendElement((FVector2f)UV);
 			NewUVIndices.Add(ElemTri[j]);
 		}
 		UVOverlay->SetTriangle(TriangleID, ElemTri);
@@ -277,10 +277,10 @@ void FDynamicMeshUVEditor::SetTriangleUVsFromPlanarProjection(
 			{
 				FVector3d Pos = Mesh->GetVertex(BaseTri[j]);
 				FVector3d TransformPos = PointTransform(Pos);
-				FVector2f UV = (FVector2f)ProjectionFrame.ToPlaneUV(TransformPos, 2);
+				FVector2d UV = ProjectionFrame.ToPlaneUV(TransformPos, 2);
 				UV.X *= ScaleX;
 				UV.Y *= ScaleY;
-				ElemTri[j] = UVOverlay->AppendElement(UV);
+				ElemTri[j] = UVOverlay->AppendElement(FVector2f(UV));
 				NewUVIndices.Add(ElemTri[j]);
 				BaseToOverlayVIDMap.Add(BaseTri[j], ElemTri[j]);
 			}
@@ -486,10 +486,10 @@ bool FDynamicMeshUVEditor::SetTriangleUVsFromExpMap(
 	{
 		if (Param.HasUV(vid))
 		{
-			FVector2f UV = (FVector2f)Param.GetUV(vid);
+			FVector2d UV = Param.GetUV(vid);
 			UV.X *= ScaleX;
 			UV.Y *= ScaleY;
-			VtxElementIDs[vid] = UVOverlay->AppendElement(UV);
+			VtxElementIDs[vid] = UVOverlay->AppendElement(FVector2f(UV));
 			NewElementIDs.Add(VtxElementIDs[vid]);
 		}
 	}
@@ -1055,7 +1055,9 @@ void FDynamicMeshUVEditor::SetTriangleUVsFromBoxProjection(
 
 		int MajorAxis = TriBoxInfo.A;
 		int Bucket = TriBoxInfo.B;
-		double MajorAxisSign = FMathd::Sign(N[MajorAxis]);
+		int MajorAxisSign =  (N[MajorAxis] > 0.0) ? 1 : ( (N[MajorAxis] < 0.0) ? -1 : 0 );
+		
+		FMathd::Sign(N[MajorAxis]);
 		int Minor1 = Minor1s[MajorAxis];
 		int Minor2 = Minor2s[MajorAxis];
 
@@ -1071,7 +1073,7 @@ void FDynamicMeshUVEditor::SetTriangleUVsFromBoxProjection(
 				FVector3d BoxPos = BoxFrame.ToFramePoint(TransformPos);
 				BoxPos *= Scale;
 
-				FVector2f UV = ProjAxis(BoxPos, Minor1, Minor2, MajorAxisSign * Minor1Flip[MajorAxis], Minor2Flip[MajorAxis]);
+				FVector2f UV = ProjAxis(BoxPos, Minor1, Minor2, float(MajorAxisSign * Minor1Flip[MajorAxis] ), (float)Minor2Flip[MajorAxis]);
 
 				ElemTri[j] = UVOverlay->AppendElement(UV);
 				NewUVIndices.Add(ElemTri[j]);
@@ -1193,7 +1195,8 @@ void FDynamicMeshUVEditor::SetTriangleUVsFromCylinderProjection(
 				FVector2f UV = FVector2f::Zero();
 				if (Bucket <= 2)
 				{
-					UV = ProjAxis(BoxPos, 0, 1, FMathd::Sign(N[MajorAxis]) * Minor1Flip[MajorAxis], Minor2Flip[MajorAxis]);
+					int MajorAxisSign =  (N[MajorAxis] > 0.0) ? 1 : ( (N[MajorAxis] < 0.0) ? -1 : 0 );
+					UV = ProjAxis(BoxPos, 0, 1, float( MajorAxisSign * Minor1Flip[MajorAxis] ), (float)Minor2Flip[MajorAxis]);
 				}
 				else
 				{
@@ -1268,7 +1271,7 @@ bool FDynamicMeshUVEditor::ScaleUVAreaTo3DArea(const TArray<int32>& Triangles, b
 	for (int32 eid : Elements)
 	{
 		FVector2f UV = UVOverlay->GetElement(eid);
-		UV = (UV - ScaleOrigin) * UVScale + Translation;
+		UV = (UV - ScaleOrigin) * float(UVScale) + Translation;
 		UVOverlay->SetElement(eid, UV);
 	}
 
diff --git a/Engine/Plugins/Runtime/GeometryProcessing/Source/DynamicMesh/Private/Selections/MeshEdgeSelection.cpp b/Engine/Plugins/Runtime/GeometryProcessing/Source/DynamicMesh/Private/Selections/MeshEdgeSelection.cpp
--- a/Engine/Plugins/Runtime/GeometryProcessing/Source/DynamicMesh/Private/Selections/MeshEdgeSelection.cpp
+++ b/Engine/Plugins/Runtime/GeometryProcessing/Source/DynamicMesh/Private/Selections/MeshEdgeSelection.cpp
@@ -6,6 +6,7 @@
 #include "Containers/Set.h"
 #include "Selections/MeshFaceSelection.h"
 #include "Selections/MeshVertexSelection.h"
+#include "Math/UnrealMathUtility.h"
 
 using namespace UE::Geometry;
 
@@ -27,7 +28,7 @@ FMeshEdgeSelection::FMeshEdgeSelection(const FDynamicMesh3* mesh, const FMeshVer
 // convert face selection to edge selection. Require at least minCount tris of edge to be selected
 FMeshEdgeSelection::FMeshEdgeSelection(const FDynamicMesh3* mesh, const FMeshFaceSelection& convertT, int minCount) : Mesh(mesh)
 {
-	minCount = FMathd::Clamp(minCount, 1, 2);
+	minCount = FMath::Clamp(minCount, 1, 2);
 
 	if (minCount == 1)
 	{
diff --git a/Engine/Plugins/Runtime/GeometryProcessing/Source/DynamicMesh/Private/Selections/MeshFaceSelection.cpp b/Engine/Plugins/Runtime/GeometryProcessing/Source/DynamicMesh/Private/Selections/MeshFaceSelection.cpp
--- a/Engine/Plugins/Runtime/GeometryProcessing/Source/DynamicMesh/Private/Selections/MeshFaceSelection.cpp
+++ b/Engine/Plugins/Runtime/GeometryProcessing/Source/DynamicMesh/Private/Selections/MeshFaceSelection.cpp
@@ -3,6 +3,7 @@
 #include "Selections/MeshFaceSelection.h"
 #include "Selections/MeshVertexSelection.h"
 #include "DynamicMesh/DynamicMesh3.h"
+#include "Math/UnrealMathUtility.h"
 
 using namespace UE::Geometry;
 
@@ -10,7 +11,7 @@ using namespace UE::Geometry;
 // tri to be selected (valid values are 1,2,3)
 FMeshFaceSelection::FMeshFaceSelection(const FDynamicMesh3* mesh, const FMeshVertexSelection& convertV, int minCount) : Mesh(mesh)
 {
-	minCount = FMathd::Clamp(minCount, 1, 3);
+	minCount = FMath::Clamp(minCount, 1, 3);
 
 	if (minCount == 1)
 	{
