GeometryProcessing, Mesh Operations.  Fix implicit type conversion warnings.

diff --git a/Engine/Plugins/Runtime/GeometryProcessing/Source/DynamicMesh/Private/Operations/ExtrudeMesh.cpp b/Engine/Plugins/Runtime/GeometryProcessing/Source/DynamicMesh/Private/Operations/ExtrudeMesh.cpp
--- a/Engine/Plugins/Runtime/GeometryProcessing/Source/DynamicMesh/Private/Operations/ExtrudeMesh.cpp
+++ b/Engine/Plugins/Runtime/GeometryProcessing/Source/DynamicMesh/Private/Operations/ExtrudeMesh.cpp
@@ -171,7 +171,7 @@ bool FExtrudeMesh::ApplyExtrude(FExtrusionInfo& Region, FMeshNormals* UseNormals
 
 					if (k > 0)
 					{
-						AccumUVTranslation += Distance(Mesh->GetVertex(BaseLoop.Vertices[k]), Mesh->GetVertex(BaseLoop.Vertices[k - 1]));
+						AccumUVTranslation +=(float) Distance(Mesh->GetVertex(BaseLoop.Vertices[k]), Mesh->GetVertex(BaseLoop.Vertices[k - 1]));
 					}
 
 					// translate horizontally such that vertical spans are adjacent in UV space (so textures tile/wrap properly)
diff --git a/Engine/Plugins/Runtime/GeometryProcessing/Source/DynamicMesh/Private/Operations/InsetMeshRegion.cpp b/Engine/Plugins/Runtime/GeometryProcessing/Source/DynamicMesh/Private/Operations/InsetMeshRegion.cpp
--- a/Engine/Plugins/Runtime/GeometryProcessing/Source/DynamicMesh/Private/Operations/InsetMeshRegion.cpp
+++ b/Engine/Plugins/Runtime/GeometryProcessing/Source/DynamicMesh/Private/Operations/InsetMeshRegion.cpp
@@ -185,7 +185,7 @@ bool FInsetMeshRegion::ApplyInset(FInsetInfo& Region, FMeshNormals* UseNormals)
 			int32 BaseEdgeID = Mesh->FindEdge(BaseLoopV[k], BaseLoopV[(k + 1) % NumLoopV]);
 			int32 BaseGroupID = (BaseEdgeID >= 0) ? Mesh->GetTriangleGroup(Mesh->GetEdgeT(BaseEdgeID).A) : InsetGroupID;
 
-			TPair<int32, int32> GroupPair(FMathd::Min(BaseGroupID, InsetGroupID), FMathd::Max(BaseGroupID, InsetGroupID));
+			TPair<int32, int32> GroupPair(FMath::Min(BaseGroupID, InsetGroupID), FMath::Max(BaseGroupID, InsetGroupID));
 			if (NewGroupsMap.Contains(GroupPair) == false)
 			{
 				int32 NewGroupID = Mesh->AllocateTriangleGroup();
@@ -359,7 +359,7 @@ bool FInsetMeshRegion::ApplyInset(FInsetInfo& Region, FMeshNormals* UseNormals)
 
 				if (k > 0)
 				{
-					AccumUVTranslation += Distance(Mesh->GetVertex(BaseLoopV[k]), Mesh->GetVertex(BaseLoopV[k - 1]));
+					AccumUVTranslation += (float)Distance(Mesh->GetVertex(BaseLoopV[k]), Mesh->GetVertex(BaseLoopV[k - 1]));
 				}
 
 				// translate horizontally such that vertical spans are adjacent in UV space (so textures tile/wrap properly)
diff --git a/Engine/Plugins/Runtime/GeometryProcessing/Source/DynamicMesh/Private/Operations/MeshBoolean.cpp b/Engine/Plugins/Runtime/GeometryProcessing/Source/DynamicMesh/Private/Operations/MeshBoolean.cpp
--- a/Engine/Plugins/Runtime/GeometryProcessing/Source/DynamicMesh/Private/Operations/MeshBoolean.cpp
+++ b/Engine/Plugins/Runtime/GeometryProcessing/Source/DynamicMesh/Private/Operations/MeshBoolean.cpp
@@ -948,7 +948,7 @@ bool FMeshBoolean::CollapseWouldChangeShapeOrUVs(
 				{
 					return;
 				}
-				float LerpT = (RemoveVPos - OtherVPos).Dot(OtherEdgeDir) / (KeepVPos - OtherVPos).Dot(OtherEdgeDir);
+				float LerpT = float( (RemoveVPos - OtherVPos).Dot(OtherEdgeDir) / (KeepVPos - OtherVPos).Dot(OtherEdgeDir) );
 				if (bPreserveVertexUVs && Mesh.HasVertexUVs())
 				{
 					FVector2f OtherUV = Mesh.GetVertexUV(OtherV);
diff --git a/Engine/Plugins/Runtime/GeometryProcessing/Source/DynamicMesh/Private/Operations/MeshMeshCut.cpp b/Engine/Plugins/Runtime/GeometryProcessing/Source/DynamicMesh/Private/Operations/MeshMeshCut.cpp
--- a/Engine/Plugins/Runtime/GeometryProcessing/Source/DynamicMesh/Private/Operations/MeshMeshCut.cpp
+++ b/Engine/Plugins/Runtime/GeometryProcessing/Source/DynamicMesh/Private/Operations/MeshMeshCut.cpp
@@ -525,7 +525,7 @@ namespace MeshCut
 
 		int ClosestEdge(FIndex2i EIDs, const FVector3d& Pos)
 		{
-			double BestIdx = -1;
+			int BestIdx = -1;
 			double BestDSq = FMathd::MaxReal;
 			for (int Idx = 0; Idx < 2; Idx++)
 			{
@@ -544,7 +544,7 @@ namespace MeshCut
 
 		int OnEdge(FIndex3i EIDs, const FVector3d& Pos, double BestDSq)
 		{
-			double BestIdx = -1;
+			int BestIdx = -1;
 			for (int Idx = 0; Idx < 3; Idx++)
 			{
 				int EID = EIDs[Idx];
diff --git a/Engine/Plugins/Runtime/GeometryProcessing/Source/DynamicMesh/Private/Operations/OffsetMeshRegion.cpp b/Engine/Plugins/Runtime/GeometryProcessing/Source/DynamicMesh/Private/Operations/OffsetMeshRegion.cpp
--- a/Engine/Plugins/Runtime/GeometryProcessing/Source/DynamicMesh/Private/Operations/OffsetMeshRegion.cpp
+++ b/Engine/Plugins/Runtime/GeometryProcessing/Source/DynamicMesh/Private/Operations/OffsetMeshRegion.cpp
@@ -392,7 +392,7 @@ bool FOffsetMeshRegion::ApplyOffset(FOffsetInfo& Region, FMeshNormals* UseNormal
 
 				if (k > 0)
 				{
-					AccumUVTranslation += Distance(Mesh->GetVertex(BaseLoopV[k]), Mesh->GetVertex(BaseLoopV[k - 1]));
+					AccumUVTranslation += (float)Distance(Mesh->GetVertex(BaseLoopV[k]), Mesh->GetVertex(BaseLoopV[k - 1]));
 				}
 
 				// translate horizontally such that vertical spans are adjacent in UV space (so textures tile/wrap properly)
diff --git a/Engine/Plugins/Runtime/GeometryProcessing/Source/DynamicMesh/Private/Operations/PNTriangles.cpp b/Engine/Plugins/Runtime/GeometryProcessing/Source/DynamicMesh/Private/Operations/PNTriangles.cpp
--- a/Engine/Plugins/Runtime/GeometryProcessing/Source/DynamicMesh/Private/Operations/PNTriangles.cpp
+++ b/Engine/Plugins/Runtime/GeometryProcessing/Source/DynamicMesh/Private/Operations/PNTriangles.cpp
@@ -81,7 +81,7 @@ namespace FPNTrianglesLocals
 			FVector3f Normal;
 			if (FMath::IsNearlyZero(Divisor)) 
 			{
-				Normal = Normal21 / 2.0; // If edge is degenerate then simply interpolate between the normals
+				Normal = 0.5f * Normal21; // If edge is degenerate then simply interpolate between the normals
 			} 
 			else
 			{ 
diff --git a/Engine/Plugins/Runtime/GeometryProcessing/Source/DynamicMesh/Private/Operations/SimpleHoleFiller.cpp b/Engine/Plugins/Runtime/GeometryProcessing/Source/DynamicMesh/Private/Operations/SimpleHoleFiller.cpp
--- a/Engine/Plugins/Runtime/GeometryProcessing/Source/DynamicMesh/Private/Operations/SimpleHoleFiller.cpp
+++ b/Engine/Plugins/Runtime/GeometryProcessing/Source/DynamicMesh/Private/Operations/SimpleHoleFiller.cpp
@@ -135,7 +135,7 @@ bool FSimpleHoleFiller::UpdateAttributes(TArray<FMeshRegionBoundaryLoops::VidOve
 			{
 				NewVertUV += UVMap[Vid].Value;
 			}
-			NewVertUV /= Loop.Vertices.Num();
+			NewVertUV /= (float)Loop.Vertices.Num();
 			UVMap.Add(NewVertex, FMeshRegionBoundaryLoops::ElementIDAndValue<FVector2f>(
 				IndexConstants::InvalidID, NewVertUV));
 		}
diff --git a/Engine/Plugins/Runtime/GeometryProcessing/Source/DynamicMesh/Private/Operations/UniformTessellate.cpp b/Engine/Plugins/Runtime/GeometryProcessing/Source/DynamicMesh/Private/Operations/UniformTessellate.cpp
--- a/Engine/Plugins/Runtime/GeometryProcessing/Source/DynamicMesh/Private/Operations/UniformTessellate.cpp
+++ b/Engine/Plugins/Runtime/GeometryProcessing/Source/DynamicMesh/Private/Operations/UniformTessellate.cpp
@@ -275,7 +275,7 @@ namespace UniformTessellateLocals
 
 		void LerpElements(const RealType* Element1, const RealType* Element2, RealType* OutElement, RealType Alpha) 
 		{
-			Alpha = FMath::Clamp(Alpha, 0.0, 1.0);
+			Alpha = FMath::Clamp(Alpha, RealType(0), RealType(1));
 			RealType OneMinusAlpha = (RealType)1 - Alpha;
 
 			for (int Idx = 0; Idx < ElementSize; ++Idx) 
@@ -393,7 +393,7 @@ namespace UniformTessellateLocals
 					RealType OutElement[ElementSize];
 					for (int VertexOffset = 0; VertexOffset < NumNewLevelVertices; ++VertexOffset) 
 					{
-						const RealType Alpha = (RealType)(VertexOffset + 1) / (NumNewLevelVertices + 1);
+						const RealType Alpha = RealType(VertexOffset + 1) / (NumNewLevelVertices + 1);
 						LerpElements(Element1, Element2, OutElement, Alpha);
 						this->SetElementOnTriangle(TriangleID, ElementIDCounter, OutElement);
 						ElementIDCounter++;
@@ -1609,7 +1609,7 @@ namespace UniformTessellateLocals
 		void LerpWeights(const FBoneWeights& Weight1, const FBoneWeights& Weight2, FBoneWeights& OutWeight, double Alpha) 
 		{
 			Alpha = FMath::Clamp(Alpha, 0.0, 1.0);
-			OutWeight = FBoneWeights::Blend(Weight1, Weight2, Alpha);
+			OutWeight = FBoneWeights::Blend(Weight1, Weight2, (float)Alpha);
 		} 
 
 		bool GenerateEdgeElements() 
