GeometryProcessing: Add additional vertex based connected component discovery methods to the MeshConnectedComponents class. When used, the resulting component indices are representing mesh vertices instead of mesh triangles.

diff --git a/Engine/Plugins/Runtime/GeometryProcessing/Source/DynamicMesh/Private/Selections/MeshConnectedComponents.cpp b/Engine/Plugins/Runtime/GeometryProcessing/Source/DynamicMesh/Private/Selections/MeshConnectedComponents.cpp
--- a/Engine/Plugins/Runtime/GeometryProcessing/Source/DynamicMesh/Private/Selections/MeshConnectedComponents.cpp
+++ b/Engine/Plugins/Runtime/GeometryProcessing/Source/DynamicMesh/Private/Selections/MeshConnectedComponents.cpp
@@ -5,21 +5,36 @@
 
 using namespace UE::Geometry;
 
+namespace ConnectedComponentsLocals
+{
+	// We'll use these enums below instead of raw integers to provide some clarity in the code.
+	// Note: To avoid changing the existing public API, the ActiveSet will remain a uint8 array,
+	// rather than an array of this enum, even though that would be a more type safe solution.
+	enum class EProcessingState : uint8
+	{
+		Unprocessed = 0,
+		InQueue = 1,
+		Done = 2,
+		Invalid = 255
+	};
+}
+
+using namespace ConnectedComponentsLocals;
+
 void FMeshConnectedComponents::FindConnectedTriangles(TFunction<bool(int32, int32)> TrisConnectedPredicate)
 {
 	TRACE_CPUPROFILER_EVENT_SCOPE(MeshConnectedComponents_FindConnectedTriangles);
 	
 	// initial active set contains all valid triangles
-	// active values are as follows:   0: unprocessed,  1: in queue,  2: done,  3: invalid
 	TArray<uint8> ActiveSet;
 	int NumTriangles = Mesh->MaxTriangleID();
-	ActiveSet.Init(255, NumTriangles);
+	ActiveSet.Init((uint8)EProcessingState::Invalid, NumTriangles);
 	FInterval1i ActiveRange = FInterval1i::Empty();
 	for (int tid = 0; tid < NumTriangles; ++tid)
 	{
 		if (Mesh->IsTriangle(tid))
 		{
-			ActiveSet[tid] = 0;
+			ActiveSet[tid] = (uint8)EProcessingState::Unprocessed;
 			ActiveRange.Contain(tid);
 		}
 	}
@@ -34,16 +49,15 @@ void FMeshConnectedComponents::FindConnectedTriangles(TFunctionRef<bool(int)> In
 	TRACE_CPUPROFILER_EVENT_SCOPE(MeshConnectedComponents_FindConnectedTriangles_Filtered);
 	
 	// initial active set contains all valid triangles
-	// active values are as follows:   0: unprocessed,  1: in queue,  2: done,  3: invalid
 	TArray<uint8> ActiveSet;
 	int NumTriangles = Mesh->MaxTriangleID();
-	ActiveSet.Init(255, NumTriangles);
+	ActiveSet.Init((uint8)EProcessingState::Invalid, NumTriangles);
 	FInterval1i ActiveRange = FInterval1i::Empty();
 	for (int tid = 0; tid < NumTriangles; ++tid)
 	{
 		if (Mesh->IsTriangle(tid) && IndexFilterFunc(tid))
 		{
-			ActiveSet[tid] = 0;
+			ActiveSet[tid] = (uint8)EProcessingState::Unprocessed;
 			ActiveRange.Contain(tid);
 		}
 	}
@@ -58,16 +72,15 @@ void FMeshConnectedComponents::FindConnectedTriangles(const TArray<int>& Triangl
 	TRACE_CPUPROFILER_EVENT_SCOPE(MeshConnectedComponents_FindConnectedTriangles_TriangleROI);
 
 	// initial active set contains all valid triangles
-	// active values are as follows:   0: unprocessed,  1: in queue,  2: done,  3: invalid
 	TArray<uint8> ActiveSet;
 	int NumTriangles = Mesh->MaxTriangleID();
-	ActiveSet.Init(255, NumTriangles);
+	ActiveSet.Init((uint8)EProcessingState::Invalid, NumTriangles);
 	FInterval1i ActiveRange = FInterval1i::Empty();
 	for (int tid : TriangleROI)
 	{
 		if (Mesh->IsTriangle(tid))
 		{
-			ActiveSet[tid] = 0;
+			ActiveSet[tid] = (uint8)EProcessingState::Unprocessed;
 			ActiveRange.Contain(tid);
 		}
 	}
@@ -83,15 +96,14 @@ void FMeshConnectedComponents::FindTrianglesConnectedToSeeds(const TArray<int>&
 	TRACE_CPUPROFILER_EVENT_SCOPE(MeshConnectedComponents_FindConnectedTriangles_SeedTriangles);
 	
 	// initial active set contains all valid triangles
-	// active values are as follows:   0: unprocessed,  1: in queue,  2: done,  3: invalid
 	TArray<uint8> ActiveSet;
 	int32 NumTriangles = Mesh->MaxTriangleID();
-	ActiveSet.Init(255, NumTriangles);
+	ActiveSet.Init((uint8)EProcessingState::Invalid, NumTriangles);
 	for (int32 tid = 0; tid < NumTriangles; ++tid)
 	{
 		if (Mesh->IsTriangle(tid))
 		{
-			ActiveSet[tid] = 0;
+			ActiveSet[tid] = (uint8)EProcessingState::Unprocessed;
 		}
 	}
 
@@ -99,7 +111,94 @@ void FMeshConnectedComponents::FindTrianglesConnectedToSeeds(const TArray<int>&
 }
 
 
+void FMeshConnectedComponents::FindConnectedVertices(TFunction<bool(int32, int32)> VertsConnectedPredicate)
+{
+	TRACE_CPUPROFILER_EVENT_SCOPE(MeshConnectedComponents_FindConnectedVertices);
+
+	// initial active set contains all valid vertices
+	TArray<uint8> ActiveSet;
+	int NumVertices = Mesh->MaxVertexID();
+	ActiveSet.Init((uint8)EProcessingState::Invalid, NumVertices);
+	FInterval1i ActiveRange = FInterval1i::Empty();
+	for (int vid = 0; vid < NumVertices; ++vid)
+	{
+		if (Mesh->IsVertex(vid))
+		{
+			ActiveSet[vid] = (uint8)EProcessingState::Unprocessed;
+			ActiveRange.Contain(vid);
+		}
+	}
+
+	FindVertComponents(ActiveRange, ActiveSet, VertsConnectedPredicate);
+}
+
+
+
+void FMeshConnectedComponents::FindConnectedVertices(TFunctionRef<bool(int)> IndexFilterFunc, TFunction<bool(int32, int32)> VertsConnectedPredicate)
+{
+	TRACE_CPUPROFILER_EVENT_SCOPE(MeshConnectedComponents_FindConnectedVertices_Filtered);
+
+	// initial active set contains all valid vertices
+	TArray<uint8> ActiveSet;
+	int NumVertices = Mesh->MaxVertexID();
+	ActiveSet.Init((uint8)EProcessingState::Invalid, NumVertices);
+	FInterval1i ActiveRange = FInterval1i::Empty();
+	for (int vid = 0; vid < NumVertices; ++vid)
+	{
+		if (Mesh->IsVertex(vid) && IndexFilterFunc(vid))
+		{
+			ActiveSet[vid] = (uint8)EProcessingState::Unprocessed;
+			ActiveRange.Contain(vid);
+		}
+	}
+
+	FindVertComponents(ActiveRange, ActiveSet, VertsConnectedPredicate);
+}
+
+
+
+void FMeshConnectedComponents::FindConnectedVertices(const TArray<int>& VertexROI, TFunction<bool(int32, int32)> VertsConnectedPredicate)
+{
+	TRACE_CPUPROFILER_EVENT_SCOPE(MeshConnectedComponents_FindConnectedVertices_VertexROI);
+
+	// initial active set contains all valid vertices
+	TArray<uint8> ActiveSet;
+	int NumVertices = Mesh->MaxVertexID();
+	ActiveSet.Init((uint8)EProcessingState::Invalid, NumVertices);
+	FInterval1i ActiveRange = FInterval1i::Empty();
+	for (int vid : VertexROI)
+	{
+		if (Mesh->IsVertex(vid))
+		{
+			ActiveSet[vid] = (uint8)EProcessingState::Unprocessed;
+			ActiveRange.Contain(vid);
+		}
+	}
+
+	FindVertComponents(ActiveRange, ActiveSet, VertsConnectedPredicate);
+}
+
+
+
+
+void FMeshConnectedComponents::FindVerticesConnectedToSeeds(const TArray<int>& SeedVertices, TFunction<bool(int32, int32)> VertsConnectedPredicate)
+{
+	TRACE_CPUPROFILER_EVENT_SCOPE(MeshConnectedComponents_FindConnectedVertices_SeedVertices);
+
+	// initial active set contains all valid vertices
+	TArray<uint8> ActiveSet;
+	int NumVertices = Mesh->MaxVertexID();
+	ActiveSet.Init((uint8)EProcessingState::Invalid, NumVertices);
+	for (int32 vid = 0; vid < NumVertices; ++vid)
+	{
+		if (Mesh->IsVertex(vid))
+		{
+			ActiveSet[vid] = (uint8)EProcessingState::Unprocessed;
+		}
+	}
 
+	FindVertComponents(SeedVertices, ActiveSet, VertsConnectedPredicate);
+}
 
 
 void FMeshConnectedComponents::FindTriComponents(FInterval1i ActiveRange, TArray<uint8>& ActiveSet, TFunction<bool(int32, int32)> TrisConnectedPredicate)
@@ -116,11 +215,11 @@ void FMeshConnectedComponents::FindTriComponents(FInterval1i ActiveRange, TArray
 	// until we are done
 	for (int i = ActiveRange.Min; i <= ActiveRange.Max; i++)
 	{
-		if (ActiveSet[i] != 255)
+		if (ActiveSet[i] != (uint8)EProcessingState::Invalid)
 		{
 			int SeedTri = i;
 			ComponentQueue.Add(SeedTri);
-			ActiveSet[SeedTri] = 1;      // in ComponentQueue
+			ActiveSet[SeedTri] = (uint8)EProcessingState::InQueue;      // in ComponentQueue
 
 			FComponent* Component = new FComponent();
 			if (TrisConnectedPredicate)
@@ -157,10 +256,10 @@ void FMeshConnectedComponents::FindTriComponents(const TArray<int32>& SeedList,
 	// until we are done
 	for ( int32 SeedTri : SeedList )
 	{
-		if (ActiveSet[SeedTri] != 255)
+		if (ActiveSet[SeedTri] != (uint8)EProcessingState::Invalid)
 		{
 			ComponentQueue.Add(SeedTri);
-			ActiveSet[SeedTri] = 1;      // in ComponentQueue
+			ActiveSet[SeedTri] = (uint8)EProcessingState::InQueue;      // in ComponentQueue
 
 			FComponent* Component = new FComponent();
 			if (TrisConnectedPredicate)
@@ -193,17 +292,17 @@ void FMeshConnectedComponents::FindTriComponent(FComponent* Component, TArray<in
 	{
 		int32 CurTriangle = ComponentQueue.Pop(false);
 
-		ActiveSet[CurTriangle] = 2;   // tri has been processed
+		ActiveSet[CurTriangle] = (uint8)EProcessingState::Done;   // tri has been processed
 		Component->Indices.Add(CurTriangle);
 
 		FIndex3i TriNbrTris = Mesh->GetTriNeighbourTris(CurTriangle);
 		for (int j = 0; j < 3; ++j)
 		{
 			int NbrTri = TriNbrTris[j];
-			if (NbrTri != FDynamicMesh3::InvalidID && ActiveSet[NbrTri] == 0)
+			if (NbrTri != FDynamicMesh3::InvalidID && ActiveSet[NbrTri] == (uint8)EProcessingState::Unprocessed)
 			{
 				ComponentQueue.Add(NbrTri);
-				ActiveSet[NbrTri] = 1;           // in ComponentQueue
+				ActiveSet[NbrTri] = (uint8)EProcessingState::InQueue;           // in ComponentQueue
 			}
 		}
 	}
@@ -219,17 +318,17 @@ void FMeshConnectedComponents::FindTriComponent(FComponent* Component, TArray<in
 	{
 		int32 CurTriangle = ComponentQueue.Pop(false);
 
-		ActiveSet[CurTriangle] = 2;   // tri has been processed
+		ActiveSet[CurTriangle] = (uint8)EProcessingState::Done;   // tri has been processed
 		Component->Indices.Add(CurTriangle);
 
 		FIndex3i TriNbrTris = Mesh->GetTriNeighbourTris(CurTriangle);
 		for (int j = 0; j < 3; ++j)
 		{
 			int NbrTri = TriNbrTris[j];
-			if (NbrTri != FDynamicMesh3::InvalidID && ActiveSet[NbrTri] == 0 && TriConnectedPredicate(CurTriangle, NbrTri))
+			if (NbrTri != FDynamicMesh3::InvalidID && ActiveSet[NbrTri] == (uint8)EProcessingState::Unprocessed && TriConnectedPredicate(CurTriangle, NbrTri))
 			{
 				ComponentQueue.Add(NbrTri);
-				ActiveSet[NbrTri] = 1;           // in ComponentQueue
+				ActiveSet[NbrTri] = (uint8)EProcessingState::InQueue;           // in ComponentQueue
 			}
 		}
 	}
@@ -242,10 +341,135 @@ void FMeshConnectedComponents::RemoveFromActiveSet(const FComponent* Component,
 	int32 ComponentSize = Component->Indices.Num();
 	for (int32 j = 0; j < ComponentSize; ++j)
 	{
-		ActiveSet[Component->Indices[j]] = 255;
+		ActiveSet[Component->Indices[j]] = (uint8)EProcessingState::Invalid;
+	}
+}
+
+
+void FMeshConnectedComponents::FindVertComponents(FInterval1i ActiveRange, TArray<uint8>& ActiveSet, TFunction<bool(int32, int32)> VertsConnectedPredicate)
+{
+	TRACE_CPUPROFILER_EVENT_SCOPE(FindVertComponents_ActiveRange);
+
+	Components.Empty();
+
+	// temporary queue
+	TArray<int32> ComponentQueue;
+	ComponentQueue.Reserve(256);
+
+	// keep finding valid seed vertices and growing connected components
+	// until we are done
+	for (int i = ActiveRange.Min; i <= ActiveRange.Max; i++)
+	{
+		if (ActiveSet[i] != (uint8)EProcessingState::Invalid)
+		{
+			int SeedVert = i;
+			ComponentQueue.Add(SeedVert);
+			ActiveSet[SeedVert] = (uint8)EProcessingState::InQueue;      // in ComponentQueue
+
+			FComponent* Component = new FComponent();
+			if (VertsConnectedPredicate)
+			{
+				FindVertComponent(Component, ComponentQueue, ActiveSet, VertsConnectedPredicate);
+			}
+			else
+			{
+				FindVertComponent(Component, ComponentQueue, ActiveSet);
+			}
+			Components.Add(Component);
+
+			RemoveFromActiveSet(Component, ActiveSet);
+
+			ComponentQueue.Reset(0);
+		}
+	}
+}
+
+void FMeshConnectedComponents::FindVertComponents(const TArray<int32>& SeedList, TArray<uint8>& ActiveSet, TFunction<bool(int32, int32)> VertsConnectedPredicate)
+{
+	TRACE_CPUPROFILER_EVENT_SCOPE(FindVertComponents_SeedList);
+
+	Components.Empty();
+
+	// temporary queue
+	TArray<int32> ComponentQueue;
+	ComponentQueue.Reserve(256);
+
+	// keep finding valid seed vertices and growing connected components
+	// until we are done
+	for (int32 SeedVert : SeedList)
+	{
+		if (ActiveSet[SeedVert] != (uint8)EProcessingState::Invalid)
+		{
+			ComponentQueue.Add(SeedVert);
+			ActiveSet[SeedVert] = (uint8)EProcessingState::InQueue;      // in ComponentQueue
+
+			FComponent* Component = new FComponent();
+			if (VertsConnectedPredicate)
+			{
+				FindVertComponent(Component, ComponentQueue, ActiveSet, VertsConnectedPredicate);
+			}
+			else
+			{
+				FindVertComponent(Component, ComponentQueue, ActiveSet);
+			}
+			Components.Add(Component);
+
+			RemoveFromActiveSet(Component, ActiveSet);
+
+			ComponentQueue.Reset(0);
+		}
+	}
+
+}
+
+void FMeshConnectedComponents::FindVertComponent(FComponent* Component, TArray<int32>& ComponentQueue, TArray<uint8>& ActiveSet)
+{
+	TRACE_CPUPROFILER_EVENT_SCOPE(FindVertComponents_Component);
+
+	while (ComponentQueue.Num() > 0)
+	{
+		int32 CurVert = ComponentQueue.Pop(false);
+
+		ActiveSet[CurVert] = (uint8)EProcessingState::Done;   // Vert has been processed
+		Component->Indices.Add(CurVert);
+
+		auto ProcessOneRing = [&ActiveSet, &ComponentQueue](int32 NbrVert) mutable
+		{
+			if (NbrVert != FDynamicMesh3::InvalidID && ActiveSet[NbrVert] == (uint8)EProcessingState::Unprocessed)
+			{
+				ComponentQueue.Add(NbrVert);
+				ActiveSet[NbrVert] = (uint8)EProcessingState::InQueue;           // in ComponentQueue
+			}
+		};
+
+		Mesh->EnumerateVertexVertices(CurVert, ProcessOneRing);
 	}
 }
 
+void FMeshConnectedComponents::FindVertComponent(FComponent* Component, TArray<int32>& ComponentQueue, TArray<uint8>& ActiveSet,
+	 TFunctionRef<bool(int32, int32)> VertsConnectedPredicate)
+{
+	TRACE_CPUPROFILER_EVENT_SCOPE(FindVertComponents_Component_Filtered);
+
+	while (ComponentQueue.Num() > 0)
+	{
+		int32 CurVert = ComponentQueue.Pop(false);
+
+		ActiveSet[CurVert] = (uint8)EProcessingState::Done;   // Vert has been processed
+		Component->Indices.Add(CurVert);
+
+		auto ProcessOneRing = [&CurVert, &ActiveSet, &VertsConnectedPredicate, &ComponentQueue](int32 NbrVert) mutable
+		{
+			if (NbrVert != FDynamicMesh3::InvalidID && ActiveSet[NbrVert] == (uint8)EProcessingState::Unprocessed && VertsConnectedPredicate(CurVert, NbrVert))
+			{
+				ComponentQueue.Add(NbrVert);
+				ActiveSet[NbrVert] = (uint8)EProcessingState::InQueue;           // in ComponentQueue
+			}
+		};
+
+		Mesh->EnumerateVertexVertices(CurVert, ProcessOneRing);
+	}
+}
 
 
 
@@ -367,7 +591,82 @@ bool FMeshConnectedComponents::InitializeFromTriangleComponents(TArray<TArray<in
 	return true;
 }
 
+bool FMeshConnectedComponents::InitializeFromVertexComponents(const TArray<TArray<int32>>& ComponentLists, bool bValidateIDs)
+{
+	if (bValidateIDs)
+	{
+		for (const TArray<int32>& ComponentList : ComponentLists)
+		{
+			for (int32 Vid : ComponentList)
+			{
+				if (Mesh->IsVertex(Vid) == false)
+				{
+					return false;
+				}
+			}
+		}
+	}
+
+	Components.Reset();
+	Components.Reserve(ComponentLists.Num());
+	for (const TArray<int32>& ComponentList : ComponentLists)
+	{
+		// skip any empty lists
+		if (ComponentList.Num() == 0)
+		{
+			continue;
+		}
+
+		FComponent* Component = new FComponent();
+		Component->Indices.Append(ComponentList);
+		Components.Add(Component);
+	}
+
+	return true;
+}
+
+
+bool FMeshConnectedComponents::InitializeFromVertexComponents(TArray<TArray<int32>>& ComponentLists, bool bMoveSubLists, bool bValidateIDs)
+{
+	if (bValidateIDs)
+	{
+		for (const TArray<int32>& ComponentList : ComponentLists)
+		{
+			for (int32 Vid : ComponentList)
+			{
+				if (Mesh->IsVertex(Vid) == false)
+				{
+					return false;
+				}
+			}
+		}
+	}
 
+	Components.Reset();
+	Components.Reserve(ComponentLists.Num());
+	for (TArray<int32>& ComponentList : ComponentLists)
+	{
+		// skip any empty lists
+		if (ComponentList.Num() == 0)
+		{
+			continue;
+		}
+
+		FComponent* Component = new FComponent();
+		if (bMoveSubLists)
+		{
+			Component->Indices = MoveTemp(ComponentList);
+		}
+		else
+		{
+			Component->Indices.Append(ComponentList);
+		}
+
+		Components.Add(Component);
+	}
+
+	return true;
+}
 
 
 
diff --git a/Engine/Plugins/Runtime/GeometryProcessing/Source/DynamicMesh/Public/Selections/MeshConnectedComponents.h b/Engine/Plugins/Runtime/GeometryProcessing/Source/DynamicMesh/Public/Selections/MeshConnectedComponents.h
--- a/Engine/Plugins/Runtime/GeometryProcessing/Source/DynamicMesh/Public/Selections/MeshConnectedComponents.h
+++ b/Engine/Plugins/Runtime/GeometryProcessing/Source/DynamicMesh/Public/Selections/MeshConnectedComponents.h
@@ -96,6 +96,49 @@ public:
 	*/
 	bool InitializeFromTriangleComponents(TArray<TArray<int32>>& ComponentLists, bool bMoveSubLists, bool bValidateIDs);
 
+	/**
+	* Find all connected vertex components of the Mesh and store in Components array.
+	* @param VertsConnectedPredicate optional function that specifies whether two edge-connected vertices should be considered connected by the search
+	*/
+	void FindConnectedVertices(TFunction<bool(int32, int32)> VertsConnectedPredicate = nullptr);
+
+	/**
+	* Find all connected vertex components of a subset of vertices of the Mesh and store in Components array.
+	* @param VertexROI list of vertices to search across
+	* @param VertsConnectedPredicate optional function that specifies whether two edge-connected vertices should be considered connected by the search
+	*/
+	void FindConnectedVertices(const TArray<int>& VertexROI, TFunction<bool(int32, int32)> VertsConnectedPredicate = nullptr);
+
+	/**
+	* Find all connected vertex components of a subset of vertices of the Mesh and store in Components array.
+	* @param IndexFilterFunc defines set of vertices to search across, return true for vertex IDs that are to be considered
+	* @param VertsConnectedPredicate optional function that specifies whether two edge-connected vertices should be considered connected by the search
+	*/
+	void FindConnectedVertices(TFunctionRef<bool(int)> IndexFilterFunc, TFunction<bool(int32, int32)>VertsConnectedPredicate = nullptr);
+
+	/**
+	* Find all connected vertex components that contain one or more Seed Vertices and store in Components array.
+	 Search only starts from Seed Vertices.
+	* @param SeedVertices list of start vertices, each component contains at least one of these vertices
+	* @param VertsConnectedPredicate optional function that specifies whether two edge-connected vertices should be considered connected by the search
+	*/
+	void FindVerticesConnectedToSeeds(const TArray<int>& SeedVertices, TFunction<bool(int32, int32)> VertsConnectedPredicate = nullptr);
+
+	/**
+	 * Initialize the internal FComponent list from the input ComponentLists, skipping any empty input lists
+	 * @param bValidateIDs if true, test that each value corresponds to a valid vertex ID on the Mesh
+	 * @return true if all IDs are valid, or if check was skipped
+	 */
+	bool InitializeFromVertexComponents(const TArray<TArray<int32>>& ComponentLists, bool bValidateIDs);
+
+	/**
+	* Initialize the internal FComponent list from the input ComponentLists, skipping any empty input lists
+	* @param bMoveSubLists if true, steal the arrays inside the ComponentLists (via MoveTemp), to avoid memory copies
+	* @param bValidateIDs if true, test that each value corresponds to a valid vertex ID on the Mesh
+	* @return true if all IDs are valid, or if check was skipped
+	*/
+	bool InitializeFromVertexComponents(TArray<TArray<int32>>& ComponentLists, bool bMoveSubLists, bool bValidateIDs);
+
 	//
 	// Query functions. Only valid to call after a Calculation function has been called.
 	//
@@ -153,7 +196,10 @@ protected:
 	void FindTriComponent(FComponent* Component, TArray<int32>& ComponentQueue, TArray<uint8>& ActiveSet, TFunctionRef<bool(int32, int32)> TriConnectedPredicate);
 	void RemoveFromActiveSet(const FComponent* Component, TArray<uint8>& ActiveSet);
 
-
+	void FindVertComponents(FInterval1i ActiveRange, TArray<uint8>& ActiveSet, TFunction<bool(int32, int32)> VertsConnectedPredicate);
+	void FindVertComponents(const TArray<int32>& SeedList, TArray<uint8>& ActiveSet, TFunction<bool(int32, int32)> VertsConnectedPredicate);
+	void FindVertComponent(FComponent* Component, TArray<int32>& ComponentQueue, TArray<uint8>& ActiveSet);
+	void FindVertComponent(FComponent* Component, TArray<int32>& ComponentQueue, TArray<uint8>& ActiveSet, TFunctionRef<bool(int32, int32)> VertsConnectedPredicate);
 
 public:
 	/**
