GeometryProcesssing: Fix incorrect test case in MeshConnectedComponents unit test and make mesh connected components guard against bad seeds.

diff --git a/Engine/Plugins/Runtime/GeometryProcessing/Source/DynamicMesh/Private/Selections/MeshConnectedComponents.cpp b/Engine/Plugins/Runtime/GeometryProcessing/Source/DynamicMesh/Private/Selections/MeshConnectedComponents.cpp
--- a/Engine/Plugins/Runtime/GeometryProcessing/Source/DynamicMesh/Private/Selections/MeshConnectedComponents.cpp
+++ b/Engine/Plugins/Runtime/GeometryProcessing/Source/DynamicMesh/Private/Selections/MeshConnectedComponents.cpp
@@ -256,7 +256,7 @@ void FMeshConnectedComponents::FindTriComponents(const TArray<int32>& SeedList,
 	// until we are done
 	for ( int32 SeedTri : SeedList )
 	{
-		if (ActiveSet[SeedTri] != (uint8)EProcessingState::Invalid)
+		if (ActiveSet.IsValidIndex(SeedTri) && ActiveSet[SeedTri] != (uint8)EProcessingState::Invalid)
 		{
 			ComponentQueue.Add(SeedTri);
 			ActiveSet[SeedTri] = (uint8)EProcessingState::InQueue;      // in ComponentQueue
@@ -398,7 +398,7 @@ void FMeshConnectedComponents::FindVertComponents(const TArray<int32>& SeedList,
 	// until we are done
 	for (int32 SeedVert : SeedList)
 	{
-		if (ActiveSet[SeedVert] != (uint8)EProcessingState::Invalid)
+		if (ActiveSet.IsValidIndex(SeedVert) && ActiveSet[SeedVert] != (uint8)EProcessingState::Invalid)
 		{
 			ComponentQueue.Add(SeedVert);
 			ActiveSet[SeedVert] = (uint8)EProcessingState::InQueue;      // in ComponentQueue
