Added methods for handling the tessellation of the overlay data.
Expanded FTessellationData class and added FOverlayTessellationData child class to be able to handle overlay seam edges.

#rnx

diff --git a/Engine/Plugins/Experimental/GeometryScripting/Source/GeometryScriptingCore/Private/MeshSubdivideFunctions.cpp b/Engine/Plugins/Experimental/GeometryScripting/Source/GeometryScriptingCore/Private/MeshSubdivideFunctions.cpp
--- a/Engine/Plugins/Experimental/GeometryScripting/Source/GeometryScriptingCore/Private/MeshSubdivideFunctions.cpp
+++ b/Engine/Plugins/Experimental/GeometryScripting/Source/GeometryScriptingCore/Private/MeshSubdivideFunctions.cpp
@@ -158,8 +158,9 @@ UDynamicMesh* UGeometryScriptLibrary_MeshSubdivideFunctions::ApplyAdaptiveTessel
 	} 
 
 	TargetMesh->EditMesh([&](FDynamicMesh3& EditMesh) 
-	{
-		FAdaptiveTessellate Tessellator(&EditMesh);
+	{	
+		FDynamicMesh3 TessellatedMesh;
+		FAdaptiveTessellate Tessellator(&EditMesh, &TessellatedMesh);
 
 		TSharedPtr<TArray<int>> List = IndexList.List;
 		TUniquePtr<FTessellationPattern> Pattern; 
@@ -227,6 +228,10 @@ UDynamicMesh* UGeometryScriptLibrary_MeshSubdivideFunctions::ApplyAdaptiveTessel
 		{
 			UE::Geometry::AppendError(Debug, EGeometryScriptErrorType::OperationFailed, LOCTEXT("ApplyAdapativeTessellate_Failed", "ApplyAdapativeTessellate: Tessellation failed")); 
 		}
+		else 
+		{
+			EditMesh = MoveTemp(TessellatedMesh);
+		}
 
 	}, EDynamicMeshChangeType::GeneralEdit, EDynamicMeshAttributeChangeFlags::Unknown, false);
 
diff --git a/Engine/Plugins/Runtime/GeometryProcessing/Source/DynamicMesh/Private/Operations/AdaptiveTessellate.cpp b/Engine/Plugins/Runtime/GeometryProcessing/Source/DynamicMesh/Private/Operations/AdaptiveTessellate.cpp
--- a/Engine/Plugins/Runtime/GeometryProcessing/Source/DynamicMesh/Private/Operations/AdaptiveTessellate.cpp
+++ b/Engine/Plugins/Runtime/GeometryProcessing/Source/DynamicMesh/Private/Operations/AdaptiveTessellate.cpp
@@ -8,17 +8,39 @@
 #include "Distance/DistLine3Line3.h"
 #include "Misc/AssertionMacros.h"
 #include "Util/CompactMaps.h"
+#include "DynamicMesh/DynamicMeshAttributeSet.h"
 
 using namespace UE::Geometry;
 
 namespace AdaptiveTessellateLocals 
-{
+{	
+	// Forward declare
+	class FTessellationData;
+
+	template<typename RealType, int ElementSize> 
+	class FOverlayTessellationData;
+
+	void ConstructTessellatedMesh(const FDynamicMesh3* Mesh,
+								  FTessellationData& TessData,
+								  FCompactMaps& CompactInfo,
+								  FDynamicMesh3* ResultMesh);
+
+	template<typename RealType, int ElementSize> 
+	void ConstructTessellatedOverlay(const FDynamicMesh3* Mesh, 
+									 const TDynamicMeshOverlay<RealType, ElementSize>* Overlay,
+									 FTessellationData& MeshTessData, 
+									 FOverlayTessellationData<RealType, ElementSize>& TessData, 
+									 const FCompactMaps& CompactInfo,
+								     TDynamicMeshOverlay<RealType, ElementSize>* ResultOverlay);
+
 	/**
      * Manage data containing tessellation information. This includes barycentric coordinates of the new vertices,
-     * new vertex IDs and triangle indices. 
+     * new vertex/elemnent IDs and triangle indices. This class handles cases where 2 elements are stored per vertex. 
+	 * By default we assume we are interpolating geometry data. Can subclass and implement the virtual methods 
+	 * to handle other sources of data (see FOverlayTessellationData).
      *
      * Given an Edge or Triangle ID, allows to create a TArrayView for reading and writing to the data block 
-	 * containing barycentric coordinates, vertex IDs or triangle data.
+	 * containing barycentric coordinates, vertex/element IDs or triangle data.
 	 */
 	class FTessellationData
 	{
@@ -28,47 +50,59 @@ namespace AdaptiveTessellateLocals
 		TSet<int> EdgesToTessellate;
 		TSet<int> TrianglesToTessellate;
 
-        // Array of all the new IDs for all the vertices added along the tessellated edges. Different edges 
-        // can have a different number of new vertices added. Use the EdgeCoordOffsets to figure out the starting point 
-        // and the length of the data block containing all of the vertices added along the edge. Similar logic applies 
-		// to EdgeCoord, InnerVIDs, InnerCoord, Triangles.
+        // Array of the new IDs for the vertices/elements added along the tessellated edges. Different edges can have a 
+		// different number of vertices/elements added. Each vertex can have 1 or 2 IDs added to handle the overlays. 
+		// Use the EdgeIDOffsets to figure out the starting point and the length of the data block containing all of 
+		// the IDs for the edge/triangle pair. Similar logic applies to EdgeCoord, InnerIDs, InnerCoord, TriangleIDs, 
+		// Triangles.
 		//
-		//            Edge1                Edge2              Edge50   Edge51
-		//            ---------------------------------------------------------------------
-		// EdgeVIDs   | VID1  VID2  VID3 | VID4  VID5 | ... | VID100 | VID101 VID102 |....
-		//            ---------------------------------------------------------------------
-		//           /                                      /
-		//          /                                       EdgeCoordOffsets[Edge50][0]
-		//         EdgeCoordOffsets[Edge1][0] 	
-
-		TArray<int> EdgeVIDs;     // Vertex IDs of the vertices added inside tessellated edges 
-		TArray<double> EdgeCoord; // Lerp cooefficients of the vertices added along tessellated edges		
+		//            Edge1 (seam edge)    Edge2            Edge50  Edge51
+		//            --------------------------------------------------------------------------
+		// EdgeIDs   | ID1  ID2 : ID3 ID4 | ID5  ID6 | ... | ID100 | ID101 ID102 |....
+		//            --------------------------------------------------------------------------
+		//           /          /										
+		//          /           EdgeIDOffsets[(Edge1 , Triangle2)][0]
+		//         EdgeIDOffsets[(Edge1 , Triangle1)][0] 	
+
+		TArray<int> EdgeIDs;     			    // IDs of the vertices/elements added inside the tessellated edges 
+		TMap<FIndex2i, FIndex2i> EdgeIDOffsets; // Maps a tuple (Edge ID, Triangle ID) to tuple of (offset, length) 
+											    // numbers used to access a data block in the EdgeIDs array.
+			
+		TArray<double> EdgeCoord; 			   // Lerp coefficients of the vertices/elements added along the tessellated edges		
 		TMap<int, FIndex2i> EdgeCoordOffsets;  // Maps Edge ID to tuple of (offset, length) numbers used to access a data 
-											   // block in the EdgeVIDs and EdgeCoord arrays.
-
-
-		TArray<int> InnerVIDs; 	      // Vertex IDs of the vertices added inside tessellated triangles 
-		TArray<FVector3d> InnerCoord; // Barycentric coordinates of the vertices added inside tessellated triangles 
-		TMap<int, FIndex2i> InnerCoordOffsets; // Maps Triangle ID to tuple of offset and length numbers used to access 
-											   // a data block in the InnerVIDs and InnerCoord arrays.
+											   // block in the EdgeCoord arrays.
 
+		TArray<int> InnerIDs; 	          // IDs of the vertices/elements added inside the tessellated triangles 
+		TArray<FVector3d> InnerCoord;     // Barycentric coordinates of the vertices added inside the tessellated triangles 
+		TMap<int, FIndex2i> InnerOffsets; // Maps Triangle ID to tuple of (offset, length) numbers used to access 
+										  // a data block in the InnerIDs and InnerCoord arrays.
 
-		TArray<FIndex3i> Triangles; 		  // Indices of the new triangles
-		TMap<int, FIndex2i> TrianglesOffsets; // Maps Triangle ID to tuple of offset and length numbers used to access 
-											  // a data block in the Triangles array.
-		
+		TArray<int> TriangleIDs; 			  // IDs of the new triangles that the original triangle is split into
+		TArray<FIndex3i> Triangles; 		  // Triangle indicies (i.e. vertex id, element id, etc)
+		TMap<int, FIndex2i> TrianglesOffsets; // Maps Triangle ID to tuple of (offset, length) numbers used to access 
+											  // a data block in the TriangleIDs and Triangles array.
 	protected:
 
-		const FDynamicMesh3* Mesh;
+		const FDynamicMesh3* Mesh = nullptr;
+
+		// If we are not working with overlays then we don't care which side of the edge we are working with
+		constexpr static int AnyTriangleID = -1; 
 
+		// The starting ID for new vertices/elements to be added
+		int MaxID = -1;
+ 
 	public:
 
 		FTessellationData(const FDynamicMesh3* InMesh) 
 		:
 		Mesh(InMesh)
 		{
+			MaxID = Mesh->MaxVertexID();
 		}
 
+		virtual ~FTessellationData()
+		{
+		}
 
 		void Init(const FTessellationPattern* Pattern) 
 		{
@@ -89,34 +123,50 @@ namespace AdaptiveTessellateLocals
 			return ArrayView.Slice(OffsetLength[0], OffsetLength[1]);
 		}
 
+		/** @return Array view of the data block containing the ids for the new vertices/elements added along the edge. */
+		TArrayView<int> MapEdgeIDBufferBlock(const int EdgeID) 
+		{
+			TArrayView<int> ArrayView(EdgeIDs);
+			const FIndex2i OffsetLength = EdgeIDOffsets[FIndex2i(EdgeID, AnyTriangleID)];
+			return ArrayView.Slice(OffsetLength[0], OffsetLength[1]);
+		}
 
-		/** @return Array view of the data block containing the vertex ids for the new vertices added along the edge. */
-		TArrayView<int> MapEdgeVIDBufferBlock(const int EdgeID) 
+
+		/** @return Array view of the data block containing the ids for the new elements assosiated with an edge and a triangle. */
+		TArrayView<int> MapEdgeIDBufferBlock(const int EdgeID, const int TriangleID) 
 		{
-			TArrayView<int> ArrayView(EdgeVIDs);
-			const FIndex2i OffsetLength = EdgeCoordOffsets[EdgeID]; 
+			TArrayView<int> ArrayView(EdgeIDs);
+			const FIndex2i OffsetLength = EdgeIDOffsets[FIndex2i(EdgeID, TriangleID)]; 
 			return ArrayView.Slice(OffsetLength[0], OffsetLength[1]);
 		}
 
 
 		/** @return Array view of the data block containing the barycentric coordinates for new vertices added inside the triangle. */
-		TArrayView<FVector3d> MapTriangleCoordBufferBlock(const int TriangleID)
+		TArrayView<FVector3d> MapInnerCoordBufferBlock(const int TriangleID)
 		{
 			TArrayView<FVector3d> ArrayView(InnerCoord);
-			const FIndex2i OffsetLength = InnerCoordOffsets[TriangleID]; 
+			const FIndex2i OffsetLength = InnerOffsets[TriangleID]; 
 			return ArrayView.Slice(OffsetLength[0], OffsetLength[1]);
 		}
 
-
-		/** @return Array view of a data block containing the vertex ids for new vertices added inside the triangle. */
-		TArrayView<int> MapTriangleVIDBufferBlock(const int TriangleID)
+		/** @return Array view of the data block containing the ids for new vertices added inside the triangle. */
+		TArrayView<int> MapInnerIDBufferBlock(const int TriangleID)
 		{
-			TArrayView<int> ArrayView(InnerVIDs);
-			const FIndex2i OffsetLength = InnerCoordOffsets[TriangleID]; 
+			TArrayView<int> ArrayView(InnerIDs);
+			const FIndex2i OffsetLength = InnerOffsets[TriangleID]; 
 			return ArrayView.Slice(OffsetLength[0], OffsetLength[1]);
 		}
 
 
+		/** @return Array view of a data block containing the triangle ids for new triangles. */
+		TArrayView<int> MapTriangleIDBufferBlock(const int TriangleID)
+		{
+			TArrayView<int> ArrayView(TriangleIDs);
+			checkSlow(TrianglesOffsets.Contains(TriangleID));
+			FIndex2i OffsetLengthPair = TrianglesOffsets[TriangleID]; 
+			return ArrayView.Slice(OffsetLengthPair[0], OffsetLengthPair[1]);
+		}
+
 		/** @return Array view of a data block containing the triangle indices for new triangles. */
 		TArrayView<FIndex3i> MapTrianglesBufferBlock(const int TriangleID)
 		{
@@ -125,9 +175,11 @@ namespace AdaptiveTessellateLocals
 			return ArrayView.Slice(OffsetLengthPair[0], OffsetLengthPair[1]);
 		}
 
+		//TODO: Add const versions of the Map* methods
+
 	protected:
 
-		void InitEdgeVertexBuffers(const FTessellationPattern* Pattern) 
+		virtual void InitEdgeVertexBuffers(const FTessellationPattern* Pattern)
 		{
 			int BufferNumVertices = 0;
 			for (const int EdgeID : Mesh->EdgeIndicesItr())
@@ -137,47 +189,47 @@ namespace AdaptiveTessellateLocals
 				{
 					EdgesToTessellate.Add(EdgeID);
 					EdgeCoordOffsets.Add(EdgeID, FIndex2i(BufferNumVertices, NumNewVertices));
+					EdgeIDOffsets.Add(FIndex2i(EdgeID, AnyTriangleID), FIndex2i(BufferNumVertices, NumNewVertices));
 					BufferNumVertices += NumNewVertices;
 				}
 			}
 
-			EdgeVIDs.SetNum(BufferNumVertices);
-			
-			// IDs of the new vertices added along edges start with the max vertex ID of the input mesh.
-			const int MaxVertexID = Mesh->MaxVertexID();
+			EdgeIDs.SetNum(BufferNumVertices);
 			
-			for (int Index = 0; Index < EdgeVIDs.Num(); ++Index)
+			for (int Index = 0; Index < EdgeIDs.Num(); ++Index)
 			{
-				EdgeVIDs[Index] = Index + MaxVertexID;
+				EdgeIDs[Index] = Index + MaxID;
 			}
 			
-			// Pre-allocate memory for barycentric coordinates we will be computing
+			// Pre-allocate memory for linear coordinates we will be computing
 			EdgeCoord.SetNum(BufferNumVertices);
 		}
 
-		
-		void InitTriVertexBuffers(const FTessellationPattern* Pattern) 
+
+		virtual void InitTriVertexBuffers(const FTessellationPattern* Pattern) 
 		{
 			int BufferNumVertices = 0;
-			for (const int TriangleID : Mesh->TriangleIndicesItr())
+			for (int TriangleID = 0; TriangleID < Mesh->MaxTriangleID(); ++TriangleID)
 			{
-				const int NumNewVertices = Pattern->GetNumberOfNewVerticesForTrianglePatch(TriangleID); 
-				if (NumNewVertices > 0) 
+				if (this->IsValidTriangle(TriangleID))
 				{
-					InnerCoordOffsets.Add(TriangleID, FIndex2i(BufferNumVertices, NumNewVertices));
-					BufferNumVertices += NumNewVertices;
+					const int NumNewVertices = Pattern->GetNumberOfNewVerticesForTrianglePatch(TriangleID); 
+					if (NumNewVertices > 0) 
+					{
+						InnerOffsets.Add(TriangleID, FIndex2i(BufferNumVertices, NumNewVertices));
+						BufferNumVertices += NumNewVertices;
+					}
 				}
 			}
 
-			InnerVIDs.SetNum(BufferNumVertices);
+			InnerIDs.SetNum(BufferNumVertices);
 			
-			// IDs of the new vertices added inside triangles start with the last vertex id added along edges. It is 
-			// possible that none of the edges were tessellated in which case we start with the max vertex id of the 
-			// input mesh.
-			const int LastEdgeVIDs = EdgeVIDs.Num() > 0 ? EdgeVIDs.Last() + 1 : Mesh->MaxVertexID();
-			for (int Index = 0; Index < InnerVIDs.Num(); ++Index)
+			// IDs of the new vertices/elements added inside triangles start with the last id added along edges. It is 
+			// possible that none of the edges were tessellated in which case we start with the max vertex/element id. 
+			const int LastEdgeVIDs = EdgeIDs.Num() > 0 ? EdgeIDs.Last() + 1 : MaxID;
+			for (int Index = 0; Index < InnerIDs.Num(); ++Index)
 			{
-				InnerVIDs[Index] = Index + LastEdgeVIDs;
+				InnerIDs[Index] = Index + LastEdgeVIDs;
 			}
 
 			// Pre-allocate memory for barycentric coordinates we will be computing
@@ -185,26 +237,131 @@ namespace AdaptiveTessellateLocals
 		}
 
 
-		void InitTrianglesBuffers(const FTessellationPattern* Pattern) 
+		virtual void InitTrianglesBuffers(const FTessellationPattern* Pattern) 
 		{
 			int BufferNumTriangles = 0;
-			for (const int TriangleID : Mesh->TriangleIndicesItr())
+			for (int TriangleID = 0; TriangleID < Mesh->MaxTriangleID(); ++TriangleID)
 			{
-				const int NumNewTriangles = Pattern->GetNumberOfPatchTriangles(TriangleID); 
-				if (NumNewTriangles > 0) 
-				{	
-					TrianglesToTessellate.Add(TriangleID);
-					TrianglesOffsets.Add(TriangleID, FIndex2i(BufferNumTriangles, NumNewTriangles));
-					BufferNumTriangles += NumNewTriangles;
+				if (this->IsValidTriangle(TriangleID))
+				{
+					const int NumNewTriangles = Pattern->GetNumberOfPatchTriangles(TriangleID); 
+					if (NumNewTriangles > 0) 
+					{	
+						TrianglesToTessellate.Add(TriangleID);
+						TrianglesOffsets.Add(TriangleID, FIndex2i(BufferNumTriangles, NumNewTriangles));
+						BufferNumTriangles += NumNewTriangles;
+					}
 				}
 			}
 
+			TriangleIDs.SetNum(BufferNumTriangles);
+
+			for (int Index = 0; Index < TriangleIDs.Num(); ++Index) 
+			{
+				TriangleIDs[Index] = Mesh->MaxTriangleID() + Index;
+			}
+
 			// Pre-allocate memory for triangles we will be computing
 			Triangles.SetNum(BufferNumTriangles);
 		}
+
+		virtual bool IsValidTriangle(int TriangleID) const
+		{
+			return Mesh->IsTriangle(TriangleID);
+		}
 	}; 
 
 
+	/**  
+	 * Handle the data for tessellating overlays. The main difference between this class and its parent is that we 
+	 * handle 2 elements per vertex along edges. We also need to consider the fact that a triangle could be marked 
+	 * for tessellation but not be set in the overlay, in which case we can skip any element computation for it.
+	 */
+	template<typename RealType, int ElementSize>
+	class FOverlayTessellationData : public FTessellationData
+	{
+	protected:
+
+		const TDynamicMeshOverlay<RealType, ElementSize>* Overlay = nullptr; 
+
+	public:
+
+		FOverlayTessellationData(const FDynamicMesh3* InMesh, const TDynamicMeshOverlay<RealType, ElementSize>* InOverlay)
+		:
+		FTessellationData(InMesh), Overlay(InOverlay)
+		{
+			MaxID = Overlay->MaxElementID();
+		}
+
+	protected:
+
+		virtual void InitEdgeVertexBuffers(const FTessellationPattern* Pattern) override
+		{
+			int CoordBufferSize = 0;
+			int VIDBufferSize = 0;
+			for (const int EdgeID : Mesh->EdgeIndicesItr())
+			{
+				const int NumNewVertices = Pattern->GetNumberOfNewVerticesForEdgePatch(EdgeID); 
+				const FIndex2i EdgeTri = Mesh->GetEdgeT(EdgeID);
+
+				// Edge is invalid in the overlay if both triangles that share it are not set in the overlay 
+				const bool bIsNotBndry = Mesh->GetEdgeT(EdgeID).B != FDynamicMesh3::InvalidID;
+				const bool bIsValidEdge = Overlay->IsSetTriangle(EdgeTri.A) || (bIsNotBndry && Overlay->IsSetTriangle(EdgeTri.B));
+
+				if (NumNewVertices > 0 && bIsValidEdge)
+				{
+					EdgesToTessellate.Add(EdgeID);			
+					EdgeCoordOffsets.Add(EdgeID, FIndex2i(CoordBufferSize, NumNewVertices));	
+					CoordBufferSize += NumNewVertices;
+
+					if (Overlay->IsSetTriangle(EdgeTri.A)) 
+					{	
+						EdgeIDOffsets.Add(FIndex2i(EdgeID, EdgeTri.A), FIndex2i(VIDBufferSize, NumNewVertices));
+						
+						if (bIsNotBndry && Overlay->IsSetTriangle(EdgeTri.B)) 
+						{
+							if (Overlay->IsSeamEdge(EdgeID) == true) 
+							{
+								EdgeIDOffsets.Add(FIndex2i(EdgeID, EdgeTri.B), FIndex2i(VIDBufferSize + NumNewVertices, NumNewVertices));
+								VIDBufferSize += 2*NumNewVertices;
+							}
+							else 
+							{
+								// Not a seam edge so simply point to the same data block as FIndex2i(EdgeID, EdgeTri.A) case
+								EdgeIDOffsets.Add(FIndex2i(EdgeID, EdgeTri.B), FIndex2i(VIDBufferSize, NumNewVertices));
+								VIDBufferSize += NumNewVertices;
+							}
+						}
+						else 
+						{
+							VIDBufferSize += NumNewVertices;
+						}
+					}
+					else 
+					{
+						checkSlow(bIsNotBndry && Overlay->IsSetTriangle(EdgeTri.B));
+						EdgeIDOffsets.Add(FIndex2i(EdgeID, EdgeTri.B), FIndex2i(VIDBufferSize, NumNewVertices));
+						VIDBufferSize += NumNewVertices;
+					}
+				}
+			}
+
+			EdgeIDs.SetNum(VIDBufferSize);
+			
+			for (int Index = 0; Index < EdgeIDs.Num(); ++Index)
+			{
+				EdgeIDs[Index] = Index + MaxID;
+			}
+			
+			// Pre-allocate memory for linear coordinates we will be computing
+			EdgeCoord.SetNum(CoordBufferSize);
+		}
+
+		virtual bool IsValidTriangle(int TriangleID) const override
+		{
+			return Mesh->IsTriangle(TriangleID) && Overlay->IsSetTriangle(TriangleID);
+		}
+	}; 
 
 
 	/**
@@ -342,7 +499,7 @@ namespace AdaptiveTessellateLocals
 
 			// If the inner area of the triangle patch is not being tessellated but at least one of its edges is, then 
 			// we insert a single vertex in the middle and connect it to all of the edge vertices. Therefore, we find 
-			// those triangles and set their inner tessellation level to 1. This would tell the TessellatePatch function 
+			// those triangles and set their inner tessellation level to 1. This would tell the TessellateTriPatch function 
 			// to insert a vertex and generate the triangle fan.
 			for (const int TriangleID : Mesh->TriangleIndicesItr()) 
 			{
@@ -663,12 +820,17 @@ namespace AdaptiveTessellateLocals
 
 
 
-	/** 
-	 * Use the tessellation pattern to tessellate the mesh and generate the tessellation data (barycentric coordinates 
-	 * of the new vertices and new triangles).
-	 */
-	void Tessellate(const FDynamicMesh3* Mesh, const FTessellationPattern* Pattern, const bool bUseParallel, FTessellationData& TessData) 
+	/** Use the tessellation pattern to tessellate the mesh and generate the tessellation data. */
+	void TessellateGeometry(const FDynamicMesh3* Mesh, 
+							const FTessellationPattern* Pattern, 
+							const bool bUseParallel, 
+							FCompactMaps& OutCompactInfo, 
+							FTessellationData& TessData,
+							FDynamicMesh3* OutMesh) 
 	{	
+		
+		TessData.Init(Pattern);
+
 		// Tessellate edge patches
 		TArray<int> EdgesToTessellate = TessData.EdgesToTessellate.Array();
 		ParallelFor(EdgesToTessellate.Num(), [&](int32 Index)
@@ -678,7 +840,7 @@ namespace AdaptiveTessellateLocals
 			const int EdgeID = EdgesToTessellate[Index];
 			Patch.EdgeID = EdgeID;
 			Patch.LinearCoord = TessData.MapEdgeCoordBufferBlock(EdgeID);
-			Patch.VIDs = TessData.MapEdgeVIDBufferBlock(EdgeID);
+			Patch.VIDs = TessData.MapEdgeIDBufferBlock(EdgeID);
 			
 			Pattern->TessellateEdgePatch(Patch);
 		}, bUseParallel ? EParallelForFlags::None : EParallelForFlags::ForceSingleThread);
@@ -704,7 +866,7 @@ namespace AdaptiveTessellateLocals
 			{
 				Patch.UVEdge.EdgeID = TriEdges[0];
 				Patch.UVEdge.LinearCoord = TessData.MapEdgeCoordBufferBlock(TriEdges[0]);
-				Patch.UVEdge.VIDs = TessData.MapEdgeVIDBufferBlock(TriEdges[0]);
+				Patch.UVEdge.VIDs = TessData.MapEdgeIDBufferBlock(TriEdges[0]);
 				Patch.UVEdge.bIsReversed = Mesh->GetEdgeV(TriEdges[0])[0] != TriVertices[0];
 			}
 
@@ -713,7 +875,7 @@ namespace AdaptiveTessellateLocals
 			{
 				Patch.VWEdge.EdgeID = TriEdges[1];
 				Patch.VWEdge.LinearCoord = TessData.MapEdgeCoordBufferBlock(TriEdges[1]);
-				Patch.VWEdge.VIDs = TessData.MapEdgeVIDBufferBlock(TriEdges[1]);
+				Patch.VWEdge.VIDs = TessData.MapEdgeIDBufferBlock(TriEdges[1]);
 				Patch.VWEdge.bIsReversed = Mesh->GetEdgeV(TriEdges[1])[0] != TriVertices[1];
 			}
 			
@@ -722,13 +884,13 @@ namespace AdaptiveTessellateLocals
 			{
 				Patch.UWEdge.EdgeID = TriEdges[2];
 				Patch.UWEdge.LinearCoord = TessData.MapEdgeCoordBufferBlock(TriEdges[2]);
-				Patch.UWEdge.VIDs = TessData.MapEdgeVIDBufferBlock(TriEdges[2]);
+				Patch.UWEdge.VIDs = TessData.MapEdgeIDBufferBlock(TriEdges[2]);
 				Patch.UWEdge.bIsReversed = Mesh->GetEdgeV(TriEdges[2])[0] != TriVertices[2];
 			}
 
 			//  Inner Triangle Vertices
-			Patch.BaryCoord = TessData.MapTriangleCoordBufferBlock(TriangleID);
-			Patch.VIDs = TessData.MapTriangleVIDBufferBlock(TriangleID);
+			Patch.BaryCoord = TessData.MapInnerCoordBufferBlock(TriangleID);
+			Patch.VIDs = TessData.MapInnerIDBufferBlock(TriangleID);
 
 			// Inner Triangles
 			Patch.Triangles = TessData.MapTrianglesBufferBlock(TriangleID);
@@ -736,25 +898,136 @@ namespace AdaptiveTessellateLocals
 			// Tessellate the patch, i.e. generate inner vertices and triangles
 			Pattern->TessellateTriPatch(Patch);
 		}, bUseParallel ? EParallelForFlags::None : EParallelForFlags::ForceSingleThread);
+
+		AdaptiveTessellateLocals::ConstructTessellatedMesh(Mesh, TessData, OutCompactInfo, OutMesh);
+	}
+
+
+	/** Use the tessellation pattern to tessellate an overlay. */
+	template<typename RealType, int ElementSize> 
+	void TessellateOverlay(const FDynamicMesh3* Mesh, 
+						   const TDynamicMeshOverlay<RealType, ElementSize>* Overlay,
+						   const FTessellationPattern* Pattern, 
+						   FTessellationData& MeshTessData,
+						   const bool bUseParallel, 
+						   const FCompactMaps& CompactInfo,
+						   TDynamicMeshOverlay<RealType, ElementSize>* OutOverlay)
+	{	
+		
+		//TODO: The overlay tessellation data should only compute and contain the connectivity information since we 
+	 	// can simply  reuse coordinates from the geometry tessellation (i.e. overlay elements live on the vertices). 
+		AdaptiveTessellateLocals::FOverlayTessellationData OverlayTessData(Mesh, Overlay);
+		OverlayTessData.Init(Pattern);
+
+		// Tessellate edge patches
+		TArray<int> EdgesToTessellate = OverlayTessData.EdgesToTessellate.Array();
+		ParallelFor(EdgesToTessellate.Num(), [&](int32 Index)
+		{
+			const int EdgeID = EdgesToTessellate[Index];
+			
+			const FIndex2i EdgeTri = Mesh->GetEdgeT(EdgeID);
+			
+			if (Overlay->IsSetTriangle(EdgeTri.A))
+			{
+				FTessellationPattern::EdgePatch PatchA;
+				PatchA.EdgeID = EdgeID;
+				PatchA.LinearCoord = OverlayTessData.MapEdgeCoordBufferBlock(EdgeID);
+				PatchA.VIDs = OverlayTessData.MapEdgeIDBufferBlock(EdgeID, EdgeTri.A);
+				Pattern->TessellateEdgePatch(PatchA);
+			}
+
+			if (EdgeTri.B != FDynamicMesh3::InvalidID && Overlay->IsSetTriangle(EdgeTri.B))
+			{
+				FTessellationPattern::EdgePatch PatchB;
+				PatchB.EdgeID = EdgeID;
+				PatchB.LinearCoord = OverlayTessData.MapEdgeCoordBufferBlock(EdgeID);
+				PatchB.VIDs = OverlayTessData.MapEdgeIDBufferBlock(EdgeID, EdgeTri.B);
+				Pattern->TessellateEdgePatch(PatchB);
+			}
+		}, bUseParallel ? EParallelForFlags::None : EParallelForFlags::ForceSingleThread);
+
+		// Tessellate triangle patches
+		TArray<int> TrianglesToTessellate = OverlayTessData.TrianglesToTessellate.Array();
+		ParallelFor(TrianglesToTessellate.Num(), [&](int32 Index)
+		{
+			FTessellationPattern::TrianglePatch Patch;
+			
+			const int TriangleID = TrianglesToTessellate[Index]; 
+			
+			Patch.TriangleID = TriangleID;
+			
+			FIndex3i TriVertices = Mesh->GetTriangle(TriangleID);
+			Patch.UVWCorners = Overlay->GetTriangle(TriangleID);
+			
+			const FIndex3i& TriEdges = Mesh->GetTriEdgesRef(TriangleID);
+			
+			// UV Edge
+			if (OverlayTessData.EdgesToTessellate.Contains(TriEdges[0]))
+			{
+				Patch.UVEdge.EdgeID = TriEdges[0];
+				Patch.UVEdge.LinearCoord = OverlayTessData.MapEdgeCoordBufferBlock(TriEdges[0]);
+				Patch.UVEdge.VIDs = OverlayTessData.MapEdgeIDBufferBlock(TriEdges[0], TriangleID);
+				Patch.UVEdge.bIsReversed = Mesh->GetEdgeV(TriEdges[0])[0] != TriVertices[0];
+			}
+
+			// VW Edge
+			if (OverlayTessData.EdgesToTessellate.Contains(TriEdges[1]))
+			{
+				Patch.VWEdge.EdgeID = TriEdges[1];
+				Patch.VWEdge.LinearCoord = OverlayTessData.MapEdgeCoordBufferBlock(TriEdges[1]);
+				Patch.VWEdge.VIDs = OverlayTessData.MapEdgeIDBufferBlock(TriEdges[1], TriangleID);
+				Patch.VWEdge.bIsReversed = Mesh->GetEdgeV(TriEdges[1])[0] != TriVertices[1];
+			}
+			
+			// UW Edge
+			if (OverlayTessData.EdgesToTessellate.Contains(TriEdges[2]))
+			{
+				Patch.UWEdge.EdgeID = TriEdges[2];
+				Patch.UWEdge.LinearCoord = OverlayTessData.MapEdgeCoordBufferBlock(TriEdges[2]);
+				Patch.UWEdge.VIDs = OverlayTessData.MapEdgeIDBufferBlock(TriEdges[2], TriangleID);
+				Patch.UWEdge.bIsReversed = Mesh->GetEdgeV(TriEdges[2])[0] != TriVertices[2];
+			}
+
+			//  Inner Triangle Vertices
+			Patch.BaryCoord = OverlayTessData.MapInnerCoordBufferBlock(TriangleID);
+			Patch.VIDs = OverlayTessData.MapInnerIDBufferBlock(TriangleID);
+
+			// Inner Triangles
+			Patch.Triangles = OverlayTessData.MapTrianglesBufferBlock(TriangleID);
+
+			// Tessellate the patch, i.e. generate inner vertices and triangles
+			Pattern->TessellateTriPatch(Patch);
+		}, bUseParallel ? EParallelForFlags::None : EParallelForFlags::ForceSingleThread);
+
+		AdaptiveTessellateLocals::ConstructTessellatedOverlay(Mesh, 
+															  Overlay,
+															  MeshTessData, 
+															  OverlayTessData, 
+															  CompactInfo, 
+															  OutOverlay);
 	}
 
 
 	/**
-	 * Given the original FDynamicMesh3 and the tessellation data, construct new tessellated FDynamicMesh3 geometry and 
-	 * interpolate attributes.
+	 * Given the original FDynamicMesh3 and the tessellation data, construct new tessellated FDynamicMesh3 geometry.
+	 * Output mesh will be compact. 
+	 * 
+	 * @param CompactInfo Stores vertex and triangle mappings.
 	 */
-	void ConstructTessellatedMesh(const FDynamicMesh3* Mesh, 
-								  const FTessellationData& TessData, 
-								  FDynamicMesh3* ResultMesh)  
-	{	
+	void ConstructTessellatedMesh(const FDynamicMesh3* Mesh,
+								  FTessellationData& TessData,
+								  FCompactMaps& CompactInfo,
+								  FDynamicMesh3* ResultMesh)
+	{
 		ResultMesh->Clear();
 
-		// Handle non-compact meshes
-		FCompactMaps CompactInfo;
-
-		const int ResultNumVertexIDs = Mesh->MaxVertexID() + TessData.EdgeVIDs.Num() + TessData.InnerVIDs.Num();
+		const int ResultNumVertexIDs = Mesh->MaxVertexID() + TessData.EdgeIDs.Num() + TessData.InnerIDs.Num();
 		CompactInfo.ResetVertexMap(ResultNumVertexIDs, false);
 
+		const int ResultNumTriangleIDs = Mesh->MaxTriangleID() + TessData.Triangles.Num();
+		CompactInfo.ResetTriangleMap(ResultNumTriangleIDs, false);
+
+		// Append all of the vertices from the input mesh
 		for (int VertexID = 0; VertexID < Mesh->MaxVertexID(); ++VertexID)
 		{
 			if (Mesh->IsVertex(VertexID))
@@ -767,37 +1040,40 @@ namespace AdaptiveTessellateLocals
 			}
 		}
 
+		// Append all of the vertices along tessellated edges
 		for (const int EdgeID : TessData.EdgesToTessellate)
 		{
-			const FIndex2i OffsetNum = TessData.EdgeCoordOffsets[EdgeID];
-
+			TArrayView<double> CoordBlock = TessData.MapEdgeCoordBufferBlock(EdgeID);
+			TArrayView<int> IDBlock = TessData.MapEdgeIDBufferBlock(EdgeID);
+			
 			const FIndex2i EdgeV = Mesh->GetEdgeV(EdgeID);
 
-			for (int Index = 0; Index < OffsetNum[1]; ++Index)
+			for (int Index = 0; Index < CoordBlock.Num(); ++Index)
 			{	
-				const int Offset = OffsetNum[0] + Index;
-				const double Alpha = TessData.EdgeCoord[Offset];
-				const int VertexID = TessData.EdgeVIDs[Offset]; 
+				double Alpha = CoordBlock[Index];
+				Alpha = FMath::Clamp(Alpha, 0.0, 1.0);
 
+				const int VertexID = IDBlock[Index]; 
+				
 				const FVector3d Vertex = (1.0 - Alpha)*Mesh->GetVertex(EdgeV[0]) + Alpha*Mesh->GetVertex(EdgeV[1]);
 				
 				CompactInfo.SetVertexMapping(VertexID, ResultMesh->AppendVertex(Vertex));
 			}
 		}
 
-
+		// Append all of the vertices inside tessellated triangles
 		for (const int TriangleID : TessData.TrianglesToTessellate)
 		{
-			const FIndex2i OffsetNum = TessData.InnerCoordOffsets[TriangleID];
+			TArrayView<FVector3d> CoordBlock = TessData.MapInnerCoordBufferBlock(TriangleID);
+			TArrayView<int> IDBlock = TessData.MapInnerIDBufferBlock(TriangleID);
+			checkSlow(CoordBlock.Num() == IDBlock.Num());
 
 			const FIndex3i TriangleV = Mesh->GetTriangle(TriangleID);
 
-			for (int Index = 0; Index < OffsetNum[1]; ++Index)
+			for (int Index = 0; Index < CoordBlock.Num(); ++Index)
 			{	
-				const int Offset = OffsetNum[0] + Index;	
-				
-				const FVector3d Bary = TessData.InnerCoord[Offset];
-				const int VertexID = TessData.InnerVIDs[Offset];
+				const FVector3d Bary = CoordBlock[Index];
+				const int VertexID = IDBlock[Index];
 
 				const FVector3d Vertex = Bary[0] * Mesh->GetVertex(TriangleV[0]) + 
 										 Bary[1] * Mesh->GetVertex(TriangleV[1]) + 
@@ -807,34 +1083,238 @@ namespace AdaptiveTessellateLocals
 			}
 		}
 
-
+		// Append all the triangles from the input mesh that we are not tessellating
 		for (const int TriangleID : Mesh->TriangleIndicesItr())
-		{		
-			if (TessData.TrianglesToTessellate.Contains(TriangleID)) 
+		{
+			if (TessData.TrianglesToTessellate.Contains(TriangleID) == false)
 			{
-				const FIndex2i OffsetNum = TessData.TrianglesOffsets[TriangleID];
-				for (int Index = 0; Index < OffsetNum[1]; ++Index)
-				{	
-					const int Offset = OffsetNum[0] + Index;	
-					FIndex3i Tri = TessData.Triangles[Offset];
-					FIndex3i MappedTri = CompactInfo.GetVertexMapping(Tri);
-					ResultMesh->AppendTriangle(MappedTri);
-				}
+				const FIndex3i Tri = Mesh->GetTriangle(TriangleID);
+				const FIndex3i MappedTri = CompactInfo.GetVertexMapping(Tri);
+				CompactInfo.SetTriangleMapping(TriangleID, ResultMesh->AppendTriangle(MappedTri));
 			}
-			else 
-			{	
-				FIndex3i Tri = Mesh->GetTriangle(TriangleID);
-				FIndex3i MappedTri = CompactInfo.GetVertexMapping(Tri);
-				ResultMesh->AppendTriangle(MappedTri);
+		}
+		
+		// Append all the new triangles generated during the tessellation
+		for (const int TriangleID : TessData.TrianglesToTessellate) 
+		{
+			TArrayView<FIndex3i> TrianglesBlock = TessData.MapTrianglesBufferBlock(TriangleID);
+			TArrayView<int> TriangleIDBlock = TessData.MapTriangleIDBufferBlock(TriangleID);
+			checkSlow(TrianglesBlock.Num() == TriangleIDBlock.Num());
+
+			for (int Index = 0; Index < TriangleIDBlock.Num(); ++Index)
+			{
+				const int NewTriangleID = TriangleIDBlock[Index];
+				const FIndex3i Tri = TrianglesBlock[Index];
+				const FIndex3i MappedTri = CompactInfo.GetVertexMapping(Tri);
+				CompactInfo.SetTriangleMapping(NewTriangleID, ResultMesh->AppendTriangle(MappedTri));
 			}
 		}
 
-		//TODO: use barycentric data to interpolate attributes
-		
 		//TODO: instead of generating the mesh from scratch, we could remove triangles we are tessellating from the 
 		// input mesh and append the new ones. This can be an option when the number of triangles tessellated is small,
 		// compared to the total number of triangles.
 	}
+
+
+	/**
+	 * Construct a new overlay from the tessellation data and compact information from the geometry tessellation.
+	 */
+	template<typename RealType, int ElementSize> 
+	void ConstructTessellatedOverlay(const FDynamicMesh3* Mesh, 
+									 const TDynamicMeshOverlay<RealType, ElementSize>* Overlay,
+									 FTessellationData& MeshTessData, 
+									 FOverlayTessellationData<RealType, ElementSize>& OverlayTessData, 
+									 const FCompactMaps& CompactInfo,
+								     TDynamicMeshOverlay<RealType, ElementSize>* ResultOverlay)  
+	{	
+		auto LerpElements = [](const RealType* Element1, const RealType* Element2, RealType* OutElement, RealType Alpha)
+		{
+			Alpha = FMath::Clamp(Alpha, 0.0, 1.0);
+			RealType OneMinusAlpha = (RealType)1 - Alpha;
+
+			for (int Idx = 0; Idx < ElementSize; ++Idx)
+			{
+				OutElement[Idx] = OneMinusAlpha * Element1[Idx] + Alpha * Element2[Idx];
+			}
+		};
+
+		auto BaryElements = [](const RealType* Element1, const RealType* Element2, const RealType* Element3,
+							   RealType* OutElement, 
+							   RealType Alpha, RealType Beta, RealType Theta)
+		{
+			checkSlow(FMath::Abs(Alpha + Beta + Theta - 1.0f) < KINDA_SMALL_NUMBER);
+			for (int Idx = 0; Idx < ElementSize; ++Idx)
+			{
+				OutElement[Idx] = Alpha * Element1[Idx] + Beta * Element2[Idx] + Theta * Element3[Idx];
+			}
+		};
+
+		ResultOverlay->ClearElements();
+
+		// Need to track the ID of the appended elements to handle non-compact meshes
+		TMap<int, int> MapE; 
+		MapE.Reserve(Overlay->ElementCount() + OverlayTessData.EdgeIDOffsets.Num());
+		
+		
+		// Buffers to be reused
+		RealType Element1[ElementSize];
+		RealType Element2[ElementSize];
+		RealType Element3[ElementSize];
+		RealType Out[ElementSize];
+		
+		// First add all the existing elements
+		for (const int ElementID : Overlay->ElementIndicesItr())
+		{
+			Overlay->GetElement(ElementID, Out);
+			MapE.Add(ElementID, ResultOverlay->AppendElement(Out));
+		}
+		
+		// Add all the new elements inserted along the edges by the tessellator
+		for (const int EdgeID : OverlayTessData.EdgesToTessellate)
+		{
+			const FIndex2i EdgeTri = Mesh->GetEdgeT(EdgeID);
+			const FIndex2i EdgeV = Mesh->GetEdgeV(EdgeID);
+			TArrayView<double> CoordBlock = OverlayTessData.MapEdgeCoordBufferBlock(EdgeID);
+
+			if (Overlay->IsSetTriangle(EdgeTri.A))  
+			{
+				TArrayView<int> IDBlock = OverlayTessData.MapEdgeIDBufferBlock(EdgeID, EdgeTri.A);
+				checkSlow(IDBlock.Num() == CoordBlock.Num());
+				
+				Overlay->GetElementAtVertex(EdgeTri.A, EdgeV[0], Element1);
+				Overlay->GetElementAtVertex(EdgeTri.A, EdgeV[1], Element2);
+
+				for (int Index = 0; Index < IDBlock.Num(); ++Index)
+				{	
+					const double Alpha = CoordBlock[Index];
+					const int ElementID = IDBlock[Index]; 
+					
+					LerpElements(Element1, Element2, Out, Alpha);
+					
+					MapE.Add(ElementID, ResultOverlay->AppendElement(Out));
+				}
+			}	
+
+			if (EdgeTri.B != FDynamicMesh3::InvalidID && Overlay->IsSetTriangle(EdgeTri.B) && Overlay->IsSeamEdge(EdgeID))
+			{
+				TArrayView<int> IDBlock = OverlayTessData.MapEdgeIDBufferBlock(EdgeID, EdgeTri.B);
+				checkSlow(IDBlock.Num() == CoordBlock.Num());
+
+				Overlay->GetElementAtVertex(EdgeTri.B, EdgeV[0], Element1);
+				Overlay->GetElementAtVertex(EdgeTri.B, EdgeV[1], Element2);
+
+				for (int Index = 0; Index < IDBlock.Num(); ++Index)
+				{	
+					const double Alpha = CoordBlock[Index];
+					const int ElementID = IDBlock[Index]; 
+					
+					LerpElements(Element1, Element2, Out, Alpha);
+					
+					MapE.Add(ElementID, ResultOverlay->AppendElement(Out));
+				}
+			}
+		}
+
+		// Add all the elements added inside of the triangles by the tessellator
+		for (const int TriangleID : OverlayTessData.TrianglesToTessellate)
+		{
+			const FIndex3i TriangleV = Mesh->GetTriangle(TriangleID);
+			
+			Overlay->GetElementAtVertex(TriangleID, TriangleV[0], Element1);
+			Overlay->GetElementAtVertex(TriangleID, TriangleV[1], Element2);
+			Overlay->GetElementAtVertex(TriangleID, TriangleV[2], Element3);
+
+			TArrayView<FVector3d> CoordBlock = OverlayTessData.MapInnerCoordBufferBlock(TriangleID);
+			TArrayView<int> IDBlock = OverlayTessData.MapInnerIDBufferBlock(TriangleID);
+
+			checkSlow(CoordBlock.Num() == IDBlock.Num());
+			for (int Index = 0; Index < CoordBlock.Num(); ++Index)
+			{
+				const FVector3d Bary = CoordBlock[Index]; 
+				const int ElementID = IDBlock[Index];
+
+				BaryElements(Element1, Element2, Element3, Out, Bary[0], Bary[1], Bary[2]);
+
+				MapE.Add(ElementID, ResultOverlay->AppendElement(Out));
+			}
+		}
+
+		// Set the overlay triangles and point them to the correct element ids
+		for (const int TriangleID : Mesh->TriangleIndicesItr())
+		{
+			if (OverlayTessData.TrianglesToTessellate.Contains(TriangleID) == false)
+			{
+				const FIndex3i ElementTri = Overlay->GetTriangle(TriangleID);
+				const int ToTID = CompactInfo.GetTriangleMapping(TriangleID);
+				ResultOverlay->SetTriangle(ToTID, ElementTri);
+			}
+		}
+
+		for (const int TriangleID : OverlayTessData.TrianglesToTessellate)
+		{
+			TArrayView<FIndex3i> TrianglesBlock = OverlayTessData.MapTrianglesBufferBlock(TriangleID);
+			TArrayView<int> TriangleIDBlock = MeshTessData.MapTriangleIDBufferBlock(TriangleID);
+			
+			checkSlow(TrianglesBlock.Num() == TriangleIDBlock.Num());
+			
+			for (int Index = 0; Index < TrianglesBlock.Num(); ++Index)
+			{
+				const FIndex3i Tri = TrianglesBlock[Index];
+			
+				const int NewTriangleID = TriangleIDBlock[Index];
+				const int ToTID = CompactInfo.GetTriangleMapping(NewTriangleID);
+
+				ResultOverlay->SetTriangle(ToTID, FIndex3i(MapE[Tri.A], MapE[Tri.B], MapE[Tri.C]));
+			}
+		}
+	}
+
+	/** 
+	 * Run the main tessellation logic for the geometry, overlays and attributes. 
+	 * 
+	 * @param OutMesh The resulting tessellated mesh.
+	 */
+	void Tessellate(const FDynamicMesh3* InMesh, const FTessellationPattern* Pattern, const bool bUseParallel, FDynamicMesh3* OutMesh)
+	{
+		OutMesh->Clear();
+
+		FCompactMaps CompactInfo;
+
+		AdaptiveTessellateLocals::FTessellationData TessData(InMesh);
+		TessellateGeometry(InMesh, Pattern, bUseParallel, CompactInfo, TessData, OutMesh);
+
+		if (InMesh->HasAttributes()) 
+		{	
+			OutMesh->EnableAttributes(); 
+
+			const FDynamicMeshAttributeSet* InAttributes = InMesh->Attributes(); 
+			FDynamicMeshAttributeSet* OutAttributes = OutMesh->Attributes(); 
+			
+			if (InAttributes->NumNormalLayers()) 
+			{
+				OutAttributes->SetNumNormalLayers(InAttributes->NumNormalLayers());
+				for (int Idx = 0; Idx < InAttributes->NumNormalLayers(); ++Idx) 
+				{	
+					TessellateOverlay(InMesh, InAttributes->GetNormalLayer(Idx), Pattern, TessData, bUseParallel, CompactInfo,  OutAttributes->GetNormalLayer(Idx));
+				} 
+			}
+
+			if (InAttributes->NumUVLayers()) 
+			{
+				OutAttributes->SetNumUVLayers(InAttributes->NumUVLayers());
+				for (int Idx = 0; Idx < InAttributes->NumUVLayers(); ++Idx) 
+				{	
+					TessellateOverlay(InMesh, InAttributes->GetUVLayer(Idx), Pattern, TessData, bUseParallel, CompactInfo, OutAttributes->GetUVLayer(Idx));
+				}
+			}
+
+			if (InAttributes->HasPrimaryColors()) 
+			{
+				OutAttributes->EnablePrimaryColors();
+				TessellateOverlay(InMesh, InAttributes->PrimaryColors(), Pattern, TessData, bUseParallel, CompactInfo, OutAttributes->PrimaryColors());
+			}
+		}
+	}
 }
 
 
@@ -1044,12 +1524,7 @@ bool FAdaptiveTessellate::Compute()
 		Mesh = ResultMeshCopy.Get();
 	}
 
-	AdaptiveTessellateLocals::FTessellationData TessData(Mesh);
-	TessData.Init(Pattern);
-
-	AdaptiveTessellateLocals::Tessellate(Mesh, Pattern, bUseParallel, TessData);
-
-	AdaptiveTessellateLocals::ConstructTessellatedMesh(Mesh, TessData, ResultMesh);
+	AdaptiveTessellateLocals::Tessellate(Mesh, Pattern, bUseParallel, ResultMesh);
 
 	if (Cancelled()) 
 	{
