Fixed a bug where UnifromTessellation was not copying over the names of the attributes.

#rnx

diff --git a/Engine/Plugins/Runtime/GeometryProcessing/Source/DynamicMesh/Private/Operations/UniformTessellate.cpp b/Engine/Plugins/Runtime/GeometryProcessing/Source/DynamicMesh/Private/Operations/UniformTessellate.cpp
--- a/Engine/Plugins/Runtime/GeometryProcessing/Source/DynamicMesh/Private/Operations/UniformTessellate.cpp
+++ b/Engine/Plugins/Runtime/GeometryProcessing/Source/DynamicMesh/Private/Operations/UniformTessellate.cpp
@@ -1842,6 +1842,7 @@ namespace UniformTessellateLocals
 				}
 
 				Generator->CopyToAttribute(OutAttributes->GetMaterialID());
+				OutAttributes->GetMaterialID()->SetName(InAttributes->GetMaterialID()->GetName());
 			}
 
 			if (InAttributes->NumPolygroupLayers())  
@@ -1857,6 +1858,7 @@ namespace UniformTessellateLocals
 					}
 
 					Generator->CopyToAttribute(OutAttributes->GetPolygroupLayer(Idx));
+					OutAttributes->GetPolygroupLayer(Idx)->SetName(InAttributes->GetPolygroupLayer(Idx)->GetName());
 				}
 			}
 
@@ -1870,7 +1872,8 @@ namespace UniformTessellateLocals
 
 				FDynamicMeshVertexSkinWeightsAttribute* SkinAttribute = new FDynamicMeshVertexSkinWeightsAttribute(OutMesh);			
 				Generator.CopyToAttribute(SkinAttribute);
-
+				SkinAttribute->SetName(AttributeInfo.Value.Get()->GetName());
+				
 				OutAttributes->AttachSkinWeightsAttribute(AttributeInfo.Key, SkinAttribute);
 			}
 
@@ -1887,6 +1890,7 @@ namespace UniformTessellateLocals
 					}
 
 					Generator->CopyToAttribute(OutAttributes->GetWeightLayer(Idx));
+					OutAttributes->GetWeightLayer(Idx)->SetName(InAttributes->GetWeightLayer(Idx)->GetName());
 				}
 			}
 
