Handle the interpolation of the per-triangle attributes for the adaptive tessellation.

#rnx

diff --git a/Engine/Plugins/Runtime/GeometryProcessing/Source/DynamicMesh/Private/Operations/AdaptiveTessellate.cpp b/Engine/Plugins/Runtime/GeometryProcessing/Source/DynamicMesh/Private/Operations/AdaptiveTessellate.cpp
--- a/Engine/Plugins/Runtime/GeometryProcessing/Source/DynamicMesh/Private/Operations/AdaptiveTessellate.cpp
+++ b/Engine/Plugins/Runtime/GeometryProcessing/Source/DynamicMesh/Private/Operations/AdaptiveTessellate.cpp
@@ -1008,6 +1008,41 @@ namespace AdaptiveTessellateLocals
 	}
 
 
+	/**
+	 * Construct a new triangle attribute from the tessellation data and compact information from the geometry tessellation.
+	 */
+	template<typename RealType, int ElementSize> 
+	void ConstructTriangleAttribute(const FDynamicMesh3* Mesh, 
+						   			 const TDynamicMeshTriangleAttribute<RealType, ElementSize>* Attribute,
+						   			 FTessellationData& TessData,
+						   			 const bool bUseParallel, 
+						   			 const FCompactMaps& CompactInfo,
+						   			 TDynamicMeshTriangleAttribute<RealType, ElementSize>* OutAttribute) 
+	{
+		RealType Value[ElementSize];
+
+		for (const int TID : Mesh->TriangleIndicesItr())
+		{
+			Attribute->GetValue(TID, Value);
+
+			if (TessData.TrianglesToTessellate.Contains(TID))
+			{
+				TArrayView<int> TriangleIDBlock = TessData.MapTriangleIDBufferBlock(TID);
+				for (int Index = 0; Index < TriangleIDBlock.Num(); ++Index) 
+				{
+					const int NewTID = TriangleIDBlock[Index];
+					const int ToTID = CompactInfo.GetTriangleMapping(NewTID);		
+					OutAttribute->SetValue(ToTID, Value);
+				}
+			}
+			else 
+			{
+				const int ToTID = CompactInfo.GetTriangleMapping(TID);		
+				OutAttribute->SetValue(ToTID, Value);
+			}
+		}
+	}
+
 	/**
 	 * Given the original FDynamicMesh3 and the tessellation data, construct new tessellated FDynamicMesh3 geometry.
 	 * Output mesh will be compact. 
@@ -1083,14 +1118,20 @@ namespace AdaptiveTessellateLocals
 			}
 		}
 
+		if (Mesh->HasTriangleGroups()) 
+		{
+			ResultMesh->EnableTriangleGroups();
+		}
+
 		// Append all the triangles from the input mesh that we are not tessellating
 		for (const int TriangleID : Mesh->TriangleIndicesItr())
 		{
 			if (TessData.TrianglesToTessellate.Contains(TriangleID) == false)
 			{
 				const FIndex3i Tri = Mesh->GetTriangle(TriangleID);
+				const int TriGrp = Mesh->GetTriangleGroup(TriangleID);
 				const FIndex3i MappedTri = CompactInfo.GetVertexMapping(Tri);
-				CompactInfo.SetTriangleMapping(TriangleID, ResultMesh->AppendTriangle(MappedTri));
+				CompactInfo.SetTriangleMapping(TriangleID, ResultMesh->AppendTriangle(MappedTri, TriGrp));
 			}
 		}
 		
@@ -1100,13 +1141,14 @@ namespace AdaptiveTessellateLocals
 			TArrayView<FIndex3i> TrianglesBlock = TessData.MapTrianglesBufferBlock(TriangleID);
 			TArrayView<int> TriangleIDBlock = TessData.MapTriangleIDBufferBlock(TriangleID);
 			checkSlow(TrianglesBlock.Num() == TriangleIDBlock.Num());
-
+			const int TriGrp = Mesh->GetTriangleGroup(TriangleID);
+			
 			for (int Index = 0; Index < TriangleIDBlock.Num(); ++Index)
 			{
 				const int NewTriangleID = TriangleIDBlock[Index];
 				const FIndex3i Tri = TrianglesBlock[Index];
 				const FIndex3i MappedTri = CompactInfo.GetVertexMapping(Tri);
-				CompactInfo.SetTriangleMapping(NewTriangleID, ResultMesh->AppendTriangle(MappedTri));
+				CompactInfo.SetTriangleMapping(NewTriangleID, ResultMesh->AppendTriangle(MappedTri, TriGrp));
 			}
 		}
 
@@ -1313,6 +1355,23 @@ namespace AdaptiveTessellateLocals
 				OutAttributes->EnablePrimaryColors();
 				TessellateOverlay(InMesh, InAttributes->PrimaryColors(), Pattern, TessData, bUseParallel, CompactInfo, OutAttributes->PrimaryColors());
 			}
+
+			if (InAttributes->HasMaterialID())	 
+			{
+				OutAttributes->EnableMaterialID();
+				OutAttributes->GetMaterialID()->SetName(InAttributes->GetMaterialID()->GetName());
+				ConstructTriangleAttribute(InMesh, InAttributes->GetMaterialID(), TessData, bUseParallel, CompactInfo, OutAttributes->GetMaterialID());
+			}
+
+			if (InAttributes->NumPolygroupLayers())  
+			{
+				OutAttributes->SetNumPolygroupLayers(InAttributes->NumPolygroupLayers());
+				for (int Idx = 0; Idx < InAttributes->NumPolygroupLayers(); ++Idx) 
+				{	
+					OutAttributes->GetPolygroupLayer(Idx)->SetName(InAttributes->GetPolygroupLayer(Idx)->GetName());
+					ConstructTriangleAttribute(InMesh, InAttributes->GetPolygroupLayer(Idx), TessData, bUseParallel, CompactInfo, OutAttributes->GetPolygroupLayer(Idx));
+				}
+			}
 		}
 	}
 }
