Move GeometrySelection.h and GeometrySelectionUtil.h to GeometryProcessing plugin. Move types in GeometrySelection.h into UE::Geometry:: namespace. Fix up usage sites.

diff --git a/Engine/Plugins/Editor/ModelingToolsEditorMode/Source/ModelingToolsEditorMode/ModelingToolsEditorMode.Build.cs b/Engine/Plugins/Editor/ModelingToolsEditorMode/Source/ModelingToolsEditorMode/ModelingToolsEditorMode.Build.cs
--- a/Engine/Plugins/Editor/ModelingToolsEditorMode/Source/ModelingToolsEditorMode/ModelingToolsEditorMode.Build.cs
+++ b/Engine/Plugins/Editor/ModelingToolsEditorMode/Source/ModelingToolsEditorMode/ModelingToolsEditorMode.Build.cs
@@ -47,6 +47,8 @@ public class ModelingToolsEditorMode : ModuleRules
 				"InteractiveToolsFramework",
 				"EditorInteractiveToolsFramework",
 				"GeometryFramework",
+				"GeometryCore",
+				"DynamicMesh",
 				"ModelingComponents",
 				"ModelingComponentsEditorOnly",
 				"MeshModelingTools",
diff --git a/Engine/Plugins/Runtime/GeometryProcessing/Source/DynamicMesh/Private/Selections/GeometrySelection.cpp b/Engine/Plugins/Runtime/GeometryProcessing/Source/DynamicMesh/Private/Selections/GeometrySelection.cpp
new file mode 100644
--- /dev/null
+++ b/Engine/Plugins/Runtime/GeometryProcessing/Source/DynamicMesh/Private/Selections/GeometrySelection.cpp
@@ -0,0 +1,52 @@
+// Copyright Epic Games, Inc. All Rights Reserved.
+
+#include "Selections/GeometrySelection.h"
+#include "DynamicMesh/DynamicMesh3.h"
+
+using namespace UE::Geometry;
+
+void FGeometrySelectionEditor::Initialize(FGeometrySelection* TargetSelectionIn)
+{
+	check(TargetSelectionIn != nullptr);
+	TargetSelection = TargetSelectionIn;
+}
+
+void FGeometrySelectionEditor::ClearSelection(FGeometrySelectionDelta& DeltaOut)
+{
+	for (uint64 ID : TargetSelection->Selection)
+	{
+		DeltaOut.Removed.Add(ID);
+	}
+
+	TargetSelection->Reset();
+}
+
+
+bool FGeometrySelectionEditor::IsSelected(uint64 ID) const
+{
+	return TargetSelection->Selection.Contains(ID);
+}
+
+bool FGeometrySelectionEditor::Replace(const FGeometrySelection& NewSelection, FGeometrySelectionDelta& DeltaOut)
+{
+	// currently hard-enforcing this...maybe we can allow switches? 
+	if (ensure(TargetSelection->ElementType == NewSelection.ElementType) == false ||
+		ensure(TargetSelection->TopologyType == NewSelection.TopologyType) == false)
+	{
+		return false;
+	}
+
+	for (uint64 ID : TargetSelection->Selection)
+	{
+		DeltaOut.Removed.Add(ID);
+	}
+	TargetSelection->Selection.Reset();
+
+	for (uint64 ID : NewSelection.Selection)
+	{
+		TargetSelection->Selection.Add(ID);
+		DeltaOut.Added.Add(ID);
+	}
+
+	return true;
+}
\ No newline at end of file
diff --git a/Engine/Plugins/Runtime/GeometryProcessing/Source/DynamicMesh/Private/Selections/GeometrySelectionUtil.cpp b/Engine/Plugins/Runtime/GeometryProcessing/Source/DynamicMesh/Private/Selections/GeometrySelectionUtil.cpp
new file mode 100644
--- /dev/null
+++ b/Engine/Plugins/Runtime/GeometryProcessing/Source/DynamicMesh/Private/Selections/GeometrySelectionUtil.cpp
@@ -0,0 +1,837 @@
+// Copyright Epic Games, Inc. All Rights Reserved.
+
+
+#include "Selections/GeometrySelectionUtil.h"
+#include "DynamicMesh/DynamicMesh3.h"
+#include "DynamicMesh/ColliderMesh.h"
+#include "DynamicMesh/MeshNormals.h"
+#include "TriangleTypes.h"
+#include "SegmentTypes.h"
+#include "GroupTopology.h"
+#include "Selections/MeshConnectedComponents.h"
+
+
+void UE::Geometry::UpdateTriangleSelectionViaRaycast(
+	const FColliderMesh* ColliderMesh,
+	FGeometrySelectionEditor* Editor,
+	const FRay3d& LocalRay,
+	const FGeometrySelectionUpdateConfig& UpdateConfig,
+	FGeometrySelectionUpdateResult& ResultOut )
+{
+	check(Editor->GetTopologyType() == EGeometryTopologyType::Triangle);
+
+	ResultOut.bSelectionMissed = true;
+
+	double RayHitT; int32 HitTriangleID; FVector3d HitBaryCoords;
+	if (ColliderMesh->FindNearestHitTriangle(LocalRay, RayHitT, HitTriangleID, HitBaryCoords))
+	{
+		HitTriangleID = ColliderMesh->GetSourceTriangleID(HitTriangleID);
+		if (HitTriangleID == IndexConstants::InvalidID)
+		{
+			return;
+		}
+
+		if (Editor->GetElementType() == EGeometryElementType::Face)
+		{
+			ResultOut.bSelectionModified = UpdateSelectionWithNewElements(Editor, UpdateConfig.ChangeType, TArray<uint64>{(uint64)HitTriangleID}, &ResultOut.SelectionDelta);
+			ResultOut.bSelectionMissed = false;
+		}
+		else if (Editor->GetElementType() == EGeometryElementType::Vertex)
+		{
+			FVector3d HitPos = LocalRay.PointAt(RayHitT);
+			FIndex3i TriVerts = ColliderMesh->GetTriangle(HitTriangleID);
+			int32 NearestIdx = 0;
+			double NearestDistSqr = DistanceSquared(ColliderMesh->GetVertex(TriVerts[0]), HitPos);
+			for (int32 k = 1; k < 3; ++k)
+			{
+				double DistSqr = DistanceSquared(ColliderMesh->GetVertex(TriVerts[k]), HitPos);
+				if (DistSqr < NearestDistSqr)
+				{
+					NearestDistSqr = DistSqr;
+					NearestIdx = k;
+				}
+			}
+
+			ResultOut.bSelectionModified = UpdateSelectionWithNewElements(Editor, UpdateConfig.ChangeType, TArray<uint64>{(uint64)TriVerts[NearestIdx]}, &ResultOut.SelectionDelta);
+			ResultOut.bSelectionMissed = false;
+		}
+		else if (Editor->GetElementType() == EGeometryElementType::Edge)
+		{
+			FVector3d HitPos = LocalRay.PointAt(RayHitT);
+			FIndex3i TriVerts = ColliderMesh->GetTriangle(HitTriangleID);
+			FVector3d Positions[3];
+			ColliderMesh->GetTriVertices(HitTriangleID, Positions[0], Positions[1], Positions[2]);
+			int32 NearestIdx = 0;
+			double NearestDistSqr = FSegment3d(Positions[0], Positions[1]).DistanceSquared(HitPos);
+			for (int32 k = 1; k < 3; ++k)
+			{
+				double DistSqr = FSegment3d(Positions[k], Positions[(k+1)%3]).DistanceSquared(HitPos);
+				if (DistSqr < NearestDistSqr)
+				{
+					NearestDistSqr = DistSqr;
+					NearestIdx = k;
+				}
+			}
+			FMeshTriEdgeID TriEdgeID(HitTriangleID, NearestIdx);
+
+			ResultOut.bSelectionModified = UpdateSelectionWithNewElements(Editor, UpdateConfig.ChangeType, TArray<uint64>{(uint64)TriEdgeID.Encoded()}, &ResultOut.SelectionDelta);
+			ResultOut.bSelectionMissed = false;
+		}
+	}
+}
+
+
+
+
+
+void UE::Geometry::UpdateGroupSelectionViaRaycast(
+	const FColliderMesh* ColliderMesh,
+	const FGroupTopology* GroupTopology,
+	FGeometrySelectionEditor* Editor,
+	const FRay3d& LocalRay,
+	const FGeometrySelectionUpdateConfig& UpdateConfig,
+	FGeometrySelectionUpdateResult& ResultOut)
+{
+	check(Editor->GetTopologyType() == EGeometryTopologyType::Polygroup);
+
+	ResultOut.bSelectionMissed = true;
+
+	double RayHitT; int32 HitTriangleID; FVector3d HitBaryCoords;
+	if (ColliderMesh->FindNearestHitTriangle(LocalRay, RayHitT, HitTriangleID, HitBaryCoords))
+	{
+		HitTriangleID = ColliderMesh->GetSourceTriangleID(HitTriangleID);
+		if (HitTriangleID == IndexConstants::InvalidID)
+		{
+			return;
+		}
+		int32 GroupID = GroupTopology->GetGroupID(HitTriangleID);
+
+		if (Editor->GetElementType() == EGeometryElementType::Face)
+		{
+			FGeoSelectionID GroupTriID((uint32)HitTriangleID, (uint32)GroupID);
+			ResultOut.bSelectionModified = UpdateSelectionWithNewElements(Editor, UpdateConfig.ChangeType, TArray<uint64>{GroupTriID.Encoded()}, &ResultOut.SelectionDelta);
+			ResultOut.bSelectionMissed = false;
+		}
+		else if (Editor->GetElementType() == EGeometryElementType::Vertex)
+		{
+			FVector3d HitPos = LocalRay.PointAt(RayHitT);
+			FIndex3i TriVerts = ColliderMesh->GetTriangle(HitTriangleID);
+			int32 NearestIdx = -1;
+			double NearestDistSqr = DistanceSquared(ColliderMesh->GetVertex(TriVerts[0]), HitPos);
+			for (int32 k = 0; k < 3; ++k)
+			{
+				if (GroupTopology->GetCornerIDFromVertexID(TriVerts[k]) != IndexConstants::InvalidID)
+				{
+					double DistSqr = DistanceSquared(ColliderMesh->GetVertex(TriVerts[k]), HitPos);
+					if (DistSqr < NearestDistSqr)
+					{
+						NearestDistSqr = DistSqr;
+						NearestIdx = k;
+					}
+				}
+			}
+			if (NearestIdx >= 0)
+			{
+				// do we need a group here?
+				int32 VertexID = TriVerts[NearestIdx];
+				ResultOut.bSelectionModified = UpdateSelectionWithNewElements(Editor, UpdateConfig.ChangeType, TArray<uint64>{(uint64)VertexID}, &ResultOut.SelectionDelta);
+				ResultOut.bSelectionMissed = false;
+			}
+		}
+		else if (Editor->GetElementType() == EGeometryElementType::Edge)
+		{
+			FVector3d HitPos = LocalRay.PointAt(RayHitT);
+			FIndex3i TriVerts = ColliderMesh->GetTriangle(HitTriangleID);
+			FVector3d Positions[3];
+			ColliderMesh->GetTriVertices(HitTriangleID, Positions[0], Positions[1], Positions[2]);
+			int32 NearestIdx = -1;
+			double NearestDistSqr = TNumericLimits<double>::Max();
+			for (int32 k = 0; k < 3; ++k)
+			{
+				if ( GroupTopology->IsGroupEdge(FMeshTriEdgeID(HitTriangleID, k), true) )
+				{
+					double DistSqr = FSegment3d(Positions[k], Positions[(k + 1) % 3]).DistanceSquared(HitPos);
+					if (DistSqr < NearestDistSqr)
+					{
+						NearestDistSqr = DistSqr;
+						NearestIdx = k;
+					}
+				}
+			}
+			if ( NearestIdx >= 0 )
+			{
+				// do we need a group here?
+				FMeshTriEdgeID TriEdgeID(HitTriangleID, NearestIdx);
+				ResultOut.bSelectionModified = UpdateSelectionWithNewElements(Editor, UpdateConfig.ChangeType, TArray<uint64>{(uint64)TriEdgeID.Encoded()}, & ResultOut.SelectionDelta);
+				ResultOut.bSelectionMissed = false;
+			}
+		}
+	}
+}
+
+
+
+
+
+
+
+bool UE::Geometry::UpdateSelectionWithNewElements(
+	FGeometrySelectionEditor* Editor,
+	EGeometrySelectionChangeType ChangeType,
+	const TArray<uint64>& NewIDs,
+	FGeometrySelectionDelta* Delta)
+{
+	FGeometrySelectionDelta LocalDelta;
+	FGeometrySelectionDelta& UseDelta = (Delta != nullptr) ? (*Delta) : LocalDelta;
+
+	if (ChangeType == EGeometrySelectionChangeType::Replace)
+	{
+		// [TODO] this could be optimized...
+		Editor->ClearSelection(UseDelta);
+		return Editor->Select(NewIDs, UseDelta);
+	}
+	else if (ChangeType == EGeometrySelectionChangeType::Add)
+	{
+		return Editor->Select(NewIDs, UseDelta);
+	}
+	else if (ChangeType == EGeometrySelectionChangeType::Remove)
+	{
+		return Editor->Deselect(NewIDs, UseDelta);
+	}
+	else
+	{
+		ensure(false);
+		return false;
+	}
+	
+}
+
+
+
+bool UE::Geometry::EnumerateTriangleSelectionVertices(
+	const FGeometrySelection& MeshSelection,
+	const UE::Geometry::FDynamicMesh3& Mesh,
+	const FTransform& ApplyTransform,
+	TFunctionRef<void(uint64, const FVector3d&)> VertexFunc)
+{
+	if ( ensure( MeshSelection.TopologyType == EGeometryTopologyType::Triangle ) == false )
+	{
+		return false;
+	}
+
+	if (MeshSelection.ElementType == EGeometryElementType::Face)
+	{
+		for (uint64 TriangleID : MeshSelection.Selection)
+		{
+			if (Mesh.IsTriangle(TriangleID))
+			{
+				FIndex3i Triangle = Mesh.GetTriangle((int32)TriangleID);
+				VertexFunc((uint64)Triangle.A, ApplyTransform.TransformPosition(Mesh.GetVertex(Triangle.A)));
+				VertexFunc((uint64)Triangle.B, ApplyTransform.TransformPosition(Mesh.GetVertex(Triangle.B)));
+				VertexFunc((uint64)Triangle.C, ApplyTransform.TransformPosition(Mesh.GetVertex(Triangle.C)));
+			}
+		}
+	}
+	else if (MeshSelection.ElementType == EGeometryElementType::Edge)
+	{
+		for (uint64 EncodedID : MeshSelection.Selection)
+		{
+			FMeshTriEdgeID TriEdgeID((uint32)EncodedID);
+			int32 EdgeID = Mesh.IsTriangle(TriEdgeID.TriangleID) ? Mesh.GetTriEdge(TriEdgeID.TriangleID, TriEdgeID.TriEdgeIndex) : IndexConstants::InvalidID;
+			if (Mesh.IsEdge(EdgeID))
+			{
+				FIndex2i EdgeV = Mesh.GetEdgeV(EdgeID);
+				VertexFunc((uint64)EdgeV.A, ApplyTransform.TransformPosition(Mesh.GetVertex(EdgeV.A)));
+				VertexFunc((uint64)EdgeV.B, ApplyTransform.TransformPosition(Mesh.GetVertex(EdgeV.B)));
+			}
+		}
+	}
+	else if (MeshSelection.ElementType == EGeometryElementType::Vertex)
+	{
+		for (uint64 VertexID : MeshSelection.Selection)
+		{
+			if (Mesh.IsVertex((int32)VertexID))
+			{
+				VertexFunc((uint64)VertexID, ApplyTransform.TransformPosition(Mesh.GetVertex(VertexID)));
+			}
+		}
+	}
+	else
+	{
+		return false;
+	}
+
+	return true;
+}
+
+
+
+bool UE::Geometry::EnumeratePolygroupSelectionVertices(
+	const FGeometrySelection& MeshSelection,
+	const UE::Geometry::FDynamicMesh3& Mesh,
+	const FGroupTopology* GroupTopology,
+	const FTransform& ApplyTransform,
+	TFunctionRef<void(uint64, const FVector3d&)> VertexFunc)
+{
+	if ( ensure( MeshSelection.TopologyType == EGeometryTopologyType::Polygroup ) == false )
+	{
+		return false;
+	}
+
+	if (MeshSelection.ElementType == EGeometryElementType::Face)
+	{
+		for (uint64 EncodedID : MeshSelection.Selection)
+		{
+			FGeoSelectionID GroupTriID(EncodedID);
+			int32 SeedTriangleID = (int32)GroupTriID.GeometryID, GroupID = (int32)GroupTriID.TopologyID;
+			if (Mesh.IsTriangle(SeedTriangleID))
+			{
+				for (int32 TriangleID : GroupTopology->GetGroupFaces(GroupID))
+				{
+					FIndex3i Triangle = Mesh.GetTriangle((int32)TriangleID);
+					VertexFunc((uint64)Triangle.A, ApplyTransform.TransformPosition(Mesh.GetVertex(Triangle.A)));
+					VertexFunc((uint64)Triangle.B, ApplyTransform.TransformPosition(Mesh.GetVertex(Triangle.B)));
+					VertexFunc((uint64)Triangle.C, ApplyTransform.TransformPosition(Mesh.GetVertex(Triangle.C)));
+				}
+			}
+		}
+	}
+	else if (MeshSelection.ElementType == EGeometryElementType::Edge)
+	{
+		for (uint64 EncodedID : MeshSelection.Selection)
+		{
+			FMeshTriEdgeID TriEdgeID( FGeoSelectionID(EncodedID).GeometryID );
+			int32 SeedEdgeID = Mesh.IsTriangle(TriEdgeID.TriangleID) ? Mesh.GetTriEdge(TriEdgeID.TriangleID, TriEdgeID.TriEdgeIndex) : IndexConstants::InvalidID;
+			if (Mesh.IsEdge(SeedEdgeID))
+			{
+				int32 GroupEdgeID = GroupTopology->FindGroupEdgeID(SeedEdgeID);
+				for (int32 VertexID : GroupTopology->GetGroupEdgeVertices(GroupEdgeID))
+				{
+					FVector3d V = Mesh.GetVertex(VertexID);
+					VertexFunc((uint64)VertexID, ApplyTransform.TransformPosition(V));
+				}
+			}
+		}
+	}
+	else if (MeshSelection.ElementType == EGeometryElementType::Vertex)
+	{
+		for (uint64 EncodedID : MeshSelection.Selection)
+		{
+			int32 VertexID = (int32)FGeoSelectionID(EncodedID).GeometryID;
+			if (Mesh.IsVertex(VertexID))
+			{
+				VertexFunc((uint64)VertexID, ApplyTransform.TransformPosition(Mesh.GetVertex(VertexID)));
+			}
+		}
+	}
+	else
+	{
+		return false;
+	}
+
+	return true;
+}
+
+
+
+bool UE::Geometry::EnumerateSelectionTriangles(
+	const FGeometrySelection& MeshSelection,
+	const UE::Geometry::FDynamicMesh3& Mesh,
+	TFunctionRef<void(int32)> TriangleFunc,
+	const UE::Geometry::FPolygroupSet* UseGroupSet
+)
+{
+	if (MeshSelection.TopologyType == EGeometryTopologyType::Triangle)
+	{
+		return EnumerateTriangleSelectionTriangles(MeshSelection, Mesh, TriangleFunc);
+	}
+	else if (MeshSelection.TopologyType == EGeometryTopologyType::Polygroup)
+	{
+		return EnumeratePolygroupSelectionTriangles(MeshSelection, Mesh, 
+			(UseGroupSet != nullptr) ? *UseGroupSet : FPolygroupSet(&Mesh), 
+			TriangleFunc);
+	}
+	return false;
+}
+
+
+bool UE::Geometry::EnumerateTriangleSelectionTriangles(
+	const FGeometrySelection& MeshSelection,
+	const UE::Geometry::FDynamicMesh3& Mesh,
+	TFunctionRef<void(int32)> TriangleFunc)
+{
+	if ( ensure( MeshSelection.TopologyType == EGeometryTopologyType::Triangle ) == false )
+	{
+		return false;
+	}
+
+	if (MeshSelection.ElementType == EGeometryElementType::Face)
+	{
+		for (uint64 TriangleID : MeshSelection.Selection)
+		{
+			if (Mesh.IsTriangle(TriangleID))
+			{
+				TriangleFunc(TriangleID);
+			}
+		}
+	}
+	else if (MeshSelection.ElementType == EGeometryElementType::Edge)
+	{
+		for (uint64 EncodedID : MeshSelection.Selection)
+		{
+			FMeshTriEdgeID TriEdgeID( FGeoSelectionID(EncodedID).GeometryID );
+			int32 EdgeID = Mesh.IsTriangle(TriEdgeID.TriangleID) ? Mesh.GetTriEdge(TriEdgeID.TriangleID, TriEdgeID.TriEdgeIndex) : IndexConstants::InvalidID;
+			Mesh.EnumerateEdgeTriangles(EdgeID, [&](int32 TriangleID)
+			{
+				TriangleFunc(TriangleID);
+			});
+		}
+	}
+	else if (MeshSelection.ElementType == EGeometryElementType::Vertex)
+	{
+		for (uint64 VertexID : MeshSelection.Selection)
+		{
+			Mesh.EnumerateVertexTriangles(VertexID, [&](int32 TriangleID)
+			{
+				TriangleFunc(TriangleID);
+			});
+		}
+	}
+	else
+	{
+		return false;
+	}
+
+	return true;
+}
+
+
+bool UE::Geometry::EnumeratePolygroupSelectionTriangles(
+	const FGeometrySelection& MeshSelection,
+	const UE::Geometry::FDynamicMesh3& Mesh,
+	const UE::Geometry::FPolygroupSet& GroupSet,
+	TFunctionRef<void(int32)> TriangleFunc
+)
+{
+	if ( ensure( MeshSelection.TopologyType == EGeometryTopologyType::Polygroup ) == false )
+	{
+		return false;
+	}
+
+	TArray<int32> SeedGroups;
+	TArray<int32> SeedTriangles;
+	TSet<int32> UniqueSeedGroups;
+
+	// TODO: the code below will not work correctly if the selection contains
+	// multiple disconnected-components with the same GroupID. They will be
+	// filtered out by the UniqueSeedGroups test. Seems like it will be necessary
+	// to detect this case up-front and do something more expensive, like filtering
+	// out duplicates inside the connected-components loop instead of up-front
+
+	if (MeshSelection.ElementType == EGeometryElementType::Face)
+	{
+		for (uint64 EncodedID : MeshSelection.Selection)
+		{
+			FGeoSelectionID SelectionID(EncodedID);
+			int32 SeedTriangleID = (int32)SelectionID.GeometryID;
+			if (Mesh.IsTriangle(SeedTriangleID))
+			{
+				int32 GroupID = GroupSet.GetGroup(SeedTriangleID);
+				check(GroupID == (int32)SelectionID.TopologyID);		// sanity-check that we are using the right group
+				if ( GroupID >= 0 && UniqueSeedGroups.Contains(GroupID) == false)
+				{
+					UniqueSeedGroups.Add(GroupID);
+					SeedGroups.Add(GroupID);
+					SeedTriangles.Add(SeedTriangleID);
+				}
+			}
+		}
+	}
+	else if (MeshSelection.ElementType == EGeometryElementType::Edge)
+	{
+		for (uint64 EncodedID : MeshSelection.Selection)
+		{
+			FMeshTriEdgeID TriEdgeID( FGeoSelectionID(EncodedID).GeometryID );
+			int32 SeedEdgeID = Mesh.IsTriangle(TriEdgeID.TriangleID) ? Mesh.GetTriEdge(TriEdgeID.TriangleID, TriEdgeID.TriEdgeIndex) : IndexConstants::InvalidID;
+			if (Mesh.IsEdge(SeedEdgeID))
+			{
+				Mesh.EnumerateEdgeTriangles(SeedEdgeID, [&](int32 TriangleID)
+				{
+					int32 GroupID = GroupSet.GetGroup(TriangleID);
+					if (GroupID >= 0 && UniqueSeedGroups.Contains(GroupID) == false)
+					{
+						UniqueSeedGroups.Add(GroupID);
+						SeedGroups.Add(GroupID);
+						SeedTriangles.Add(TriangleID);
+					}
+				});
+			}
+		}
+	}
+	else if (MeshSelection.ElementType == EGeometryElementType::Vertex)
+	{
+		for (uint64 EncodedID : MeshSelection.Selection)
+		{
+			int32 VertexID = (int32)FGeoSelectionID(EncodedID).GeometryID;
+			if (Mesh.IsVertex(VertexID))
+			{
+				Mesh.EnumerateVertexTriangles(VertexID, [&](int32 TriangleID)
+				{
+					int32 GroupID = GroupSet.GetGroup(TriangleID);
+					if (GroupID >= 0 && UniqueSeedGroups.Contains(GroupID) == false)
+					{
+						UniqueSeedGroups.Add(GroupID);
+						SeedGroups.Add(GroupID);
+						SeedTriangles.Add(TriangleID);
+					}
+				});
+			}
+		}
+	}
+	else
+	{
+		return false;
+	}
+
+
+	TSet<int> TempROI;		// if we could provide this as input we would not need a temporary roi...
+	TArray<int32> QueueBuffer;
+	int32 NumGroups = SeedGroups.Num();
+	for (int32 k = 0; k < NumGroups; ++k)
+	{
+		check(GroupSet.GetGroup(SeedTriangles[k]) == SeedGroups[k]);
+		int32 GroupID = SeedGroups[k];
+		FMeshConnectedComponents::GrowToConnectedTriangles(&Mesh, 
+			TArray<int>{SeedTriangles[k]}, TempROI, &QueueBuffer, 
+			[&](int32 T1, int32 T2) { return GroupSet.GetGroup(T2) == GroupID; });
+		for (int32 tid : TempROI)
+		{
+			TriangleFunc(tid);
+		}
+	}
+
+	return true;
+}
+
+
+
+
+bool UE::Geometry::EnumerateTriangleSelectionElements(
+	const FGeometrySelection& MeshSelection,
+	const UE::Geometry::FDynamicMesh3& Mesh,
+	TFunctionRef<void(int32, const FVector3d&)> VertexFunc,
+	TFunctionRef<void(int32, const FSegment3d&)> EdgeFunc,
+	TFunctionRef<void(int32, const FTriangle3d&)> TriangleFunc,
+	const FTransform* ApplyTransform
+)
+{
+	if ( ensure( MeshSelection.TopologyType == EGeometryTopologyType::Triangle ) == false )
+	{
+		return false;
+	}
+
+	if (MeshSelection.ElementType == EGeometryElementType::Face)
+	{
+		for (uint64 EncodedID : MeshSelection.Selection)
+		{
+			int32 TriangleID = (int32)FGeoSelectionID(EncodedID).GeometryID;
+			if (Mesh.IsTriangle(TriangleID))
+			{
+				FVector3d A, B, C;
+				Mesh.GetTriVertices((int32)TriangleID, A, B, C);
+				if (ApplyTransform != nullptr)
+				{
+					A = ApplyTransform->TransformPosition(A);
+					B = ApplyTransform->TransformPosition(B);
+					C = ApplyTransform->TransformPosition(C);
+				}
+				TriangleFunc(TriangleID, FTriangle3d(A, B, C));
+			}
+		}
+	}
+	else if (MeshSelection.ElementType == EGeometryElementType::Edge)
+	{
+		for (uint64 EncodedID : MeshSelection.Selection)
+		{
+			FMeshTriEdgeID TriEdgeID( FGeoSelectionID(EncodedID).GeometryID );
+			int32 EdgeID = Mesh.IsTriangle(TriEdgeID.TriangleID) ? Mesh.GetTriEdge(TriEdgeID.TriangleID, TriEdgeID.TriEdgeIndex) : IndexConstants::InvalidID;
+			if (Mesh.IsEdge(EdgeID))
+			{
+				FVector3d A, B;
+				Mesh.GetEdgeV(EdgeID, A, B);
+				if (ApplyTransform != nullptr)
+				{
+					A = ApplyTransform->TransformPosition(A);
+					B = ApplyTransform->TransformPosition(B);
+				}
+				EdgeFunc(EdgeID, FSegment3d(A, B));
+			}
+		}
+	}
+	else if (MeshSelection.ElementType == EGeometryElementType::Vertex)
+	{
+		for (uint64 EncodedID : MeshSelection.Selection)
+		{
+			int32 VertexID = (int32)FGeoSelectionID(EncodedID).GeometryID;
+			if (Mesh.IsVertex(VertexID))
+			{
+				FVector3d A = Mesh.GetVertex(VertexID);
+				VertexFunc(VertexID, (ApplyTransform != nullptr) ? ApplyTransform->TransformPosition(A) : A);
+			}
+		}
+	}
+	else
+	{
+		return false;
+	}
+
+	return true;
+}
+
+
+
+
+
+bool UE::Geometry::EnumeratePolygroupSelectionElements(
+	const FGeometrySelection& MeshSelection,
+	const UE::Geometry::FDynamicMesh3& Mesh,
+	const FGroupTopology* GroupTopology,
+	TFunctionRef<void(int32, const FVector3d&)> VertexFunc,
+	TFunctionRef<void(int32, const FSegment3d&)> EdgeFunc,
+	TFunctionRef<void(int32, const FTriangle3d&)> TriangleFunc,
+	const FTransform* ApplyTransform
+)
+{
+	if ( ensure( MeshSelection.TopologyType == EGeometryTopologyType::Polygroup ) == false )
+	{
+		return false;
+	}
+
+	if (MeshSelection.ElementType == EGeometryElementType::Face)
+	{
+		for (uint64 EncodedID : MeshSelection.Selection)
+		{
+			FGeoSelectionID SelectionID(EncodedID);
+			int32 SeedTriangleID = (int32)SelectionID.GeometryID, GroupID = (int32)SelectionID.TopologyID;
+			if (Mesh.IsTriangle(SeedTriangleID))
+			{
+				for (int32 TriangleID : GroupTopology->GetGroupFaces(GroupID))
+				{
+					FVector3d A, B, C;
+					Mesh.GetTriVertices((int32)TriangleID, A, B, C);
+					if (ApplyTransform != nullptr)
+					{
+						A = ApplyTransform->TransformPosition(A);
+						B = ApplyTransform->TransformPosition(B);
+						C = ApplyTransform->TransformPosition(C);
+					}
+					TriangleFunc(TriangleID, FTriangle3d(A, B, C));
+				}
+			}
+		}
+	}
+	else if (MeshSelection.ElementType == EGeometryElementType::Edge)
+	{
+		for (uint64 EncodedID : MeshSelection.Selection)
+		{
+			FMeshTriEdgeID TriEdgeID( FGeoSelectionID(EncodedID).GeometryID );
+			int32 SeedEdgeID = Mesh.IsTriangle(TriEdgeID.TriangleID) ? Mesh.GetTriEdge(TriEdgeID.TriangleID, TriEdgeID.TriEdgeIndex) : IndexConstants::InvalidID;
+			if (Mesh.IsEdge(SeedEdgeID))
+			{
+				int32 GroupEdgeID = GroupTopology->FindGroupEdgeID(SeedEdgeID);
+				for (int32 EdgeID : GroupTopology->GetGroupEdgeEdges(GroupEdgeID))
+				{
+					FVector3d A, B;
+					Mesh.GetEdgeV(EdgeID, A, B);
+					if (ApplyTransform != nullptr)
+					{
+						A = ApplyTransform->TransformPosition(A);
+						B = ApplyTransform->TransformPosition(B);
+					}
+					EdgeFunc(EdgeID, FSegment3d(A, B));
+				}
+			}
+		}
+	}
+	else if (MeshSelection.ElementType == EGeometryElementType::Vertex)
+	{
+		for (uint64 EncodedID : MeshSelection.Selection)
+		{
+			int32 VertexID = (int32)FGeoSelectionID(EncodedID).GeometryID;
+			if (Mesh.IsVertex(VertexID))
+			{
+				FVector3d A = Mesh.GetVertex(VertexID);
+				VertexFunc(VertexID, (ApplyTransform != nullptr) ? ApplyTransform->TransformPosition(A) : A);
+			}
+		}
+	}
+	else
+	{
+		return false;
+	}
+
+	return true;
+}
+
+
+
+
+bool UE::Geometry::ConvertPolygroupSelectionToTopologySelection(
+	const FGeometrySelection& MeshSelection,
+	const UE::Geometry::FDynamicMesh3& Mesh,
+	const FGroupTopology* GroupTopology,
+	FGroupTopologySelection& TopologySelectionOut)
+{
+	if ( ensure( MeshSelection.TopologyType == EGeometryTopologyType::Polygroup ) == false )
+	{
+		return false;
+	}
+
+	if (MeshSelection.ElementType == EGeometryElementType::Face)
+	{
+		for (uint64 EncodedID : MeshSelection.Selection)
+		{
+			int32 GroupID = FGeoSelectionID(EncodedID).TopologyID;
+			if (GroupTopology->FindGroupByID(GroupID) != nullptr)
+			{
+				TopologySelectionOut.SelectedGroupIDs.Add(GroupID);
+			}
+		}
+	}
+	else if (MeshSelection.ElementType == EGeometryElementType::Edge)
+	{
+		for (uint64 EncodedID : MeshSelection.Selection)
+		{
+			FMeshTriEdgeID MeshEdgeID( FGeoSelectionID(EncodedID).GeometryID );
+			if ( Mesh.IsTriangle((int32)MeshEdgeID.TriangleID) )
+			{
+				int32 GroupEdgeID = GroupTopology->FindGroupEdgeID(MeshEdgeID);
+				if (GroupEdgeID >= 0)
+				{
+					TopologySelectionOut.SelectedEdgeIDs.Add(GroupEdgeID);
+				}
+			}
+		}
+	}
+	else if (MeshSelection.ElementType == EGeometryElementType::Vertex)
+	{
+		for (uint64 EncodedID : MeshSelection.Selection)
+		{
+			int32 VertexID = (int32)FGeoSelectionID(EncodedID).GeometryID;
+			if ( Mesh.IsVertex(VertexID) )
+			{
+				int32 CornerID = GroupTopology->GetCornerIDFromVertexID(VertexID);
+				if (CornerID >= 0)
+				{
+					TopologySelectionOut.SelectedCornerIDs.Add(CornerID);
+				}
+			}
+		}
+	}
+	else
+	{
+		return false;
+	}
+
+	return true;
+}
+
+
+
+
+bool UE::Geometry::GetTriangleSelectionFrame(
+	const FGeometrySelection& MeshSelection,
+	const UE::Geometry::FDynamicMesh3& Mesh,
+	FFrame3d& SelectionFrameOut)
+{
+	if ( ensure( MeshSelection.TopologyType == EGeometryTopologyType::Triangle ) == false )
+	{
+		return false;
+	}
+
+	FVector3d AccumulatedOrigin = FVector3d::Zero();
+	FVector3d AccumulatedNormal = FVector3d::Zero();
+	FVector3d AxisHint = FVector3d::Zero();
+	double AccumWeight = 0;
+	
+	if (MeshSelection.ElementType == EGeometryElementType::Face)
+	{
+		for (uint64 EncodedID : MeshSelection.Selection)
+		{
+			int32 TriangleID = (int32)FGeoSelectionID(EncodedID).GeometryID;
+			if (Mesh.IsTriangle(TriangleID))
+			{
+				FVector3d Normal, Centroid; double Area; 
+				Mesh.GetTriInfo(TriangleID, Normal, Area, Centroid);
+				if (Normal.SquaredLength() > 0.9)
+				{
+					Area = FMath::Max(Area, 0.000001);
+					AccumulatedOrigin += Area * Centroid;
+					AccumulatedNormal += Area * Normal;
+					AccumWeight += Area;
+				}
+			}
+		}
+	}
+	else if (MeshSelection.ElementType == EGeometryElementType::Edge)
+	{
+		for (uint64 EncodedID : MeshSelection.Selection)
+		{
+			FMeshTriEdgeID TriEdgeID( FGeoSelectionID(EncodedID).GeometryID );
+			int32 EdgeID = Mesh.IsTriangle(TriEdgeID.TriangleID) ? Mesh.GetTriEdge(TriEdgeID.TriangleID, TriEdgeID.TriEdgeIndex) : IndexConstants::InvalidID;
+			if ( Mesh.IsEdge(EdgeID) )
+			{
+				FVector3d A, B;
+				Mesh.GetEdgeV(EdgeID, A, B);
+				AccumulatedOrigin += (A + B) * 0.5;
+				AccumulatedNormal += Mesh.GetEdgeNormal(EdgeID);
+				AxisHint += Normalized(B - A);
+				AccumWeight += 1.0;
+			}
+		}
+	}
+	else if (MeshSelection.ElementType == EGeometryElementType::Vertex)
+	{
+		for (uint64 EncodedID : MeshSelection.Selection)
+		{
+			int32 VertexID = (int32)FGeoSelectionID(EncodedID).GeometryID;
+			if (Mesh.IsVertex(VertexID))
+			{
+				AccumulatedOrigin += Mesh.GetVertex(VertexID);
+				AccumulatedNormal += FMeshNormals::ComputeVertexNormal(Mesh, VertexID);		// this could return area!
+				AccumWeight += 1.0;
+			}
+		}
+	}
+	else
+	{
+		return false;
+	}
+
+	// todo use AxisHint!
+
+	SelectionFrameOut = FFrame3d();
+	if (AccumWeight > 0)
+	{
+		AccumulatedOrigin /= (double)AccumWeight;
+		Normalize(AccumulatedNormal);
+
+		// We set our frame Z to be accumulated normal, and the other two axes are unconstrained, so
+		// we want to set them to something that will make our frame generally more useful. If the normal
+		// is aligned with world Z, then the entire frame might as well be aligned with world.
+		if (1 - AccumulatedNormal.Dot(FVector3d::UnitZ()) < KINDA_SMALL_NUMBER)
+		{
+			SelectionFrameOut = FFrame3d(AccumulatedOrigin, FQuaterniond::Identity());
+		}
+		else
+		{
+			// Otherwise, let's place one of the other axes into the XY plane so that the frame is more
+			// useful for translation. We somewhat arbitrarily choose Y for this. 
+			FVector3d FrameY = Normalized(AccumulatedNormal.Cross(FVector3d::UnitZ())); // orthogonal to world Z and frame Z 
+			FVector3d FrameX = FrameY.Cross(AccumulatedNormal); // safe to not normalize because already orthogonal
+			SelectionFrameOut = FFrame3d(AccumulatedOrigin, FrameX, FrameY, AccumulatedNormal);
+		}
+	}
+
+	return true;
+}
\ No newline at end of file
diff --git a/Engine/Plugins/Runtime/GeometryProcessing/Source/DynamicMesh/Public/Selections/GeometrySelection.h b/Engine/Plugins/Runtime/GeometryProcessing/Source/DynamicMesh/Public/Selections/GeometrySelection.h
new file mode 100644
--- /dev/null
+++ b/Engine/Plugins/Runtime/GeometryProcessing/Source/DynamicMesh/Public/Selections/GeometrySelection.h
@@ -0,0 +1,343 @@
+// Copyright Epic Games, Inc. All Rights Reserved.
+
+#pragma once
+
+#include "Math/IntVector.h"
+#include "TriangleTypes.h"
+#include "SegmentTypes.h"
+#include "BoxTypes.h"
+#include "DynamicMesh/InfoTypes.h"
+
+
+namespace UE
+{
+namespace Geometry
+{
+
+
+/**
+ * FGeoSelectionID provides a pair of 32-bit unsigned integers that can
+ * be packed into a 64-bit unsigned integer for use with FGeometrySelection.
+ * This is generally intended to be used to encode a mesh geometry ID
+ * (eg TriangleID, EdgeID, VertexID) combined with a "Topology ID",
+ * eg something like a Face Group ID. However none of this is enforced
+ * and so a caller can use these two integers for any purpose.
+ * 
+ * Note that since the ints are unsigned, IndexConstants::InvalidID
+ * is not directly representable (-1 will become positive 0xFFFFFFFF). 
+ */
+struct DYNAMICMESH_API FGeoSelectionID
+{
+	/** Topology ID, stored in upper 32 bits when packed into 64-bits */
+	uint32 TopologyID;
+	/** Geometry ID, stored in lower 32 bits when packed into 64-bits */
+	uint32 GeometryID;
+
+	FGeoSelectionID()
+	{
+		GeometryID = 0;
+		TopologyID = 0;
+	}
+
+	/**
+	 * Initialize the TopologyID and GeometryID with the given values
+	 */
+	explicit FGeoSelectionID(uint32 GeometryIDIn, uint32 TopologyIDIn = 0)
+	{
+		GeometryID = GeometryIDIn;
+		TopologyID = TopologyIDIn;
+	}
+
+	/**
+	 * Initialize the TopologyID and GeometryID by unpacking the 64-bit packed EncodedID
+	 */
+	explicit FGeoSelectionID(uint64 EncodedID)
+	{
+		GeometryID = EncodedID & 0x00000000FFFFFFFF;
+		TopologyID = (EncodedID & 0xFFFFFFFF00000000) >> 32;
+	}
+
+	/**
+	 * @return the packed 64-bit representation of [TopologyID, GeometryID]
+	 */
+	uint64 Encoded() const
+	{
+		return ((uint64)TopologyID << 32) | (uint64)GeometryID;
+	}
+
+	/**
+	 * @return a FGeoSelectionID initialized with the given TriangleID, and a TopologyID of 0
+	 */
+	static FGeoSelectionID MeshTriangle(int32 TriangleID)
+	{
+		return FGeoSelectionID((uint32)TriangleID, 0);
+	}
+
+	/**
+	 * @return a FGeoSelectionID initialized with the given VertexID, and a TopologyID of 0
+	 */
+	static FGeoSelectionID MeshVertex(int32 VertexID)
+	{
+		return FGeoSelectionID((uint32)VertexID, 0);
+	}
+
+	/**
+	 * @return a FGeoSelectionID initialized with the given 32-bit encoding of the MeshTriEdgeID, and a TopologyID of 0
+	 */
+	static FGeoSelectionID MeshEdge(UE::Geometry::FMeshTriEdgeID EdgeKey)
+	{
+		return FGeoSelectionID(EdgeKey.Encoded());
+	}
+
+	/**
+	 * @return a FGeoSelectionID initialized with the given TriangleID and GroupID used as the TopologyID
+	 */
+	static FGeoSelectionID GroupFace(int32 TriangleID, int32 GroupID)
+	{
+		return FGeoSelectionID((uint32)TriangleID, (uint32)GroupID);
+	}
+
+	friend uint32 GetTypeHash(const FGeoSelectionID& TopoKey)
+	{
+		return ::GetTypeHash(TopoKey.Encoded());
+	}
+};
+
+
+
+
+/**
+ * Type of selected Elements in a FGeometrySelection
+ */
+enum class DYNAMICMESH_API EGeometryElementType
+{
+	/** Mesh Vertices, Polygroup Corners, ... */
+	Vertex = 1,
+	/** Mesh Edges, Polygroup Edges, ... */
+	Edge = 2,
+	/** Mesh Triangles, Polygroup Faces, ... */
+	Face = 4
+};
+
+/**
+ * Type of selected Topology in a FGeometrySelection.
+ */
+enum class DYNAMICMESH_API EGeometryTopologyType
+{
+	Triangle = 1,
+	Polygroup = 2
+};
+
+
+/**
+ * FGeometrySelection represents a subset of geometric elements of a larger
+ * object, for example a Mesh (currently the only use case). The main selection
+ * is represented via 64-bit unsigned integers. The integers are stored in a TSet
+ * for efficient unique adds and removals. No assumptions are made about the values,
+ * they could be (eg) mesh indices, IDs of some type, or even pointer values.
+ */
+struct DYNAMICMESH_API FGeometrySelection
+{
+	/** Type of geometric element represented by this selection, if applicable */
+	EGeometryElementType ElementType = EGeometryElementType::Face;
+
+	/** Type of geometric topology this selection is defined relative to, if applicable */
+	EGeometryTopologyType TopologyType = EGeometryTopologyType::Triangle;
+
+	/** Set of selected items/elements */
+	TSet<uint64> Selection;
+
+	/**
+	 * @return true if Selection is empty
+	 */
+	bool IsEmpty() const
+	{
+		return Selection.Num() == 0;
+	}
+
+	/**
+	 * @return number of elements in Selection
+	 */
+	int32 Num() const
+	{
+		return Selection.Num();
+	}
+
+	/**
+	 * Clear the Selection (may not release memory)
+	 */
+	void Reset()
+	{
+		Selection.Reset();
+	}
+
+	/**
+	 * Initialize the Element and Topology types for this Selection based on another Selection
+	 */
+	void InitializeTypes(const FGeometrySelection& FromSelection)
+	{
+		ElementType = FromSelection.ElementType;
+		TopologyType = FromSelection.TopologyType;
+	}
+};
+
+
+
+/**
+ * Type of change, relative to a FGeometrySelection
+ */
+enum class DYNAMICMESH_API EGeometrySelectionChangeType
+{
+	/** Elements Added to Selection */
+	Add,
+	/** Elements Removed from Selection */
+	Remove,
+	/** Selected Elements replaced with a new set of Selected Elements */
+	Replace
+};
+
+
+/**
+ * FGeometrySelectionUpdateConfig is passed to various Selection Editing functions/classes
+ * to indicate what type of change should be applied to a FGeometrySelection
+ * (based on additional parameters, generally)
+ */
+struct DYNAMICMESH_API FGeometrySelectionUpdateConfig
+{
+	EGeometrySelectionChangeType ChangeType = EGeometrySelectionChangeType::Add;
+};
+
+
+/**
+ * FGeometrySelectionDelta represents a change to the set of elements in a FGeometrySelection.
+ * The delta is ordered, ie if the Delta was to be re-applied, the Removed elements should
+ * be removed before the Added elements are added.
+ * 
+ * (Currently there is no way to swap the order)
+ */
+struct DYNAMICMESH_API FGeometrySelectionDelta
+{
+	/**
+	 * Elements removed from a FGeometrySelection during some selection edit
+	 */
+	TArray<uint64> Removed;
+	/**
+	 * Elements added to a FGeometrySelection during some selection edit
+	 */
+	TArray<uint64> Added;
+
+	/** @return true if the Delta is empty (nothing Removed or Added) */
+	bool IsEmpty() const { return Removed.Num() == 0 && Added.Num() == 0; }
+};
+
+
+
+/**
+ * 3D Bounding information for a FGeometrySelection
+ */
+struct DYNAMICMESH_API FGeometrySelectionBounds
+{
+	UE::Geometry::FAxisAlignedBox3d WorldBounds = UE::Geometry::FAxisAlignedBox3d::Empty();
+};
+
+
+/**
+ * 3D Geometry representing a FGeometrySelection, for example
+ * suitable for passing to rendering code, etc
+ */
+struct DYNAMICMESH_API FGeometrySelectionElements
+{
+	TArray<UE::Geometry::FTriangle3d> Triangles;
+	TArray<UE::Geometry::FSegment3d> Segments;
+	TArray<FVector3d> Points;
+
+	void Reset()
+	{
+		Triangles.Reset();
+		Segments.Reset();
+		Points.Reset();
+	}
+};
+
+
+struct DYNAMICMESH_API FGeometrySelectionUpdateResult
+{
+	bool bSelectionMissed = false;
+	bool bSelectionModified = false;
+
+	FGeometrySelectionDelta SelectionDelta;
+};
+
+
+
+/**
+ * FGeometrySelectionEditor is a utility/helper class used for modifying
+ * a FGeometrySelection. The various functions can be used to add or remove
+ * to the Selection, while also tracking what changed, returned via
+ * FGeometrySelectionDelta structs. 
+ */
+class DYNAMICMESH_API FGeometrySelectionEditor
+{
+public:
+	/**
+	 * Initialize the Editor with the given Selection. The
+	 * TargetSelectionIn must live longer than the FGeometrySelectionEditor
+	 */
+	void Initialize(FGeometrySelection* TargetSelectionIn);
+
+	/** @return the Element Type of the Target Selection */
+	EGeometryElementType GetElementType() const { return TargetSelection->ElementType; }
+	/** @return the Topology Type of the Target Selection */
+	EGeometryTopologyType GetTopologyType() const { return TargetSelection->TopologyType; }
+
+	/** @return true if the given ID is currently selected in the Target Selection */
+	bool IsSelected(uint64 ID) const;
+
+	/** Clear the Target Selection and return change information in DeltaOut */
+	void ClearSelection(FGeometrySelectionDelta& DeltaOut);
+
+	/** Add the items in the List to the Target Selection and return change information in DeltaOut */
+	template<typename ListType>
+	bool Select(const ListType& List, FGeometrySelectionDelta& DeltaOut)
+	{
+		int32 NumAdded = 0;
+		for (uint64 ID : List)
+		{
+			if (TargetSelection->Selection.Contains(ID) == false)
+			{
+				TargetSelection->Selection.Add(ID);
+				DeltaOut.Added.Add(ID);
+				NumAdded++;
+			}
+		}
+		return (NumAdded > 0);
+	}
+
+	/** Remove the items in the List from the Target Selection and return change information in DeltaOut */
+	template<typename ListType>
+	bool Deselect(const ListType& List, FGeometrySelectionDelta& DeltaOut)
+	{
+		int32 TotalRemoved = 0;
+		for (uint64 ID : List)
+		{
+			int32 NumRemoved = TargetSelection->Selection.Remove(ID);
+			if (NumRemoved > 0)
+			{
+				DeltaOut.Removed.Add(ID);
+			}
+			TotalRemoved += NumRemoved;
+		}
+		return (TotalRemoved > 0);
+	}
+
+	/** Replace the current set of selected items with those in the NewSelection, and return change information in DeltaOut */
+	bool Replace(const FGeometrySelection& NewSelection, FGeometrySelectionDelta& DeltaOut);
+
+
+protected:
+	FGeometrySelection* TargetSelection = nullptr;
+};
+
+
+} // end namespace UE::Geometry
+} // end namespace UE
diff --git a/Engine/Plugins/Runtime/GeometryProcessing/Source/DynamicMesh/Public/Selections/GeometrySelectionUtil.h b/Engine/Plugins/Runtime/GeometryProcessing/Source/DynamicMesh/Public/Selections/GeometrySelectionUtil.h
new file mode 100644
--- /dev/null
+++ b/Engine/Plugins/Runtime/GeometryProcessing/Source/DynamicMesh/Public/Selections/GeometrySelectionUtil.h
@@ -0,0 +1,180 @@
+// Copyright Epic Games, Inc. All Rights Reserved.
+
+#pragma once
+
+#include "Selections/GeometrySelection.h"
+#include "TriangleTypes.h"
+#include "SegmentTypes.h"
+#include "FrameTypes.h"
+#include "Polygroups/PolygroupSet.h"
+
+PREDECLARE_GEOMETRY(class FDynamicMesh3);
+
+namespace UE
+{
+namespace Geometry
+{
+
+class FColliderMesh;
+class FGroupTopology;
+struct FGroupTopologySelection;
+
+/**
+ * Update a FGeometrySelection (via the Editor) of mesh vertices/edges/triangles
+ * based on a raycast
+ * @param UpdateConfig type of update (add/remove/replace) and other update controls
+ * @param ResultOut selection delta and information on hit/miss are returned here
+ */
+DYNAMICMESH_API void UpdateTriangleSelectionViaRaycast(
+	const FColliderMesh* ColliderMesh,
+	FGeometrySelectionEditor* Editor,
+	const FRay3d& LocalRay,
+	const FGeometrySelectionUpdateConfig& UpdateConfig,
+	FGeometrySelectionUpdateResult& ResultOut);
+
+/**
+ * Update a FGeometrySelection (via the Editor) of mesh polygroup faces/edges/corners,
+ * based on a raycast
+ * @param UpdateConfig type of update (add/remove/replace) and other update controls
+ * @param ResultOut selection delta and information on hit/miss are returned here
+ */
+DYNAMICMESH_API void UpdateGroupSelectionViaRaycast(
+	const FColliderMesh* ColliderMesh,
+	const FGroupTopology* GroupTopology,
+	FGeometrySelectionEditor* Editor,
+	const FRay3d& LocalRay,
+	const FGeometrySelectionUpdateConfig& UpdateConfig,
+	FGeometrySelectionUpdateResult& ResultOut);
+
+
+/**
+ * Update a FGeometrySelection (via the Editor)
+ * @param ChangeType type of change to make (add/remove/replace)
+ * @param NewIDs set of new IDs to use to update the selection
+ * @param Delta selection delta will be stored here, if non-nullptr is provided
+ */
+DYNAMICMESH_API bool UpdateSelectionWithNewElements(
+	FGeometrySelectionEditor* Editor,
+	EGeometrySelectionChangeType ChangeType,
+	const TArray<uint64>& NewIDs,
+	FGeometrySelectionDelta* DeltaOut = nullptr);
+
+
+/**
+ * Call VertexFunc for each selected Mesh element (vertex/edge/tri) in MeshSelection.
+ * ApplyTransform will be applied to Vertex Positions before calling VertexFunc
+ */
+DYNAMICMESH_API bool EnumerateTriangleSelectionVertices(
+	const FGeometrySelection& MeshSelection,
+	const UE::Geometry::FDynamicMesh3& Mesh,
+	const FTransform& ApplyTransform,
+	TFunctionRef<void(uint64, const FVector3d&)> VertexFunc
+);
+/**
+ * Call VertexFunc for each selected Mesh element (vertex/edge/tri) in the set
+ * of polygroup faces/edges/corners specified by GroupSelection (relative to GroupTopology parameter)
+ * ApplyTransform will be applied to Vertex Positions before calling VertexFunc
+ */
+DYNAMICMESH_API bool EnumeratePolygroupSelectionVertices(
+	const FGeometrySelection& GroupSelection,
+	const UE::Geometry::FDynamicMesh3& Mesh,
+	const FGroupTopology* GroupTopology,
+	const FTransform& ApplyTransform,
+	TFunctionRef<void(uint64, const FVector3d&)> VertexFunc
+);
+
+
+/**
+ * Call TriangleFunc for each mesh TriangleID included in MeshSelection.
+ * TriangleFunc may be called multiple times for the same TriangleID.
+ * This will forward to EnumerateTriangleSelectionTriangles() or 
+ * EnumeratePolygroupSelectionTriangles() depending on the selection topology type.
+ * If UseGroupSet and MeshSelection is for polygroups, the default Mesh group layer will be used.
+ */
+DYNAMICMESH_API bool EnumerateSelectionTriangles(
+	const FGeometrySelection& MeshSelection,
+	const UE::Geometry::FDynamicMesh3& Mesh,
+	TFunctionRef<void(int32)> TriangleFunc,
+	const UE::Geometry::FPolygroupSet* UseGroupSet = nullptr
+);
+/**
+ * Call TriangleFunc for each mesh TriangleID included in MeshSelection.
+ * For Edges, both connected edges are included.
+ * For Vertices, all triangles in the vertex one-ring are included.
+ */
+DYNAMICMESH_API bool EnumerateTriangleSelectionTriangles(
+	const FGeometrySelection& MeshSelection,
+	const UE::Geometry::FDynamicMesh3& Mesh,
+	TFunctionRef<void(int32)> TriangleFunc
+);
+/**
+ * Call TriangleFunc for each mesh TriangleID included in MeshSelection, where MeshSelection has polygroup topology.
+ * For Polygroup Faces, all triangles in the face are included.
+ * For Polygroup Edges, currently all triangles in any group adjacent to the edge
+ * For Polygroup Corners, currently all triangles in any group touching the corner
+ */
+DYNAMICMESH_API bool EnumeratePolygroupSelectionTriangles(
+	const FGeometrySelection& MeshSelection,
+	const UE::Geometry::FDynamicMesh3& Mesh,
+	const UE::Geometry::FPolygroupSet& GroupSet,
+	TFunctionRef<void(int32)> TriangleFunc
+);
+
+
+/**
+ * Call VertexFunc/EdgeFunc/TriangleFunc for the vertices/edges/triangles identified by MeshSelection.
+ * Since a MeshSelection only stores vertices, edges, or triangles, but not combined, only one
+ * of these functions will be invoked during a call to this function.
+ * This function is useful to collect up geometry that needs to be rendered for a given MeshSelection
+ * @param ApplyTransform if non-null, transform is applied to the 3D geometry
+ */
+DYNAMICMESH_API bool EnumerateTriangleSelectionElements(
+	const FGeometrySelection& MeshSelection,
+	const UE::Geometry::FDynamicMesh3& Mesh,
+	TFunctionRef<void(int32, const FVector3d&)> VertexFunc,
+	TFunctionRef<void(int32, const FSegment3d&)> EdgeFunc,
+	TFunctionRef<void(int32, const FTriangle3d&)> TriangleFunc,
+	const FTransform* ApplyTransform = nullptr
+);
+/**
+ * Call VertexFunc/EdgeFunc/TriangleFunc for the mesh vertices/edges/triangles identified by MeshSelection,
+ * where MeshSelection has polygroup topology referring to the provided GroupTopology
+ * Since a MeshSelection only stores vertices, edges, or triangles, but not combined, only one
+ * of these functions will be invoked during a call to this function.
+ * This function is useful to collect up geometry that needs to be rendered for a given MeshSelection
+ * @param ApplyTransform if non-null, transform is applied to the 3D geometry
+ */
+DYNAMICMESH_API bool EnumeratePolygroupSelectionElements(
+	const FGeometrySelection& MeshSelection,
+	const UE::Geometry::FDynamicMesh3& Mesh,
+	const FGroupTopology* GroupTopology,
+	TFunctionRef<void(int32, const FVector3d&)> VertexFunc,
+	TFunctionRef<void(int32, const FSegment3d&)> EdgeFunc,
+	TFunctionRef<void(int32, const FTriangle3d&)> TriangleFunc,
+	const FTransform* ApplyTransform = nullptr
+);
+
+
+/**
+ * Convert a MeshSelection with Polygroup topology type to a FGroupTopologySelection
+ */
+DYNAMICMESH_API bool ConvertPolygroupSelectionToTopologySelection(
+	const FGeometrySelection& MeshSelection,
+	const UE::Geometry::FDynamicMesh3& Mesh,
+	const FGroupTopology* GroupTopology,
+	FGroupTopologySelection& TopologySelectionOut
+);
+
+
+/**
+ * Compute a 3D Frame suitable for use as a 3D transform gizmo position/orientation
+ * for the given MeshSelection
+ */
+DYNAMICMESH_API bool GetTriangleSelectionFrame(
+	const FGeometrySelection& MeshSelection,
+	const UE::Geometry::FDynamicMesh3& Mesh,
+	FFrame3d& SelectionFrameOut);
+
+
+}
+}
\ No newline at end of file
diff --git a/Engine/Plugins/Runtime/MeshModelingToolset/Source/MeshModelingTools/Private/Commands/DeleteGeometrySelectionCommand.cpp b/Engine/Plugins/Runtime/MeshModelingToolset/Source/MeshModelingTools/Private/Commands/DeleteGeometrySelectionCommand.cpp
--- a/Engine/Plugins/Runtime/MeshModelingToolset/Source/MeshModelingTools/Private/Commands/DeleteGeometrySelectionCommand.cpp
+++ b/Engine/Plugins/Runtime/MeshModelingToolset/Source/MeshModelingTools/Private/Commands/DeleteGeometrySelectionCommand.cpp
@@ -7,7 +7,7 @@
 #include "DynamicMesh/DynamicMeshChangeTracker.h"
 #include "DynamicMeshEditor.h"
 #include "Changes/MeshChange.h"
-#include "Selection/GeometrySelectionUtil.h"
+#include "Selections/GeometrySelectionUtil.h"
 
 using namespace UE::Geometry;
 
diff --git a/Engine/Plugins/Runtime/MeshModelingToolset/Source/ModelingComponents/Private/Selection/DynamicMeshSelector.cpp b/Engine/Plugins/Runtime/MeshModelingToolset/Source/ModelingComponents/Private/Selection/DynamicMeshSelector.cpp
--- a/Engine/Plugins/Runtime/MeshModelingToolset/Source/ModelingComponents/Private/Selection/DynamicMeshSelector.cpp
+++ b/Engine/Plugins/Runtime/MeshModelingToolset/Source/ModelingComponents/Private/Selection/DynamicMeshSelector.cpp
@@ -2,7 +2,7 @@
 
 
 #include "Selection/DynamicMeshSelector.h"
-#include "Selection/GeometrySelectionUtil.h"
+#include "Selections/GeometrySelectionUtil.h"
 #include "ToolContextInterfaces.h"
 #include "DynamicMesh/DynamicMeshAttributeSet.h"
 #include "DynamicMesh/MeshIndexUtil.h"
@@ -16,6 +16,16 @@ using namespace UE::Geometry;
 #define LOCTEXT_NAMESPACE "FDynamicMeshSelector"
 
 
+FDynamicMeshSelector::~FDynamicMeshSelector()
+{
+	if (TargetMesh.IsValid())
+	{
+		ensureMsgf(false, TEXT("FDynamicMeshSelector was not properly Shutdown!"));
+		FDynamicMeshSelector::Shutdown();
+	}
+}
+
+
 void FDynamicMeshSelector::Initialize(
 	FGeometryIdentifier SourceGeometryIdentifierIn,
 	UDynamicMesh* TargetMeshIn, 
diff --git a/Engine/Plugins/Runtime/MeshModelingToolset/Source/ModelingComponents/Private/Selection/GeometrySelection.cpp b/Engine/Plugins/Runtime/MeshModelingToolset/Source/ModelingComponents/Private/Selection/GeometrySelection.cpp
deleted file mode 100644
diff --git a/Engine/Plugins/Runtime/MeshModelingToolset/Source/ModelingComponents/Private/Selection/GeometrySelectionChanges.cpp b/Engine/Plugins/Runtime/MeshModelingToolset/Source/ModelingComponents/Private/Selection/GeometrySelectionChanges.cpp
--- a/Engine/Plugins/Runtime/MeshModelingToolset/Source/ModelingComponents/Private/Selection/GeometrySelectionChanges.cpp
+++ b/Engine/Plugins/Runtime/MeshModelingToolset/Source/ModelingComponents/Private/Selection/GeometrySelectionChanges.cpp
@@ -3,10 +3,11 @@
 
 #include "Selection/GeometrySelectionChanges.h"
 #include "Selection/GeometrySelectionManager.h"
-#include "Selection/GeometrySelectionUtil.h"
+#include "Selections/GeometrySelectionUtil.h"
 
 #define LOCTEXT_NAMESPACE "GeometrySelectionChanges"
 
+using namespace UE::Geometry;
 
 void FGeometrySelectionDeltaChange::Apply(UObject* Object)
 {
diff --git a/Engine/Plugins/Runtime/MeshModelingToolset/Source/ModelingComponents/Private/Selection/GeometrySelectionUtil.cpp b/Engine/Plugins/Runtime/MeshModelingToolset/Source/ModelingComponents/Private/Selection/GeometrySelectionUtil.cpp
deleted file mode 100644
diff --git a/Engine/Plugins/Runtime/MeshModelingToolset/Source/ModelingComponents/Public/Selection/DynamicMeshSelector.h b/Engine/Plugins/Runtime/MeshModelingToolset/Source/ModelingComponents/Public/Selection/DynamicMeshSelector.h
--- a/Engine/Plugins/Runtime/MeshModelingToolset/Source/ModelingComponents/Public/Selection/DynamicMeshSelector.h
+++ b/Engine/Plugins/Runtime/MeshModelingToolset/Source/ModelingComponents/Public/Selection/DynamicMeshSelector.h
@@ -5,8 +5,8 @@
 #include "CoreMinimal.h"
 #include "Misc/Change.h"
 #include "Templates/PimplPtr.h"
+#include "Selections/GeometrySelection.h"
 #include "Selection/GeometrySelector.h"
-#include "Selection/GeometrySelection.h"
 #include "Components/DynamicMeshComponent.h"
 #include "DynamicMesh/DynamicMeshAABBTree3.h"
 #include "TransformTypes.h"
@@ -26,10 +26,7 @@ PREDECLARE_GEOMETRY(class FColliderMesh);
 class MODELINGCOMPONENTS_API FDynamicMeshSelector : public IGeometrySelector
 {
 public:
-	virtual ~FDynamicMeshSelector()
-	{
-		FDynamicMeshSelector::Shutdown();
-	}
+	virtual ~FDynamicMeshSelector();
 
 	/**
 	 * Initialize the FDynamicMeshSelector for a given source/target UDynamicMesh.
diff --git a/Engine/Plugins/Runtime/MeshModelingToolset/Source/ModelingComponents/Public/Selection/GeometrySelection.h b/Engine/Plugins/Runtime/MeshModelingToolset/Source/ModelingComponents/Public/Selection/GeometrySelection.h
deleted file mode 100644
diff --git a/Engine/Plugins/Runtime/MeshModelingToolset/Source/ModelingComponents/Public/Selection/GeometrySelectionChanges.h b/Engine/Plugins/Runtime/MeshModelingToolset/Source/ModelingComponents/Public/Selection/GeometrySelectionChanges.h
--- a/Engine/Plugins/Runtime/MeshModelingToolset/Source/ModelingComponents/Public/Selection/GeometrySelectionChanges.h
+++ b/Engine/Plugins/Runtime/MeshModelingToolset/Source/ModelingComponents/Public/Selection/GeometrySelectionChanges.h
@@ -3,7 +3,7 @@
 #pragma once
 
 #include "InteractiveToolChange.h"
-#include "Selection/GeometrySelection.h"
+#include "Selections/GeometrySelection.h"
 #include "Selection/GeometrySelector.h"		// for FGeometryIdentifier
 
 
@@ -18,6 +18,10 @@
 class MODELINGCOMPONENTS_API IGeometrySelectionChange
 {
 public:
+	using FGeometrySelection = UE::Geometry::FGeometrySelection;
+	using FGeometrySelectionEditor = UE::Geometry::FGeometrySelectionEditor;
+	using FGeometrySelectionDelta = UE::Geometry::FGeometrySelectionDelta;
+
 	virtual ~IGeometrySelectionChange() {}
 
 	virtual FGeometryIdentifier GetIdentifier() const = 0;
diff --git a/Engine/Plugins/Runtime/MeshModelingToolset/Source/ModelingComponents/Public/Selection/GeometrySelectionManager.h b/Engine/Plugins/Runtime/MeshModelingToolset/Source/ModelingComponents/Public/Selection/GeometrySelectionManager.h
--- a/Engine/Plugins/Runtime/MeshModelingToolset/Source/ModelingComponents/Public/Selection/GeometrySelectionManager.h
+++ b/Engine/Plugins/Runtime/MeshModelingToolset/Source/ModelingComponents/Public/Selection/GeometrySelectionManager.h
@@ -3,8 +3,8 @@
 #pragma once
 
 #include "CoreMinimal.h"
+#include "Selections/GeometrySelection.h"
 #include "Selection/GeometrySelector.h"
-#include "Selection/GeometrySelection.h"
 #include "Selection/GeometrySelectionChanges.h"
 
 #include "GeometrySelectionManager.generated.h"
@@ -56,6 +56,15 @@ class MODELINGCOMPONENTS_API UGeometrySelectionManager : public UObject
 
 public:
 
+	using FGeometrySelection = UE::Geometry::FGeometrySelection;
+	using FGeometrySelectionEditor = UE::Geometry::FGeometrySelectionEditor;
+	using FGeometrySelectionBounds = UE::Geometry::FGeometrySelectionBounds;
+	using FGeometrySelectionElements = UE::Geometry::FGeometrySelectionElements;
+	using FGeometrySelectionUpdateConfig = UE::Geometry::FGeometrySelectionUpdateConfig;
+	using FGeometrySelectionUpdateResult = UE::Geometry::FGeometrySelectionUpdateResult;
+	using EGeometryElementType = UE::Geometry::EGeometryElementType;
+	using EGeometryTopologyType = UE::Geometry::EGeometryTopologyType;
+
 	//
 	// Setup/Teardown
 	//
diff --git a/Engine/Plugins/Runtime/MeshModelingToolset/Source/ModelingComponents/Public/Selection/GeometrySelectionUtil.h b/Engine/Plugins/Runtime/MeshModelingToolset/Source/ModelingComponents/Public/Selection/GeometrySelectionUtil.h
deleted file mode 100644
diff --git a/Engine/Plugins/Runtime/MeshModelingToolset/Source/ModelingComponents/Public/Selection/GeometrySelector.h b/Engine/Plugins/Runtime/MeshModelingToolset/Source/ModelingComponents/Public/Selection/GeometrySelector.h
--- a/Engine/Plugins/Runtime/MeshModelingToolset/Source/ModelingComponents/Public/Selection/GeometrySelector.h
+++ b/Engine/Plugins/Runtime/MeshModelingToolset/Source/ModelingComponents/Public/Selection/GeometrySelector.h
@@ -3,7 +3,7 @@
 #pragma once
 
 #include "Components/PrimitiveComponent.h"
-#include "Selection/GeometrySelection.h"
+#include "Selections/GeometrySelection.h"
 #include "FrameTypes.h"
 #include "InputState.h"
 
@@ -130,7 +130,7 @@ struct MODELINGCOMPONENTS_API FGeometryIdentifier
 struct MODELINGCOMPONENTS_API FGeometrySelectionHandle
 {
 	FGeometryIdentifier Identifier;
-	const FGeometrySelection* Selection;
+	const UE::Geometry::FGeometrySelection* Selection;
 };
 
 
@@ -167,7 +167,7 @@ public:
 	 * Start a transform (eg called at beginning of a user interaction with a 3D gizmo).
 	 * The caller provides the active Selection that is intended to be transformed.
 	 */
-	virtual void BeginTransform(const FGeometrySelection& Selection) = 0;
+	virtual void BeginTransform(const UE::Geometry::FGeometrySelection& Selection) = 0;
 
 	/**
 	 * Update the active transform (eg called on each mouse update during 3D gizmo usage).
@@ -215,6 +215,13 @@ public:
 class MODELINGCOMPONENTS_API IGeometrySelector
 {
 public:
+	using FGeometrySelection = UE::Geometry::FGeometrySelection;
+	using FGeometrySelectionEditor = UE::Geometry::FGeometrySelectionEditor;
+	using FGeometrySelectionBounds = UE::Geometry::FGeometrySelectionBounds;
+	using FGeometrySelectionElements = UE::Geometry::FGeometrySelectionElements;
+	using FGeometrySelectionUpdateConfig = UE::Geometry::FGeometrySelectionUpdateConfig;
+	using FGeometrySelectionUpdateResult = UE::Geometry::FGeometrySelectionUpdateResult;
+
 	virtual ~IGeometrySelector() {}
 
 	/**
diff --git a/Engine/Source/Runtime/GeometryFramework/Private/Components/DynamicMeshComponent.cpp b/Engine/Source/Runtime/GeometryFramework/Private/Components/DynamicMeshComponent.cpp
--- a/Engine/Source/Runtime/GeometryFramework/Private/Components/DynamicMeshComponent.cpp
+++ b/Engine/Source/Runtime/GeometryFramework/Private/Components/DynamicMeshComponent.cpp
@@ -1402,10 +1402,7 @@ void UDynamicMeshComponent::SetDeferredCollisionUpdatesEnabled(bool bEnabled, bo
 
 void UDynamicMeshComponent::SetTransientDeferCollisionUpdates(bool bEnabled)
 {
-	if (bTransientDeferCollisionUpdates != bEnabled)
-	{
-		bTransientDeferCollisionUpdates = bEnabled;
-	}
+	bTransientDeferCollisionUpdates = bEnabled;
 }
 
 void UDynamicMeshComponent::SetSceneProxyVerifyUsedMaterials(bool bState)
