remove export on enum types

diff --git a/Engine/Plugins/Runtime/GeometryProcessing/Source/DynamicMesh/Public/Selections/GeometrySelection.h b/Engine/Plugins/Runtime/GeometryProcessing/Source/DynamicMesh/Public/Selections/GeometrySelection.h
--- a/Engine/Plugins/Runtime/GeometryProcessing/Source/DynamicMesh/Public/Selections/GeometrySelection.h
+++ b/Engine/Plugins/Runtime/GeometryProcessing/Source/DynamicMesh/Public/Selections/GeometrySelection.h
@@ -109,7 +109,7 @@ struct DYNAMICMESH_API FGeoSelectionID
 /**
  * Type of selected Elements in a FGeometrySelection
  */
-enum class DYNAMICMESH_API EGeometryElementType
+enum class EGeometryElementType
 {
 	/** Mesh Vertices, Polygroup Corners, ... */
 	Vertex = 1,
@@ -122,7 +122,7 @@ enum class DYNAMICMESH_API EGeometryElementType
 /**
  * Type of selected Topology in a FGeometrySelection.
  */
-enum class DYNAMICMESH_API EGeometryTopologyType
+enum class EGeometryTopologyType
 {
 	Triangle = 1,
 	Polygroup = 2
@@ -186,7 +186,7 @@ struct DYNAMICMESH_API FGeometrySelection
 /**
  * Type of change, relative to a FGeometrySelection
  */
-enum class DYNAMICMESH_API EGeometrySelectionChangeType
+enum class EGeometrySelectionChangeType
 {
 	/** Elements Added to Selection */
 	Add,
