ModelingTools: add 'Find Quads' mode to GrpGen Tool. Implementation is in FPolygroupsGenerator::FindSourceMeshPolygonPolygroups. Add uv-seam-constraint and hard-normal-constraint options for both this new mode, and the existing normal-deviation mode.

diff --git a/Engine/Plugins/Experimental/MeshModelingToolsetExp/Source/MeshModelingToolsExp/Private/ConvertToPolygonsTool.cpp b/Engine/Plugins/Experimental/MeshModelingToolsetExp/Source/MeshModelingToolsExp/Private/ConvertToPolygonsTool.cpp
--- a/Engine/Plugins/Experimental/MeshModelingToolsetExp/Source/MeshModelingToolsExp/Private/ConvertToPolygonsTool.cpp
+++ b/Engine/Plugins/Experimental/MeshModelingToolsetExp/Source/MeshModelingToolsExp/Private/ConvertToPolygonsTool.cpp
@@ -44,7 +44,16 @@ public:
 	bool bSubdivideExisting = false;
 	FPolygroupsGenerator::EWeightingType WeightingType = FPolygroupsGenerator::EWeightingType::None;
 	FVector3d WeightingCoeffs = FVector3d::One();
+
+	bool bRespectUVSeams = false;
+	bool bRespectHardNormals = false;
+
+	double QuadMetricClamp = 1.0;
+	double QuadAdjacencyWeight = 1.0;
+	int QuadSearchRounds = 1;
+
 	int32 MinGroupSize = 2;
+
 	bool bCalculateNormals = false;
 
 	// input mesh
@@ -90,7 +99,21 @@ public:
 		case EConvertToPolygonsMode::FaceNormalDeviation:
 		{
 			double DotTolerance = 1.0 - FMathd::Cos(AngleTolerance * FMathd::DegToRad);
-			Generator.FindPolygroupsFromFaceNormals(DotTolerance);
+			Generator.FindPolygroupsFromFaceNormals(
+				DotTolerance,
+				bRespectUVSeams, 
+				bRespectHardNormals );
+			break;
+		}
+		case EConvertToPolygonsMode::FindPolygons:
+		{
+			Generator.FindSourceMeshPolygonPolygroups(
+				bRespectUVSeams, 
+				bRespectHardNormals, 
+				QuadMetricClamp, 
+				QuadAdjacencyWeight,
+				FMath::Clamp(QuadSearchRounds, 1, 100)
+			);
 			break;
 		}
 		case EConvertToPolygonsMode::FromFurthestPointSampling:
@@ -234,7 +257,12 @@ void UConvertToPolygonsTool::Setup()
 	Settings->WatchProperty(Settings->bNormalWeighted, [this](bool) { OnSettingsModified(); });
 	Settings->WatchProperty(Settings->NormalWeighting, [this](float) { OnSettingsModified(); });
 	Settings->WatchProperty(Settings->MinGroupSize, [this](int32) { OnSettingsModified(); });
-	
+	Settings->WatchProperty(Settings->QuadAdjacencyWeight, [this](float) { OnSettingsModified(); });
+	Settings->WatchProperty(Settings->QuadMetricClamp, [this](float) { OnSettingsModified(); });
+	Settings->WatchProperty(Settings->QuadSearchRounds, [this](int) { OnSettingsModified(); });
+	Settings->WatchProperty(Settings->bRespectUVSeams, [this](int) { OnSettingsModified(); });
+	Settings->WatchProperty(Settings->bRespectHardNormals, [this](int) { OnSettingsModified(); });
+
 
 	GetToolManager()->DisplayMessage(
 		LOCTEXT("OnStartTool", "Cluster triangles of the Mesh into PolyGroups using various strategies"),
@@ -251,6 +279,12 @@ void UConvertToPolygonsTool::UpdateOpParameters(FConvertToPolygonsOp& ConvertToP
 	ConvertToPolygonsOp.WeightingType = (Settings->bNormalWeighted) ? FPolygroupsGenerator::EWeightingType::NormalDeviation : FPolygroupsGenerator::EWeightingType::None;
 	ConvertToPolygonsOp.WeightingCoeffs = FVector3d(Settings->NormalWeighting, 1.0, 1.0);
 	ConvertToPolygonsOp.MinGroupSize = Settings->MinGroupSize;
+	ConvertToPolygonsOp.QuadMetricClamp = Settings->QuadMetricClamp;
+	ConvertToPolygonsOp.QuadAdjacencyWeight = Settings->QuadAdjacencyWeight;
+	ConvertToPolygonsOp.QuadSearchRounds = Settings->QuadSearchRounds;
+	ConvertToPolygonsOp.bRespectUVSeams = Settings->bRespectUVSeams;
+	ConvertToPolygonsOp.bRespectHardNormals = Settings->bRespectHardNormals;
+
 	ConvertToPolygonsOp.OriginalMesh = OriginalDynamicMesh;
 	
 	FTransform LocalToWorld = (FTransform)UE::ToolTarget::GetLocalToWorldTransform(Target);
diff --git a/Engine/Plugins/Experimental/MeshModelingToolsetExp/Source/MeshModelingToolsExp/Public/ConvertToPolygonsTool.h b/Engine/Plugins/Experimental/MeshModelingToolsetExp/Source/MeshModelingToolsExp/Public/ConvertToPolygonsTool.h
--- a/Engine/Plugins/Experimental/MeshModelingToolsetExp/Source/MeshModelingToolsExp/Public/ConvertToPolygonsTool.h
+++ b/Engine/Plugins/Experimental/MeshModelingToolsetExp/Source/MeshModelingToolsExp/Public/ConvertToPolygonsTool.h
@@ -34,6 +34,8 @@ enum class EConvertToPolygonsMode
 {
 	/** Convert based on Angle Tolerance between Face Normals */
 	FaceNormalDeviation UMETA(DisplayName = "Face Normal Deviation"),
+	/** Create Polygroups by merging triangle pairs into Quads */
+	FindPolygons UMETA(DisplayName = "Find Quads"),
 	/** Create PolyGroups based on UV Islands */
 	FromUVIslands  UMETA(DisplayName = "From UV Islands"),
 	/** Create PolyGroups based on Hard Normal Seams */
@@ -57,26 +59,47 @@ public:
 	EConvertToPolygonsMode ConversionMode = EConvertToPolygonsMode::FaceNormalDeviation;
 
 	/** Tolerance for planarity */
-	UPROPERTY(EditAnywhere, Category = PolyGroups, meta = (UIMin = "0.001", UIMax = "60.0", ClampMin = "0.0", ClampMax = "90.0", EditCondition = "ConversionMode == EConvertToPolygonsMode::FaceNormalDeviation", EditConditionHides))
+	UPROPERTY(EditAnywhere, Category = NormalDeviation, meta = (UIMin = "0.001", UIMax = "60.0", ClampMin = "0.0", ClampMax = "90.0", EditCondition = "ConversionMode == EConvertToPolygonsMode::FaceNormalDeviation", EditConditionHides))
 	float AngleTolerance = 0.1f;
 
 	/** Furthest-Point Sample count, approximately this number of polygroups will be generated */
-	UPROPERTY(EditAnywhere, Category = PolyGroups, meta = (UIMin = "1", UIMax = "100", ClampMin = "1", ClampMax = "10000", EditCondition = "ConversionMode == EConvertToPolygonsMode::FromFurthestPointSampling", EditConditionHides))
+	UPROPERTY(EditAnywhere, Category = FurthestPoint, meta = (UIMin = "1", UIMax = "100", ClampMin = "1", ClampMax = "10000", EditCondition = "ConversionMode == EConvertToPolygonsMode::FromFurthestPointSampling", EditConditionHides))
 	int32 NumPoints = 100;
 
 	/** If enabled, then furthest-point sampling happens with respect to existing Polygroups, ie the existing groups are further subdivided */
-	UPROPERTY(EditAnywhere, Category = PolyGroups, meta = (EditCondition = "ConversionMode == EConvertToPolygonsMode::FromFurthestPointSampling", EditConditionHides))
+	UPROPERTY(EditAnywhere, Category = FurthestPoint, meta = (EditCondition = "ConversionMode == EConvertToPolygonsMode::FromFurthestPointSampling", EditConditionHides))
 	bool bSplitExisting = false;
 
 	/** If true, region-growing in Sampling modes will be controlled by face normals, resulting in regions with borders that are more-aligned with curvature ridges */
-	UPROPERTY(EditAnywhere, Category = PolyGroups, meta = (EditCondition = "ConversionMode == EConvertToPolygonsMode::FromFurthestPointSampling", EditConditionHides))
+	UPROPERTY(EditAnywhere, Category = FurthestPoint, meta = (EditCondition = "ConversionMode == EConvertToPolygonsMode::FromFurthestPointSampling", EditConditionHides))
 	bool bNormalWeighted = true;
 
 	/** This parameter modulates the effect of normal weighting during region-growing */
-	UPROPERTY(EditAnywhere, Category = PolyGroups, meta = (UIMin = "0.1", UIMax = "2.0", ClampMin = "0.01", ClampMax = "100.0", EditCondition = "ConversionMode == EConvertToPolygonsMode::FromFurthestPointSampling", EditConditionHides))
+	UPROPERTY(EditAnywhere, Category = FurthestPoint, meta = (UIMin = "0.1", UIMax = "2.0", ClampMin = "0.01", ClampMax = "100.0", EditCondition = "ConversionMode == EConvertToPolygonsMode::FromFurthestPointSampling", EditConditionHides))
 	float NormalWeighting = 1.0f;
 
 
+	/** Bias for Quads that are adjacent to already-discovered Quads. Set to 0 to disable.  */
+	UPROPERTY(EditAnywhere, Category = FindQuads, meta = (UIMin = 0, UIMax = 5, EditCondition = "ConversionMode == EConvertToPolygonsMode::FindPolygons", EditConditionHides))
+	float QuadAdjacencyWeight = 1.0;
+
+	/** Set to values below 1 to ignore less-likely triangle pairings */
+	UPROPERTY(EditAnywhere, Category = FindQuads, meta = (AdvancedDisplay, UIMin = 0, UIMax = 1, EditCondition = "ConversionMode == EConvertToPolygonsMode::FindPolygons", EditConditionHides))
+	float QuadMetricClamp = 1.0;
+
+	/** Iteratively repeat quad-searching in uncertain areas, to try to slightly improve results */
+	UPROPERTY(EditAnywhere, Category = FindQuads, meta = (AdvancedDisplay, UIMin = 1, UIMax = 5, EditCondition = "ConversionMode == EConvertToPolygonsMode::FindPolygons", EditConditionHides))
+	int QuadSearchRounds = 1;
+
+	/** If true, polygroup borders will not cross existing UV seams */
+	UPROPERTY(EditAnywhere, Category = Topology, meta = (EditCondition = "ConversionMode == EConvertToPolygonsMode::FaceNormalDeviation || ConversionMode == EConvertToPolygonsMode::FindPolygons", EditConditionHides))
+	bool bRespectUVSeams = false;
+
+	/** If true, polygroup borders will not cross existing hard normal seams */
+	UPROPERTY(EditAnywhere, Category = Topology, meta = (EditCondition = "ConversionMode == EConvertToPolygonsMode::FaceNormalDeviation || ConversionMode == EConvertToPolygonsMode::FindPolygons", EditConditionHides))
+	bool bRespectHardNormals = false;
+
+
 	/** group filtering */
 	UPROPERTY(EditAnywhere, Category = Filtering, meta = (UIMin = "1", UIMax = "100", ClampMin = "1", ClampMax = "10000"))
 	int32 MinGroupSize = 2;
diff --git a/Engine/Plugins/Runtime/GeometryProcessing/Source/DynamicMesh/Private/Polygroups/PolygroupsGenerator.cpp b/Engine/Plugins/Runtime/GeometryProcessing/Source/DynamicMesh/Private/Polygroups/PolygroupsGenerator.cpp
--- a/Engine/Plugins/Runtime/GeometryProcessing/Source/DynamicMesh/Private/Polygroups/PolygroupsGenerator.cpp
+++ b/Engine/Plugins/Runtime/GeometryProcessing/Source/DynamicMesh/Private/Polygroups/PolygroupsGenerator.cpp
@@ -8,6 +8,8 @@
 #include "Parameterization/IncrementalMeshDijkstra.h"
 #include "Parameterization/MeshRegionGraph.h"
 #include "Async/ParallelFor.h"
+#include "Util/IndexPriorityQueue.h"
+#include "Selections/MeshFaceSelection.h"
 
 #include "Curve/DynamicGraph3.h"
 
@@ -117,10 +119,482 @@ bool FPolygroupsGenerator::FindPolygroupsFromConnectedTris()
 
 
 
-bool FPolygroupsGenerator::FindPolygroupsFromFaceNormals(double DotTolerance)
+
+
+void FPolygroupsGenerator::GetSeamConstraintEdges(bool bUVSeams, bool bNormalSeams, TSet<int32>& ConstraintEdgesOut) const
+{
+	if (bUVSeams && Mesh->HasAttributes() && Mesh->Attributes()->PrimaryUV() != nullptr )
+	{
+		FDynamicMeshUVOverlay* UVOverlay = Mesh->Attributes()->PrimaryUV();
+		for (int32 eid : Mesh->EdgeIndicesItr())
+		{
+			if (UVOverlay->IsSeamEdge(eid))
+			{
+				ConstraintEdgesOut.Add(eid);
+			}
+		}
+	}
+	if (bNormalSeams && Mesh->HasAttributes() && Mesh->Attributes()->PrimaryNormals() != nullptr )
+	{
+		FDynamicMeshNormalOverlay* NormalOverlay = Mesh->Attributes()->PrimaryNormals();
+		for (int32 eid : Mesh->EdgeIndicesItr())
+		{
+			if (NormalOverlay->IsSeamEdge(eid))
+			{
+				ConstraintEdgesOut.Add(eid);
+			}
+		}
+	}
+}
+
+
+bool FPolygroupsGenerator::FindSourceMeshPolygonPolygroups(
+	bool bRespectUVSeams,
+	bool bRespectNormalSeams,
+	double QuadMetricClamp,
+	double QuadAdjacencyWeight,
+	int MaxSearchRounds)
+{
+	// precompute triangle normals
+	int32 MaxTriangleID = Mesh->MaxTriangleID();
+	TArray<FVector3d> FaceNormals;
+	FaceNormals.SetNum(MaxTriangleID);
+	TArray<double> FaceAreas;
+	FaceAreas.SetNum(MaxTriangleID);
+	for (int32 tid : Mesh->TriangleIndicesItr())
+	{
+		FVector3d Centroid;
+		Mesh->GetTriInfo(tid, FaceNormals[tid], FaceAreas[tid], Centroid);
+	}
+
+	// if we are respecting seams or hard normals, find all those edges
+	TSet<int32> InvalidEdges;
+	if (bRespectUVSeams || bRespectNormalSeams)
+	{
+		GetSeamConstraintEdges(bRespectUVSeams, bRespectNormalSeams, InvalidEdges);
+	}
+
+	// FTriPairQuad is a quad formed by a pair of triangles in the mesh
+	struct FTriPairQuad
+	{
+		FIndex2i Triangles = FIndex2i::Invalid();
+		int32 MidEdgeID = -1;	// edge through middle of quad
+
+		FVector2d TriAreas = FVector2d::Zero();
+		FIndex4i VertexLoop = FIndex4i::Invalid();
+		FIndex4i QuadEdges = FIndex4i::Invalid();
+		FIndex4i QuadEdgeOtherTris = FIndex4i::Invalid();
+
+		int32 Index = -1;		// index into sorted Squareness list
+
+		bool IsValid() const { return Triangles.A >= 0 && Triangles.B >= 0; }
+
+		bool IsAdjacent(const FTriPairQuad& OtherQuad) const
+		{
+			return OtherQuad.QuadEdgeOtherTris.Contains(Triangles.A) || OtherQuad.QuadEdgeOtherTris.Contains(Triangles.B);
+		}
+		bool IsAdjacent(int32 TriangleID) const
+		{
+			return QuadEdgeOtherTris.Contains(TriangleID);
+		}
+	};
+
+	// For each triangle, there are 3 possible adjacent triangles that form valid quads. 
+	// TriPotentialQuads is an [N][3] array of those potential quads, stored as FTriPairQuad
+	TArray<FTriPairQuad[3]> TriPotentialQuads;
+	TriPotentialQuads.SetNum(MaxTriangleID);
+	for (int32 tid : Mesh->TriangleIndicesItr())
+	{
+		FIndex3i TriEdges = Mesh->GetTriEdges(tid);
+		FIndex3i TriNbrTris = Mesh->GetTriNeighbourTris(tid);
+		for (int32 j = 0; j < 3; ++j)
+		{
+			FTriPairQuad NbrQuad;
+			int32 nbr_tid = TriNbrTris[j];
+			if (nbr_tid >= 0 && (InvalidEdges.Contains(TriEdges[j]) == false) )
+			{
+				NbrQuad.Triangles = FIndex2i(tid, nbr_tid);
+				NbrQuad.MidEdgeID = TriEdges[j];
+
+				FIndex2i EdgeV = Mesh->GetEdgeV(NbrQuad.MidEdgeID);
+				FIndex2i OtherEdgeV = Mesh->GetEdgeOpposingV(NbrQuad.MidEdgeID);
+				NbrQuad.VertexLoop = FIndex4i(EdgeV.A, OtherEdgeV.A, EdgeV.B, OtherEdgeV.B);
+
+				NbrQuad.TriAreas[0] = Mesh->GetTriArea(NbrQuad.Triangles[0]);
+				NbrQuad.TriAreas[1] = Mesh->GetTriArea(NbrQuad.Triangles[1]);
+
+				for (int32 k = 0; k < 4; ++k)
+				{
+					NbrQuad.QuadEdges[k] = Mesh->FindEdge(NbrQuad.VertexLoop[k], NbrQuad.VertexLoop[(k+1)%4]);
+
+					int32 eid = NbrQuad.QuadEdges[k];
+					FIndex2i EdgeT = Mesh->GetEdgeT(eid);
+					if (EdgeT.B != IndexConstants::InvalidID)
+					{
+						int IsA = (EdgeT.A == tid || EdgeT.A == nbr_tid) ? 1 : 0;
+						int IsB = (EdgeT.B == tid || EdgeT.B == nbr_tid) ? 1 : 0;
+						check( (IsA ^ IsB) == 1 );
+
+						int32 other_tid = (EdgeT.A == tid || EdgeT.A == nbr_tid) ? EdgeT.B : EdgeT.A;
+						NbrQuad.QuadEdgeOtherTris[k] = other_tid;
+					}
+				}
+			}
+			TriPotentialQuads[tid][j] = NbrQuad;
+		}
+	}
+
+	// QuadnessMetric measures the likelihood of an individual Quad, with 0 being 
+	// extremely quad (ie a planar square) and larger numbers being less-good. 
+	// The metric is ultimately sqrt( sum-of-squares ) where each term is in
+	// range [0,1]
+
+	// to measure:
+	//   - internal angles, closer to 90 == better
+	//   - length similarity of edges (0,2) and (1<3)
+	//   - option to prevent quads from crossing UV seams, hard normal seams
+
+	// TODO: does it make sense to provide ability to modulate the face-planarity weight?
+	// It is mostly necessary on hard-surface meshes, on organic shapes it's will potentially
+	// prevent finding non-planar quads...
+	// TODO: should we consider edge length metric, comparing mid-edge w/ opposite mid-edge length?
+	auto QuadnessMetric = [this, &FaceNormals](const FTriPairQuad& Quad)
+	{
+		double Metric = 0;
+
+		double FacesDot = Dot( FaceNormals[Quad.Triangles.A], FaceNormals[Quad.Triangles.B] );
+
+		// prefer quads that are planar, ie angle between normals is closer to 0
+		double FacesAngle = AngleD( FaceNormals[Quad.Triangles.A], FaceNormals[Quad.Triangles.B] );
+		double MappedUnitFacesAngle = FMathd::Abs( FMathd::Abs(FacesAngle) / 180.0 );
+		Metric += MappedUnitFacesAngle*MappedUnitFacesAngle;
+
+		// prefer quads where opening angles are ~90 degrees
+		double OpeningAnglesDeviationMetric = 0;
+		if (FacesDot > 0.001)
+		{
+			FVector3d AvgNormal = Normalized(FaceNormals[Quad.Triangles.A] + FaceNormals[Quad.Triangles.B]);
+	
+			for (int32 j = 0; j < 4; ++j)
+			{
+				FVector3d C = Mesh->GetVertex(Quad.VertexLoop[j]);
+				FVector3d Corner = Mesh->GetVertex(Quad.VertexLoop[(j+1)%4]);
+				FVector3d D = Mesh->GetVertex(Quad.VertexLoop[(j+2)%4]);
+				FFrame3d AvgPlane(Corner, AvgNormal);
+				FVector2d ProjC = Normalized( AvgPlane.ToPlaneUV( C ) );
+				FVector2d ProjD = Normalized( AvgPlane.ToPlaneUV( D ) );
+
+				double OpeningAngle = AngleD(ProjC, ProjD);		// [-180, 180]
+				double MappedUnitOpeningAngle = FMathd::Abs( (FMathd::Abs(OpeningAngle)-90) / 90.0 );
+				OpeningAnglesDeviationMetric += MappedUnitOpeningAngle;
+			}
+		}
+		else
+		{
+			OpeningAnglesDeviationMetric = 4;		// quad opening angle is >= 90 deg, cannot reliably measure 2D opening angle and 3D is unreliable...
+		}
+		OpeningAnglesDeviationMetric /= 4.0;
+		Metric += OpeningAnglesDeviationMetric * OpeningAnglesDeviationMetric;
+
+		// prefer quads where both tris are the same area
+		double MinArea = FMathd::Min(Quad.TriAreas.X, Quad.TriAreas.Y);
+		double AreaMetric = FMathd::Max(Quad.TriAreas.X, Quad.TriAreas.Y) / FMathd::Max(MinArea, 0.0000001);
+		// min area metric should be 1...
+		AreaMetric = FMathd::Clamp(AreaMetric, 1.0, 5.0);
+		AreaMetric = (AreaMetric - 1.0) / 4.0;
+		check(AreaMetric >= 0 && AreaMetric <= 1.0);
+		Metric += AreaMetric * AreaMetric;
+
+		return FMathd::Sqrt(Metric);
+	};
+	double CurrentMaxMetric = FMathd::Sqrt(3);		// each term is in range [0,1] so this is sqrt(# of terms)
+	auto IsViableQuad = [CurrentMaxMetric, QuadMetricClamp](double Metric)
+	{
+		return (Metric / CurrentMaxMetric) <= QuadMetricClamp;
+	};
+
+
+	// FQuadWithMetric 
+	struct FQuadWithMetric
+	{
+		FIndex2i QuadIndex;			// (tri,j) index into TriPotentialQuads
+		double QuadMetric;			// initial value of QuadnessMetric()
+		int32 FixedQuadNbrs = 0;	// number of known quad neighbours
+
+		double GetMetric(double AdjacencyWeight) const
+		{
+			if (FixedQuadNbrs == 0)
+			{
+				return QuadMetric;
+			}
+			else
+			{
+				return QuadMetric / FMathd::Pow((double)(FixedQuadNbrs + 1), AdjacencyWeight);
+			}
+		}
+	};
+	TArray<FQuadWithMetric> AllQuadsMetricList;
+
+	for (int32 tid : Mesh->TriangleIndicesItr())
+	{
+		for (int32 j = 0; j < 3; ++j)
+		{
+			if (TriPotentialQuads[tid][j].IsValid())
+			{
+				AllQuadsMetricList.Add(FQuadWithMetric{ FIndex2i(tid,j), QuadnessMetric(TriPotentialQuads[tid][j]) });
+			}
+		}
+	}
+	AllQuadsMetricList.Sort( [](const FQuadWithMetric& A, const FQuadWithMetric& B) { return A.QuadMetric < B.QuadMetric; } );
+
+	// returns ref to the FTriPairQuad for the given index into AllQuadsMetricList
+	auto GetQuadForIndex = [&AllQuadsMetricList, &TriPotentialQuads](int32 QMIndex) -> FTriPairQuad&
+	{
+		FQuadWithMetric& MetricQuad = AllQuadsMetricList[QMIndex];
+		FIndex2i TriNbrIdx = MetricQuad.QuadIndex;
+		return TriPotentialQuads[TriNbrIdx.A][TriNbrIdx.B];
+	};
+
+	bool bDone = false;
+	int Iters = 0;
+	int32 LastTris = Mesh->TriangleCount(), LastQuads = 0;
+	TArray<FTriPairQuad> FoundQuads;
+	TArray<FTriPairQuad> LastRoundQuads;
+	TArray<bool> DoneTris;
+	while (!bDone && Iters++ < MaxSearchRounds)
+	{
+		DoneTris.Init(false, MaxTriangleID);
+		// returns true if either triangle of the quad has been used up in another quad
+		auto IsDoneQuad = [&DoneTris](const FTriPairQuad& Quad)
+		{
+			return DoneTris[Quad.Triangles.A] == true || DoneTris[Quad.Triangles.B] == true;
+		};
+
+		// reset index/nbrcount
+		for (int32 k = 0; k < AllQuadsMetricList.Num(); ++k)
+		{
+			GetQuadForIndex(k).Index = k;
+			AllQuadsMetricList[k].FixedQuadNbrs = 0;
+		}
+
+		int32 RemainingTriangles = Mesh->TriangleCount();
+
+		// list of final quads extracted by the search below
+		FoundQuads.Reset();
+
+		// If we are in a round > 1, we are generally fine-tuning the existing solution.
+		// Currently this is done by just deleting some quad-areas and letting them re-populate
+		// with the new neighbour set. So re-add all the "known" quads from the previous round.
+		if (LastRoundQuads.Num() > 0)
+		{
+			for (FTriPairQuad PrevQuad : LastRoundQuads)
+			{
+				FoundQuads.Add(PrevQuad);
+				DoneTris[PrevQuad.Triangles.A] = true;
+				DoneTris[PrevQuad.Triangles.B] = true;
+				RemainingTriangles -= 2;
+			}
+			LastRoundQuads.Reset();
+
+			// initialize neighbour info (technically only need this for unfinished quads?)
+			for (int32 Index = 0; Index < AllQuadsMetricList.Num(); ++Index)
+			{
+				FTriPairQuad& Quad = GetQuadForIndex(Index);
+				for (int32 j = 0; j < 4; ++j)
+				{
+					if (Quad.QuadEdgeOtherTris[j] >= 0 && DoneTris[Quad.QuadEdgeOtherTris[j]])
+					{
+						AllQuadsMetricList[Index].FixedQuadNbrs++;
+						check(AllQuadsMetricList[Index].FixedQuadNbrs >= 0 && AllQuadsMetricList[Index].FixedQuadNbrs <= 4);
+					}
+				}
+			}
+		}
+
+
+		// insert all potential quads into priority queue
+		FIndexPriorityQueue Queue;
+		Queue.Initialize(AllQuadsMetricList.Num());
+		for (int32 k = 0; k < AllQuadsMetricList.Num(); ++k)
+		{
+			if (IsDoneQuad(GetQuadForIndex(k)) == false)
+			{
+				double Metric = AllQuadsMetricList[k].GetMetric(QuadAdjacencyWeight);
+				if (IsViableQuad(Metric))
+				{
+					Queue.Insert(k, Metric);
+				}
+			}
+		}
+
+		// Repeatedly extract quads until the priority queue is empty.
+		TArray<int32> Reinsertions, OtherNbrTris;
+		while (Queue.GetCount() > 0)
+		{
+			int32 CurIndex = Queue.Dequeue();
+			FTriPairQuad& Quad = GetQuadForIndex(CurIndex);
+			if (IsDoneQuad(Quad))
+			{
+				continue;
+			}
+
+			FoundQuads.Add(Quad);
+			DoneTris[Quad.Triangles.A] = true;
+			DoneTris[Quad.Triangles.B] = true;
+			RemainingTriangles -= 2;
+
+			// After each quad is selected, it's adjacent quad-neighbours become more likely 
+			// to also be quads. Find and reinsert those neighbours with updated metrics.
+			if ( QuadAdjacencyWeight > 0 )
+			{
+				Reinsertions.Reset();
+				OtherNbrTris.Reset();
+				for (int32 j = 0; j < 4; ++j)
+				{
+					int32 adjacent_tid = Quad.QuadEdgeOtherTris[j];
+					if (adjacent_tid != -1)
+					{
+						for (int32 k = 0; k < 3; ++k)
+						{
+							FTriPairQuad& NbrQuad = TriPotentialQuads[adjacent_tid][k];
+							if (NbrQuad.IsValid() && IsDoneQuad(NbrQuad) == false && NbrQuad.IsAdjacent(Quad) )
+							{
+								Reinsertions.AddUnique(NbrQuad.Index);
+								OtherNbrTris.AddUnique( NbrQuad.Triangles.OtherElement(adjacent_tid) );
+							}
+						}
+					}
+				}
+
+				// gross need second pass here because each quad appears twice in TriPotentialQuads list, 
+				// and so the 'far' neighbour of each quad needs to be updated too..., but it's first-tris
+				// are not in the adjacency list of current Quad
+				for (int32 far_adjacent_tid : OtherNbrTris)
+				{
+					for (int32 k = 0; k < 3; ++k)
+					{
+						FTriPairQuad& NbrQuad = TriPotentialQuads[far_adjacent_tid][k];
+						if (NbrQuad.IsValid() && IsDoneQuad(NbrQuad) == false && NbrQuad.IsAdjacent(Quad) )
+						{
+							Reinsertions.AddUnique(NbrQuad.Index);
+						}
+					}
+				}
+
+				// for each nbr quad, increment it's nbr count, compute a new metric weighted
+				// by number of known neighbours, and reinsert into the queue
+				for (int32 Index : Reinsertions)
+				{
+					if (Queue.Contains(Index))		// may not be in queue due to metric limit or other filtering
+					{
+						FQuadWithMetric& UpdateQM = AllQuadsMetricList[Index];
+						UpdateQM.FixedQuadNbrs++;
+						check(UpdateQM.FixedQuadNbrs >= 0 && UpdateQM.FixedQuadNbrs <= 4);
+
+						// more quad nbrs == much more likely. Possibly just divide is not enough....
+						double ModifiedMetric = UpdateQM.GetMetric(QuadAdjacencyWeight);
+						Queue.Remove(Index);
+						Queue.Insert(Index, ModifiedMetric);
+					}
+				}
+			}
+		}
+
+#if 1
+		// In first round we do purely greedy search. In later rounds we use isolated triangles
+		// as the seeds for "failure areas", and remove those areas from the solution, and then
+		// re-run it, with the idea that the adjacency metric may correct the result. Note that
+		// this is only possible if a large enough area is removed, otherwise the solution can be
+		// "locked" by the surrounding quads. So, the 'removal' region grows each round, and the
+		// neighbour weight increases
+		if ((MaxSearchRounds - Iters) > 0)
+		{
+			FMeshFaceSelection RecomputeRegion(Mesh);
+
+			for (int32 tid : Mesh->TriangleIndicesItr())
+			{
+				if (DoneTris[tid] == false)
+				{
+					FIndex3i TriNbrTris = Mesh->GetTriNeighbourTris(tid);
+					int32 QuadNbrs = 0;
+					QuadNbrs += (TriNbrTris.A >= 0 && DoneTris[TriNbrTris.A]) ? 1 : 0;
+					QuadNbrs += (TriNbrTris.B >= 0 && DoneTris[TriNbrTris.B]) ? 1 : 0;
+					QuadNbrs += (TriNbrTris.C >= 0 && DoneTris[TriNbrTris.C]) ? 1 : 0;
+					if (QuadNbrs == 3)
+					{
+						RecomputeRegion.Select(tid);
+					}
+				}
+			}
+			int32 NumIsolated = RecomputeRegion.Num();
+			RecomputeRegion.ExpandToFaceNeighbours(1+Iters);
+
+			if (NumIsolated > 0)
+			{
+				for (FTriPairQuad Quad : FoundQuads)
+				{
+					bool bIsFilteredQuad = ( RecomputeRegion.Contains(Quad.Triangles.A) ||
+						RecomputeRegion.Contains(Quad.Triangles.B) );
+					if (bIsFilteredQuad == false)
+					{
+						LastRoundQuads.Add(Quad);
+					}
+				}
+			}
+
+			QuadAdjacencyWeight *= 1.1;
+
+			//UE_LOG(LogTemp, Warning, TEXT("Found %d quads, have %d tris remaining, %d isolated tris, %d isolated-adjacent quads"), FoundQuads.Num(), RemainingTriangles, NumIsolated, (FoundQuads.Num() - LastRoundQuads.Num()));
+		}
+#endif
+	}
+
+	// in some cases we clearly have planar polygons or tip-fans...how can we extract those?
+
+	// create group for each quad
+	FoundPolygroups.SetNum(FoundQuads.Num());
+	for (int32 QuadIdx = 0; QuadIdx < FoundQuads.Num(); ++QuadIdx)
+	{
+		FTriPairQuad Quad = FoundQuads[QuadIdx];
+		FoundPolygroups[QuadIdx].Add(Quad.Triangles.A);
+		FoundPolygroups[QuadIdx].Add(Quad.Triangles.B);
+	}
+
+	// add remaining triangles as individual groups
+	for (int32 tid : Mesh->TriangleIndicesItr())
+	{
+		if (DoneTris[tid] == false)
+		{
+			TArray<int> Group;
+			Group.Add(tid);
+			FoundPolygroups.Add(MoveTemp(Group));
+		}
+	}
+
+	if (bCopyToMesh)
+	{
+		CopyPolygroupsToMesh();
+	}
+
+	return (FoundQuads.Num() > 0);
+}
+
+
+
+bool FPolygroupsGenerator::FindPolygroupsFromFaceNormals(
+	double DotTolerance,
+	bool bRespectUVSeams,
+	bool bRespectNormalSeams)
 {
 	DotTolerance = 1.0 - DotTolerance;
 
+	// if we are respecting seams or hard normals, find all those edges
+	TSet<int32> InvalidEdges;
+	if (bRespectUVSeams || bRespectNormalSeams)
+	{
+		GetSeamConstraintEdges(bRespectUVSeams, bRespectNormalSeams, InvalidEdges);
+	}
+
 	// compute face normals
 	FMeshNormals Normals(Mesh);
 	Normals.ComputeTriangleNormals();
@@ -148,8 +622,14 @@ bool FPolygroupsGenerator::FindPolygroupsFromFaceNormals(double DotTolerance)
 		{
 			int CurTri = Stack.Pop(false);
 			FIndex3i NbrTris = Mesh->GetTriNeighbourTris(CurTri);
+			FIndex3i NbrEdges = Mesh->GetTriEdges(CurTri);
 			for (int j = 0; j < 3; ++j)
 			{
+				if (InvalidEdges.Contains(NbrEdges[j]))
+				{
+					continue;
+				}
+
 				if (NbrTris[j] >= 0
 					&& DoneTriangle[NbrTris[j]] == false)
 				{
diff --git a/Engine/Plugins/Runtime/GeometryProcessing/Source/DynamicMesh/Public/Polygroups/PolygroupsGenerator.h b/Engine/Plugins/Runtime/GeometryProcessing/Source/DynamicMesh/Public/Polygroups/PolygroupsGenerator.h
--- a/Engine/Plugins/Runtime/GeometryProcessing/Source/DynamicMesh/Public/Polygroups/PolygroupsGenerator.h
+++ b/Engine/Plugins/Runtime/GeometryProcessing/Source/DynamicMesh/Public/Polygroups/PolygroupsGenerator.h
@@ -56,7 +56,10 @@ public:
 	 * Find Polygroups by randomly picking initial seed triangles and then flood-filling outwards,
 	 * stopping when the opening angle at an edge is larger than the angle defined by the DotTolerance.
 	 */
-	bool FindPolygroupsFromFaceNormals(double DotTolerance = 0.0001);
+	bool FindPolygroupsFromFaceNormals(
+		double DotTolerance = 0.0001,
+		bool bRespectUVSeams = false,
+		bool bRespectNormalSeams = false);
 
 	/**
 	 * Find Polygroups based on UV Islands, ie each UV Island becomes a Polygroup
@@ -73,6 +76,17 @@ public:
 	 */
 	bool FindPolygroupsFromConnectedTris();
 
+	/**
+	* Find Polygroups by trying to invert triangluation of a polygon mesh
+	*/
+	bool FindSourceMeshPolygonPolygroups(
+		bool bRespectUVSeams = true,
+		bool bRespectNormalSeams = false,
+		double QuadAdjacencyWeight = 1.0,
+		double QuadMetricClamp = 1.0,
+		int MaxSearchRounds = 1
+	);
+
 	/**
 	 * Weight potions for algorithms below
 	 */
@@ -117,6 +131,8 @@ public:
 protected:
 	void PostProcessPolygroups(bool bApplyMerging, TFunctionRef<bool(int32, int32)> TrisConnectedPredicate = [](int, int) { return true; });
 	void OptimizePolygroups(TFunctionRef<bool(int32, int32)> TrisConnectedPredicate = [](int, int) { return true; });
+
+	void GetSeamConstraintEdges(bool bUVSeams, bool bNormalSeams, TSet<int32>& ConstraintEdgesOut) const;
 };
 
 
