DynamicMesh: re-normalize normal layers in MeshResolveTJunctions

diff --git a/Engine/Plugins/Runtime/GeometryProcessing/Source/DynamicMesh/Private/Operations/MeshResolveTJunctions.cpp b/Engine/Plugins/Runtime/GeometryProcessing/Source/DynamicMesh/Private/Operations/MeshResolveTJunctions.cpp
--- a/Engine/Plugins/Runtime/GeometryProcessing/Source/DynamicMesh/Private/Operations/MeshResolveTJunctions.cpp
+++ b/Engine/Plugins/Runtime/GeometryProcessing/Source/DynamicMesh/Private/Operations/MeshResolveTJunctions.cpp
@@ -3,6 +3,7 @@
 
 #include "Operations/MeshResolveTJunctions.h"
 #include "DynamicMesh/DynamicMesh3.h"
+#include "DynamicMesh/DynamicMeshAttributeSet.h"
 #include "Spatial/PointSetHashTable.h"
 #include "SegmentTypes.h"
 
@@ -112,7 +113,22 @@ bool FMeshResolveTJunctions::Apply()
 			BoundaryEdges.Add(SplitInfo.NewEdges.A);
 			NumSplitEdges++;
 		}
+	}
 
+	// re-normalize normal/tangent layers
+	if (Mesh->HasAttributes() && Mesh->Attributes()->NumNormalLayers() > 0)
+	{
+		int32 NumNormalLayers = Mesh->Attributes()->NumNormalLayers();
+		for (int32 Index = 0; Index < NumNormalLayers; ++Index)
+		{
+			FDynamicMeshNormalOverlay* NormalOverlay = Mesh->Attributes()->GetNormalLayer(Index);
+			for (int32 ElemIdx : NormalOverlay->ElementIndicesItr())
+			{
+				FVector3f Normal = NormalOverlay->GetElement(ElemIdx);
+				Normalize(Normal);
+				NormalOverlay->SetElement(ElemIdx, Normal);
+			}
+		}
 	}
 
 	return true;
