Geometry: Updates to fix type conversion warnings in PolygroupsGenerator.cpp

#rnx

diff --git a/Engine/Plugins/Runtime/GeometryProcessing/Source/DynamicMesh/Private/Polygroups/PolygroupsGenerator.cpp b/Engine/Plugins/Runtime/GeometryProcessing/Source/DynamicMesh/Private/Polygroups/PolygroupsGenerator.cpp
--- a/Engine/Plugins/Runtime/GeometryProcessing/Source/DynamicMesh/Private/Polygroups/PolygroupsGenerator.cpp
+++ b/Engine/Plugins/Runtime/GeometryProcessing/Source/DynamicMesh/Private/Polygroups/PolygroupsGenerator.cpp
@@ -422,7 +422,7 @@ bool FPolygroupsGenerator::FindSourceMeshPolygonPolygroups(
 				double Metric = AllQuadsMetricList[k].GetMetric(QuadAdjacencyWeight);
 				if (IsViableQuad(Metric))
 				{
-					Queue.Insert(k, Metric);
+					Queue.Insert(k, float(Metric));
 				}
 			}
 		}
@@ -494,7 +494,7 @@ bool FPolygroupsGenerator::FindSourceMeshPolygonPolygroups(
 						// more quad nbrs == much more likely. Possibly just divide is not enough....
 						double ModifiedMetric = UpdateQM.GetMetric(QuadAdjacencyWeight);
 						Queue.Remove(Index);
-						Queue.Insert(Index, ModifiedMetric);
+						Queue.Insert(Index, float(ModifiedMetric));
 					}
 				}
 			}
