GeometryCore and DynamicMesh: enable implicit type conversion warnings.

diff --git a/Engine/Plugins/Runtime/GeometryProcessing/Source/DynamicMesh/DynamicMesh.Build.cs b/Engine/Plugins/Runtime/GeometryProcessing/Source/DynamicMesh/DynamicMesh.Build.cs
--- a/Engine/Plugins/Runtime/GeometryProcessing/Source/DynamicMesh/DynamicMesh.Build.cs
+++ b/Engine/Plugins/Runtime/GeometryProcessing/Source/DynamicMesh/DynamicMesh.Build.cs
@@ -25,6 +25,6 @@ public class DynamicMesh : ModuleRules
 				"Eigen"
 			}
 		);
-
+		UnsafeTypeCastWarningLevel = WarningLevel.Warning;
 	}
 }
diff --git a/Engine/Source/Runtime/GeometryCore/GeometryCore.Build.cs b/Engine/Source/Runtime/GeometryCore/GeometryCore.Build.cs
--- a/Engine/Source/Runtime/GeometryCore/GeometryCore.Build.cs
+++ b/Engine/Source/Runtime/GeometryCore/GeometryCore.Build.cs
@@ -22,5 +22,7 @@ public class GeometryCore : ModuleRules
 				"Core"
 			}
 		);
+
+		UnsafeTypeCastWarningLevel = WarningLevel.Warning;
 	}
 }
