GeometryProcessing and tests: enable implicit type conversion warnings.

#rnx

diff --git a/Engine/Plugins/Runtime/GeometryProcessing/Source/GeometryAlgorithms/GeometryAlgorithms.Build.cs b/Engine/Plugins/Runtime/GeometryProcessing/Source/GeometryAlgorithms/GeometryAlgorithms.Build.cs
--- a/Engine/Plugins/Runtime/GeometryProcessing/Source/GeometryAlgorithms/GeometryAlgorithms.Build.cs
+++ b/Engine/Plugins/Runtime/GeometryProcessing/Source/GeometryAlgorithms/GeometryAlgorithms.Build.cs
@@ -18,5 +18,6 @@ public class GeometryAlgorithms : ModuleRules
 				"GeometryCore"
 			}
 			);
+		UnsafeTypeCastWarningLevel = WarningLevel.Warning;
 	}
 }
