GeometryProcessing: Mute warning in Eigen library used in FSparseMatrixD.h when using Clang

diff --git a/Engine/Plugins/Runtime/GeometryProcessing/Source/DynamicMesh/Private/Solvers/Internal/FSparseMatrixD.h b/Engine/Plugins/Runtime/GeometryProcessing/Source/DynamicMesh/Private/Solvers/Internal/FSparseMatrixD.h
--- a/Engine/Plugins/Runtime/GeometryProcessing/Source/DynamicMesh/Private/Solvers/Internal/FSparseMatrixD.h
+++ b/Engine/Plugins/Runtime/GeometryProcessing/Source/DynamicMesh/Private/Solvers/Internal/FSparseMatrixD.h
@@ -22,11 +22,22 @@
 #pragma warning(disable : 6313)
 #pragma warning(disable : 6294)
 #endif
+
+#if defined(__clang__)
+#pragma clang diagnostic push
+#pragma clang diagnostic ignored "-Wimplicit-int-float-conversion"
+#endif
+
 PRAGMA_DEFAULT_VISIBILITY_START
 THIRD_PARTY_INCLUDES_START
 #include <Eigen/Sparse>
 THIRD_PARTY_INCLUDES_END
 PRAGMA_DEFAULT_VISIBILITY_END
+
+#if defined(__clang__)
+#pragma clang diagnostic pop
+#endif
+
 #if defined(_MSC_VER) && USING_CODE_ANALYSIS
 #pragma warning(pop)
 #endif
