GeometryProcessing: Fix compiler error in FSparseMatrixD.h on Android

diff --git a/Engine/Plugins/Runtime/GeometryProcessing/Source/DynamicMesh/Private/Solvers/Internal/FSparseMatrixD.h b/Engine/Plugins/Runtime/GeometryProcessing/Source/DynamicMesh/Private/Solvers/Internal/FSparseMatrixD.h
--- a/Engine/Plugins/Runtime/GeometryProcessing/Source/DynamicMesh/Private/Solvers/Internal/FSparseMatrixD.h
+++ b/Engine/Plugins/Runtime/GeometryProcessing/Source/DynamicMesh/Private/Solvers/Internal/FSparseMatrixD.h
@@ -23,7 +23,7 @@
 #pragma warning(disable : 6294)
 #endif
 
-#if defined(__clang__)
+#if defined(__clang__) && __clang_major__ > 9
 #pragma clang diagnostic push
 #pragma clang diagnostic ignored "-Wimplicit-int-float-conversion"
 #endif
@@ -34,7 +34,7 @@ THIRD_PARTY_INCLUDES_START
 THIRD_PARTY_INCLUDES_END
 PRAGMA_DEFAULT_VISIBILITY_END
 
-#if defined(__clang__)
+#if defined(__clang__) && __clang_major__ > 9
 #pragma clang diagnostic pop
 #endif
 
