Interpolate the per-vertex attributes from the original mesh to the tessellated mesh in the adaptive tessellation operator.

#rnx

diff --git a/Engine/Plugins/Runtime/GeometryProcessing/Source/DynamicMesh/Private/Operations/AdaptiveTessellate.cpp b/Engine/Plugins/Runtime/GeometryProcessing/Source/DynamicMesh/Private/Operations/AdaptiveTessellate.cpp
--- a/Engine/Plugins/Runtime/GeometryProcessing/Source/DynamicMesh/Private/Operations/AdaptiveTessellate.cpp
+++ b/Engine/Plugins/Runtime/GeometryProcessing/Source/DynamicMesh/Private/Operations/AdaptiveTessellate.cpp
@@ -9,6 +9,7 @@
 #include "Misc/AssertionMacros.h"
 #include "Util/CompactMaps.h"
 #include "DynamicMesh/DynamicMeshAttributeSet.h"
+#include "DynamicMesh/DynamicVertexSkinWeightsAttribute.h"
 
 using namespace UE::Geometry;
 
@@ -33,6 +34,31 @@ namespace AdaptiveTessellateLocals
 									 const FCompactMaps& CompactInfo,
 								     TDynamicMeshOverlay<RealType, ElementSize>* ResultOverlay);
 
+	// Utility functions
+	template<typename RealType, int ElementSize>
+	void LerpElements(const RealType* Element1, const RealType* Element2, RealType* OutElement, double Alpha)
+	{
+		Alpha = FMath::Clamp(Alpha, 0.0, 1.0);
+		double OneMinusAlpha = 1.0 - Alpha;
+		for (int Idx = 0; Idx < ElementSize; ++Idx) 
+		{
+			OutElement[Idx] = RealType(OneMinusAlpha * double(Element1[Idx]) + Alpha * double(Element2[Idx]));
+		}
+	} 
+
+	template<typename RealType, int ElementSize>
+	void BaryElements(const RealType* Element1, const RealType* Element2, const RealType* Element3,
+					  RealType* OutElement, 
+					  double Alpha, double Beta, double Theta)
+	{
+		checkSlow(FMath::Abs(Alpha + Beta + Theta - 1.0) < KINDA_SMALL_NUMBER);
+		for (int Idx = 0; Idx < ElementSize; ++Idx)
+		{
+			OutElement[Idx] = RealType(Alpha * double(Element1[Idx]) + Beta * double(Element2[Idx]) + Theta * double(Element3[Idx]));
+		}
+	};
+
+
 	/**
      * Manage data containing tessellation information. This includes barycentric coordinates of the new vertices,
      * new vertex/elemnent IDs and triangle indices. This class handles cases where 2 elements are stored per vertex. 
@@ -1169,27 +1195,6 @@ namespace AdaptiveTessellateLocals
 									 const FCompactMaps& CompactInfo,
 								     TDynamicMeshOverlay<RealType, ElementSize>* ResultOverlay)  
 	{	
-		auto LerpElements = [](const RealType* Element1, const RealType* Element2, RealType* OutElement, RealType Alpha)
-		{
-			Alpha = FMath::Clamp(Alpha, RealType(0), RealType(1));
-			RealType OneMinusAlpha = (RealType)1 - Alpha;
-
-			for (int Idx = 0; Idx < ElementSize; ++Idx)
-			{
-				OutElement[Idx] = OneMinusAlpha * Element1[Idx] + Alpha * Element2[Idx];
-			}
-		};
-
-		auto BaryElements = [](const RealType* Element1, const RealType* Element2, const RealType* Element3,
-							   RealType* OutElement, 
-							   RealType Alpha, RealType Beta, RealType Theta)
-		{
-			checkSlow(FMath::Abs(Alpha + Beta + Theta - 1.0f) < KINDA_SMALL_NUMBER);
-			for (int Idx = 0; Idx < ElementSize; ++Idx)
-			{
-				OutElement[Idx] = Alpha * Element1[Idx] + Beta * Element2[Idx] + Theta * Element3[Idx];
-			}
-		};
 
 		ResultOverlay->ClearElements();
 
@@ -1231,7 +1236,7 @@ namespace AdaptiveTessellateLocals
 					const RealType Alpha = (RealType)CoordBlock[Index];
 					const int ElementID = IDBlock[Index]; 
 					
-					LerpElements(Element1, Element2, Out, Alpha);
+					LerpElements<RealType, ElementSize>(Element1, Element2, Out, Alpha);
 					
 					MapE.Add(ElementID, ResultOverlay->AppendElement(Out));
 				}
@@ -1250,7 +1255,7 @@ namespace AdaptiveTessellateLocals
 					const RealType Alpha = (RealType)CoordBlock[Index];
 					const int ElementID = IDBlock[Index]; 
 					
-					LerpElements(Element1, Element2, Out, Alpha);
+					LerpElements<RealType, ElementSize>(Element1, Element2, Out, Alpha);
 					
 					MapE.Add(ElementID, ResultOverlay->AppendElement(Out));
 				}
@@ -1275,7 +1280,7 @@ namespace AdaptiveTessellateLocals
 				const FVector3d Bary = CoordBlock[Index]; 
 				const int ElementID = IDBlock[Index];
 
-				BaryElements(Element1, Element2, Element3, Out, RealType(Bary[0]), RealType(Bary[1]), RealType(Bary[2]));
+				BaryElements<RealType, ElementSize>(Element1, Element2, Element3, Out, Bary[0], Bary[1], Bary[2]);
 
 				MapE.Add(ElementID, ResultOverlay->AppendElement(Out));
 			}
@@ -1311,6 +1316,283 @@ namespace AdaptiveTessellateLocals
 		}
 	}
 
+	
+	/** 
+	 * General purpose function to interpolate any per-vertex data. The function will iterate over all vertices 
+	 * inserted along the edges and call InterpolateEdgeFunc. Then it will iterate over all of the vertices inserted 
+	 * inside of the triangles and call InterpolateInnerFunc.
+	 * 
+	 * @param InterpolateEdgeFunc A callback function with the signature void(int V1, int V2, int NewV, double U).
+	 * 						      V1,V2 are the edge vertex indices into the original input mesh. NewV is the mapped 
+	 * 							  vertex id into the tessellated mesh for which we are asking the function to set the 
+	 * 							  value for. U is the linear interpolation coefficient of NewV with respect to V1,V2. 
+	 * 
+	 * @param InterpolateInnerFunc A callback function with the signature void(int V1, int V2, int V3, int NewV, double U, double V, double W).
+	 * 						       V1,V2,V3 are the vertex indices into the original input mesh. NewV is the mapped 
+	 * 							   vertex id into the tessellated mesh for which we are asking the function to set the 
+	 * 							   value for. U,V,W are the barycentric coordinates of NewV with respect to V1,V2,V3.
+	 */
+	void InterpolateVertexData(const FDynamicMesh3* Mesh,
+							   const TFunctionRef<void(int V1, int V2, int NewV, double U)> InterpolateEdgeFunc, 
+							   const TFunctionRef<void(int V1, int V2, int V3, int NewV, double U, double V, double W)> InterpolateInnerFunc, 
+						   	   FTessellationData& TessData,
+						   	   const bool bUseParallel,
+						   	   const FCompactMaps& CompactInfo)
+	{
+		// Set value for the vertices inserted along edges via interpolation
+		ParallelFor(TessData.EdgesToTessellate.Num(), [&](int32 EIndex)
+		{
+			const int EdgeID = TessData.EdgesToTessellate[FSetElementId::FromInteger(EIndex)];
+			
+			TArrayView<double> CoordBlock = TessData.MapEdgeCoordBufferBlock(EdgeID);
+			TArrayView<int> IDBlock = TessData.MapEdgeIDBufferBlock(EdgeID);
+			checkSlow(CoordBlock.Num() == IDBlock.Num());
+
+			const FIndex2i EdgeV = Mesh->GetEdgeV(EdgeID);
+
+			for (int Index = 0; Index < CoordBlock.Num(); ++Index)
+			{	
+				const double Alpha = CoordBlock[Index];
+				const int VID =  IDBlock[Index]; 
+				const int ToVID = CompactInfo.GetVertexMapping(VID);
+
+				InterpolateEdgeFunc(EdgeV.A, EdgeV.B, ToVID, Alpha);
+			}
+		}, bUseParallel ? EParallelForFlags::None : EParallelForFlags::ForceSingleThread);
+
+		// Set value for the vertices inserted inside of the triangles via interpolation 
+		ParallelFor(TessData.TrianglesToTessellate.Num(), [&](int32 TIndex)
+		{
+			const int TriangleID = TessData.TrianglesToTessellate[FSetElementId::FromInteger(TIndex)];
+
+			TArrayView<FVector3d> CoordBlock = TessData.MapInnerCoordBufferBlock(TriangleID);
+			TArrayView<int> IDBlock = TessData.MapInnerIDBufferBlock(TriangleID);
+			checkSlow(CoordBlock.Num() == IDBlock.Num());
+
+			const FIndex3i TriangleV = Mesh->GetTriangle(TriangleID);
+
+			for (int Index = 0; Index < CoordBlock.Num(); ++Index)
+			{	
+				const FVector3d BaryCoords = CoordBlock[Index];
+				const int VID = IDBlock[Index];
+				const int ToVID = CompactInfo.GetVertexMapping(VID);
+
+				InterpolateInnerFunc(TriangleV.A, TriangleV.B, TriangleV.C, ToVID, BaryCoords[0], BaryCoords[1], BaryCoords[2]);
+			}
+		}, bUseParallel ? EParallelForFlags::None : EParallelForFlags::ForceSingleThread);
+	}
+
+
+	/** Interpolate any data stored in the TDynamicMeshVertexAttribute */
+	template<typename RealType, int ElementSize>
+	void ConstructDynamicMeshVertexAttribute(const FDynamicMesh3* Mesh,
+								  		     const TDynamicMeshVertexAttribute<RealType, ElementSize>* Attribute,
+						   		  		     FTessellationData& TessData,
+						   		  		     const bool bUseParallel,
+						   		  		     const FCompactMaps& CompactInfo,
+											 TDynamicMeshVertexAttribute<RealType, ElementSize>* OutAttribute)
+	{
+		auto InterpolateEdgeFunc = [Attribute, OutAttribute] (int V1, int V2, int NewV, double U)
+		{
+			RealType Value1[ElementSize];
+			RealType Value2[ElementSize];
+			RealType OutValue[ElementSize];
+
+			Attribute->GetValue(V1, Value1);
+			Attribute->GetValue(V2, Value2);
+			
+			LerpElements<RealType, ElementSize>(Value1, Value2, OutValue, U);
+
+			OutAttribute->SetValue(NewV, OutValue);
+		};
+
+		auto InterpolateInnerFunc = [Attribute, OutAttribute](int V1, int V2, int V3, int NewV, double U, double V, double W)
+		{
+			RealType Value1[ElementSize];
+			RealType Value2[ElementSize];
+			RealType Value3[ElementSize];
+			RealType OutValue[ElementSize];
+
+			Attribute->GetValue(V1, Value1);
+			Attribute->GetValue(V2, Value2);
+
+			Attribute->GetValue(V3, Value3);
+			BaryElements<RealType, ElementSize>(Value1, Value2, Value3, OutValue, U, V, W);
+
+			OutAttribute->SetValue(NewV, OutValue);
+		};
+
+		RealType Value[ElementSize];
+		for (const int VID : Mesh-> VertexIndicesItr())
+		{
+			Attribute->GetValue(VID, Value);
+			const int ToVID = CompactInfo.GetVertexMapping(VID);
+			OutAttribute->SetValue(ToVID, Value);
+		}
+
+		InterpolateVertexData(Mesh, InterpolateEdgeFunc, InterpolateInnerFunc, TessData, bUseParallel, CompactInfo);
+	}
+
+
+	/** Interpolate skin weights */
+	void ConstructVertexSkinWeightsAttribute(const FDynamicMesh3* Mesh,
+								  		     const FDynamicMeshVertexSkinWeightsAttribute* Attribute,
+						   		  		     FTessellationData& TessData,
+						   		  		     const bool bUseParallel,
+						   		  		     const FCompactMaps& CompactInfo,
+											 FDynamicMeshVertexSkinWeightsAttribute* OutAttribute)
+	{
+		using FBoneWeights = UE::AnimationCore::FBoneWeights;
+		
+		auto InterpolateEdgeFunc = [Attribute, OutAttribute] (int V1, int V2, int NewV, double U)
+		{
+			FBoneWeights Value1;
+			FBoneWeights Value2;
+			FBoneWeights OutValue;
+
+			Attribute->GetValue(V1, Value1);
+			Attribute->GetValue(V2, Value2);
+		
+			U = FMath::Clamp(U, 0.0, 1.0);
+			OutValue = FBoneWeights::Blend(Value1, Value2, (float)U);
+
+			OutAttribute->SetValue(NewV, OutValue);
+		};
+
+		auto InterpolateInnerFunc = [Attribute, OutAttribute] (int V1, int V2, int V3, int NewV, double U, double V, double W)
+		{
+			FBoneWeights Value1, Value2, Value3, OutValue;
+
+			Attribute->GetValue(V1, Value1);
+			Attribute->GetValue(V2, Value2);
+			Attribute->GetValue(V3, Value3);
+				
+			// Since FBoneWeights only defines Blend for two inputs, we need to split the barycentric coordinate 
+			// interpolation into two blends. This mimics TDynamicVertexSkinWeightsAttribute::SetBoneWeightsFromBary
+			if (FMath::IsNearlyZero(V + W) == false)
+			{
+				const double BCW = V / (V + W);
+				const FBoneWeights BC = FBoneWeights::Blend(Value2, Value3, (float)BCW);
+				OutValue = FBoneWeights::Blend(Value1, BC, (float)U);
+			}
+			else
+			{
+				OutValue = Value1;
+			}
+
+			OutAttribute->SetValue(NewV, OutValue);
+		};
+
+		FBoneWeights Value;
+		for (const int VID : Mesh-> VertexIndicesItr())
+		{
+			Attribute->GetValue(VID, Value);
+			const int ToVID = CompactInfo.GetVertexMapping(VID);
+			OutAttribute->SetValue(ToVID, Value);
+		}
+
+		InterpolateVertexData(Mesh, InterpolateEdgeFunc, InterpolateInnerFunc, TessData, bUseParallel, CompactInfo);
+	}
+
+
+	/** Interpolate per-vertex normals */
+	void ConstructPerVertexNormals(const FDynamicMesh3* Mesh,
+						   		   FTessellationData& TessData,
+						   		   const bool bUseParallel,
+						   		   const FCompactMaps& CompactInfo,
+								   FDynamicMesh3* OutMesh)
+	{
+		auto InterpolateEdgeFunc = [Mesh, OutMesh] (int V1, int V2, int NewV, double U)
+		{
+			FVector3f OutValue;
+			U = FMath::Clamp(U, 0.0, 1.0);
+			OutValue =  (1.0 - U) * Mesh->GetVertexNormal(V1)  + U * Mesh->GetVertexNormal(V2);
+			OutMesh->SetVertexNormal(NewV, OutValue);
+		};
+
+		auto InterpolateInnerFunc = [Mesh, OutMesh] (int V1, int V2, int V3, int NewV, double U, double V, double W)
+		{
+			FVector3f OutValue;
+			checkSlow(FMath::Abs(U + V + W - 1.0) < KINDA_SMALL_NUMBER);
+			OutValue = U * Mesh->GetVertexNormal(V1) + V * Mesh->GetVertexNormal(V2) + W * Mesh->GetVertexNormal(V3);
+			OutMesh->SetVertexNormal(NewV, OutValue);
+		};
+
+		for (const int VID : Mesh->VertexIndicesItr())
+		{
+			const int ToVID = CompactInfo.GetVertexMapping(VID);
+			OutMesh->SetVertexNormal(ToVID, Mesh->GetVertexNormal(VID));
+		}
+
+		InterpolateVertexData(Mesh, InterpolateEdgeFunc, InterpolateInnerFunc, TessData, bUseParallel, CompactInfo);
+	}
+
+
+	/** Interpolate per-vertex UVs */
+	void ConstructPerVertexUVs(const FDynamicMesh3* Mesh,
+						   	   FTessellationData& TessData,
+						   	   const bool bUseParallel,
+						   	   const FCompactMaps& CompactInfo,
+							   FDynamicMesh3* OutMesh)
+	{
+		auto InterpolateEdgeFunc = [Mesh, OutMesh] (int V1, int V2, int NewV, double U)
+		{
+			FVector2f OutValue;
+			U = FMath::Clamp(U, 0.0, 1.0);
+			OutValue =  float(1.0 - U) * Mesh->GetVertexUV(V1)  + float(U) * Mesh->GetVertexUV(V2);
+			OutMesh->SetVertexUV(NewV, OutValue);
+		};
+
+		auto InterpolateInnerFunc = [Mesh, OutMesh] (int V1, int V2, int V3, int NewV, double U, double V, double W)
+		{
+			FVector2f OutValue;
+			checkSlow(FMath::Abs(U + V + W - 1.0) < KINDA_SMALL_NUMBER);
+			OutValue = float(U) * Mesh->GetVertexUV(V1) + float(V) * Mesh->GetVertexUV(V2) + float(W) * Mesh->GetVertexUV(V3);
+			OutMesh->SetVertexUV(NewV, OutValue);
+		};
+
+		for (const int VID : Mesh-> VertexIndicesItr())
+		{
+			const int ToVID = CompactInfo.GetVertexMapping(VID);
+			OutMesh->SetVertexUV(ToVID, Mesh->GetVertexUV(VID));
+		}
+
+		InterpolateVertexData(Mesh, InterpolateEdgeFunc, InterpolateInnerFunc, TessData, bUseParallel, CompactInfo);
+	}
+
+
+	/** Interpolate per-vertex colors */
+	void ConstructPerVertexColors(const FDynamicMesh3* Mesh,
+						   	      FTessellationData& TessData,
+						   	      const bool bUseParallel,
+						   	      const FCompactMaps& CompactInfo,
+							      FDynamicMesh3* OutMesh)
+	{
+		auto InterpolateEdgeFunc = [Mesh, OutMesh] (int V1, int V2, int NewV, double U)
+		{
+			FVector4f OutValue;
+			U = FMath::Clamp(U, 0.0, 1.0);
+			OutValue =  float(1.0 - U) * Mesh->GetVertexColor(V1)  + float(U) * Mesh->GetVertexColor(V2);
+			OutMesh->SetVertexColor(NewV, OutValue);
+		};
+
+		auto InterpolateInnerFunc = [Mesh, OutMesh] (int V1, int V2, int V3, int NewV, double U, double V, double W)
+		{
+			FVector4f OutValue;
+			checkSlow(FMath::Abs(U + V + W - 1.0) < KINDA_SMALL_NUMBER);
+			OutValue = float(U) * Mesh->GetVertexColor(V1) + float(V) * Mesh->GetVertexColor(V2) + float(W) * Mesh->GetVertexColor(V3);
+			OutMesh->SetVertexColor(NewV, OutValue);
+		};
+
+		for (const int VID : Mesh-> VertexIndicesItr())
+		{
+			const int ToVID = CompactInfo.GetVertexMapping(VID);
+			OutMesh->SetVertexColor(ToVID, Mesh->GetVertexColor(VID));
+		}
+
+		InterpolateVertexData(Mesh, InterpolateEdgeFunc, InterpolateInnerFunc, TessData, bUseParallel, CompactInfo);
+	}
+
 	/** 
 	 * Run the main tessellation logic for the geometry, overlays and attributes. 
 	 * 
@@ -1372,13 +1654,48 @@ namespace AdaptiveTessellateLocals
 					ConstructTriangleAttribute(InMesh, InAttributes->GetPolygroupLayer(Idx), TessData, bUseParallel, CompactInfo, OutAttributes->GetPolygroupLayer(Idx));
 				}
 			}
+
+			for (const TTuple<FName, TUniquePtr<FDynamicMeshVertexSkinWeightsAttribute>>& AttributeInfo : InAttributes->GetSkinWeightsAttributes())
+			{
+				FDynamicMeshVertexSkinWeightsAttribute* SkinAttribute = new FDynamicMeshVertexSkinWeightsAttribute(OutMesh);
+				ConstructVertexSkinWeightsAttribute(InMesh, AttributeInfo.Value.Get(), TessData, bUseParallel, CompactInfo, SkinAttribute);
+				SkinAttribute->SetName(AttributeInfo.Value.Get()->GetName());
+				OutAttributes->AttachSkinWeightsAttribute(AttributeInfo.Key, SkinAttribute);
+			}
+
+			if (InAttributes->NumWeightLayers() > 0)
+			{
+				OutAttributes->SetNumWeightLayers(InAttributes->NumWeightLayers());
+				
+				for (int Idx = 0; Idx < InAttributes->NumWeightLayers(); ++Idx)
+				{
+					ConstructDynamicMeshVertexAttribute<float, 1>(InMesh, InAttributes->GetWeightLayer(Idx), TessData, bUseParallel, CompactInfo, OutAttributes->GetWeightLayer(Idx));
+					OutAttributes->GetWeightLayer(Idx)->SetName(InAttributes->GetWeightLayer(Idx)->GetName());
+				}
+			}
+		}
+		
+		if (InMesh->HasVertexNormals()) 
+		{	
+			OutMesh->EnableVertexNormals(FVector3f::Zero());
+			ConstructPerVertexNormals(InMesh, TessData, bUseParallel, CompactInfo, OutMesh);
+		}
+		
+		if (InMesh->HasVertexUVs()) 
+		{	
+			OutMesh->EnableVertexUVs(FVector2f::Zero());
+			ConstructPerVertexUVs(InMesh, TessData, bUseParallel, CompactInfo, OutMesh);
+		}
+
+		if (InMesh->HasVertexColors()) 
+		{	
+			OutMesh->EnableVertexColors(FVector4f::Zero());
+			ConstructPerVertexColors(InMesh, TessData, bUseParallel, CompactInfo, OutMesh);
 		}
 	}
 }
 
 
-
-
 //
 // GLSL pattern
 //
@@ -1411,6 +1728,11 @@ TUniquePtr<FTessellationPattern> FAdaptiveTessellate::CreateConcentricRingsTesse
 								   							  					    		   const int InTessellationLevel,
 								   							  					    		   const TArray<int>& InTriangleList) 
 {
+	if (InTriangleList.IsEmpty()) 
+	{
+		return nullptr;
+	}
+
 	TArray<int> EdgeTessLevels;
 	EdgeTessLevels.Init(0, InMesh->MaxEdgeID());
 	for (const int TriangleID : InTriangleList) 
