Add progress cancel functionality to the adaptive tessellation operator and the displace tool when the selective option is used.

#rnx

diff --git a/Engine/Plugins/Experimental/MeshModelingToolsetExp/Source/MeshModelingToolsExp/Private/DisplaceMeshTool.cpp b/Engine/Plugins/Experimental/MeshModelingToolsetExp/Source/MeshModelingToolsExp/Private/DisplaceMeshTool.cpp
--- a/Engine/Plugins/Experimental/MeshModelingToolsetExp/Source/MeshModelingToolsExp/Private/DisplaceMeshTool.cpp
+++ b/Engine/Plugins/Experimental/MeshModelingToolsetExp/Source/MeshModelingToolsExp/Private/DisplaceMeshTool.cpp
@@ -290,12 +290,12 @@ namespace DisplaceMeshToolLocals{
 					
 					FDynamicMesh3 OutMesh;
 					FAdaptiveTessellate Tessellator(ResultMesh.Get(), &OutMesh);
-					
+					Tessellator.Progress = ProgressCancel;
 					Tessellator.SetPattern(Pattern.Get());
 					
 					if (ensureMsgf(Tessellator.Validate() == EOperationValidationResult::Ok, TEXT("The tessellator parameters are invalid.")))
 					{
-						if (ensureMsgf(Tessellator.Compute(), TEXT("Failed to tessellate the selected triangles.")))
+						if (Tessellator.Compute())
 						{
 							*ResultMesh = MoveTemp(OutMesh);
 						}
diff --git a/Engine/Plugins/Runtime/GeometryProcessing/Source/DynamicMesh/Private/Operations/AdaptiveTessellate.cpp b/Engine/Plugins/Runtime/GeometryProcessing/Source/DynamicMesh/Private/Operations/AdaptiveTessellate.cpp
--- a/Engine/Plugins/Runtime/GeometryProcessing/Source/DynamicMesh/Private/Operations/AdaptiveTessellate.cpp
+++ b/Engine/Plugins/Runtime/GeometryProcessing/Source/DynamicMesh/Private/Operations/AdaptiveTessellate.cpp
@@ -21,17 +21,19 @@ namespace AdaptiveTessellateLocals
 	template<typename RealType, int ElementSize> 
 	class FOverlayTessellationData;
 
-	void ConstructTessellatedMesh(const FDynamicMesh3* Mesh,
+	bool ConstructTessellatedMesh(const FDynamicMesh3* Mesh,
+								  FProgressCancel* Progress,
 								  FTessellationData& TessData,
 								  FCompactMaps& CompactInfo,
 								  FDynamicMesh3* ResultMesh);
 
 	template<typename RealType, int ElementSize> 
-	void ConstructTessellatedOverlay(const FDynamicMesh3* Mesh, 
+	bool ConstructTessellatedOverlay(const FDynamicMesh3* Mesh, 
 									 const TDynamicMeshOverlay<RealType, ElementSize>* Overlay,
 									 FTessellationData& MeshTessData, 
 									 FOverlayTessellationData<RealType, ElementSize>& TessData, 
 									 const FCompactMaps& CompactInfo,
+									 FProgressCancel* Progress, 
 								     TDynamicMeshOverlay<RealType, ElementSize>* ResultOverlay);
 
 	// Utility functions
@@ -846,10 +848,15 @@ namespace AdaptiveTessellateLocals
 
 
 
-	/** Use the tessellation pattern to tessellate the mesh and generate the tessellation data. */
-	void TessellateGeometry(const FDynamicMesh3* Mesh, 
+	/** 
+	 * Use the tessellation pattern to tessellate the mesh and generate the tessellation data. 
+	 * 
+	 * @return false if the tessellation failed or was cancelled by the user
+	 */
+	bool TessellateGeometry(const FDynamicMesh3* Mesh, 
 							const FTessellationPattern* Pattern, 
 							const bool bUseParallel, 
+							FProgressCancel* Progress,
 							FCompactMaps& OutCompactInfo, 
 							FTessellationData& TessData,
 							FDynamicMesh3* OutMesh) 
@@ -861,6 +868,11 @@ namespace AdaptiveTessellateLocals
 		TArray<int> EdgesToTessellate = TessData.EdgesToTessellate.Array();
 		ParallelFor(EdgesToTessellate.Num(), [&](int32 Index)
 		{
+			if (Progress && Progress->Cancelled()) 
+			{
+				return;
+			}
+
 			FTessellationPattern::EdgePatch Patch;
 
 			const int EdgeID = EdgesToTessellate[Index];
@@ -871,11 +883,20 @@ namespace AdaptiveTessellateLocals
 			Pattern->TessellateEdgePatch(Patch);
 		}, bUseParallel ? EParallelForFlags::None : EParallelForFlags::ForceSingleThread);
 
+		if (Progress && Progress->Cancelled()) 
+		{
+			return false;
+		}
 		
 		// Tessellate triangle patches
 		TArray<int> TrianglesToTessellate = TessData.TrianglesToTessellate.Array();
 		ParallelFor(TrianglesToTessellate.Num(), [&](int32 Index)
 		{
+			if (Progress && Progress->Cancelled()) 
+			{
+				return;
+			}
+
 			FTessellationPattern::TrianglePatch Patch;
 			
 			const int TriangleID = TrianglesToTessellate[Index]; 
@@ -925,17 +946,32 @@ namespace AdaptiveTessellateLocals
 			Pattern->TessellateTriPatch(Patch);
 		}, bUseParallel ? EParallelForFlags::None : EParallelForFlags::ForceSingleThread);
 
-		AdaptiveTessellateLocals::ConstructTessellatedMesh(Mesh, TessData, OutCompactInfo, OutMesh);
+		if (Progress && Progress->Cancelled()) 
+		{
+			return false;
+		}
+
+		if (AdaptiveTessellateLocals::ConstructTessellatedMesh(Mesh, Progress, TessData, OutCompactInfo, OutMesh) == false) 
+		{
+			return false;
+		}
+
+		return true;
 	}
 
 
-	/** Use the tessellation pattern to tessellate an overlay. */
+	/** 
+	 * Use the tessellation pattern to tessellate an overlay. 
+	 * 
+	 * @return false if the tessellation failed or was cancelled by the user
+	 */
 	template<typename RealType, int ElementSize> 
-	void TessellateOverlay(const FDynamicMesh3* Mesh, 
+	bool TessellateOverlay(const FDynamicMesh3* Mesh, 
 						   const TDynamicMeshOverlay<RealType, ElementSize>* Overlay,
 						   const FTessellationPattern* Pattern, 
 						   FTessellationData& MeshTessData,
 						   const bool bUseParallel, 
+						   FProgressCancel* Progress,
 						   const FCompactMaps& CompactInfo,
 						   TDynamicMeshOverlay<RealType, ElementSize>* OutOverlay)
 	{	
@@ -949,6 +985,11 @@ namespace AdaptiveTessellateLocals
 		TArray<int> EdgesToTessellate = OverlayTessData.EdgesToTessellate.Array();
 		ParallelFor(EdgesToTessellate.Num(), [&](int32 Index)
 		{
+			if (Progress && Progress->Cancelled()) 
+			{
+				return;
+			}
+
 			const int EdgeID = EdgesToTessellate[Index];
 			
 			const FIndex2i EdgeTri = Mesh->GetEdgeT(EdgeID);
@@ -972,10 +1013,20 @@ namespace AdaptiveTessellateLocals
 			}
 		}, bUseParallel ? EParallelForFlags::None : EParallelForFlags::ForceSingleThread);
 
+		if (Progress && Progress->Cancelled()) 
+		{
+			return false;
+		}
+
 		// Tessellate triangle patches
 		TArray<int> TrianglesToTessellate = OverlayTessData.TrianglesToTessellate.Array();
 		ParallelFor(TrianglesToTessellate.Num(), [&](int32 Index)
 		{
+			if (Progress && Progress->Cancelled()) 
+			{
+				return;
+			}
+			
 			FTessellationPattern::TrianglePatch Patch;
 			
 			const int TriangleID = TrianglesToTessellate[Index]; 
@@ -1025,12 +1076,23 @@ namespace AdaptiveTessellateLocals
 			Pattern->TessellateTriPatch(Patch);
 		}, bUseParallel ? EParallelForFlags::None : EParallelForFlags::ForceSingleThread);
 
-		AdaptiveTessellateLocals::ConstructTessellatedOverlay(Mesh, 
-															  Overlay,
-															  MeshTessData, 
-															  OverlayTessData, 
-															  CompactInfo, 
-															  OutOverlay);
+		if (Progress && Progress->Cancelled()) 
+		{
+			return false;
+		}
+
+		if (AdaptiveTessellateLocals::ConstructTessellatedOverlay(Mesh, 
+															  	  Overlay,
+															  	  MeshTessData, 
+															  	  OverlayTessData, 
+															  	  CompactInfo, 
+																  Progress, 
+															  	  OutOverlay) == false) 
+		{
+			return false;
+		}
+
+		return true;
 	}
 
 
@@ -1074,8 +1136,10 @@ namespace AdaptiveTessellateLocals
 	 * Output mesh will be compact. 
 	 * 
 	 * @param CompactInfo Stores vertex and triangle mappings.
+	 * @return false if the user cancelled the operation.
 	 */
-	void ConstructTessellatedMesh(const FDynamicMesh3* Mesh,
+	bool ConstructTessellatedMesh(const FDynamicMesh3* Mesh,
+								  FProgressCancel* Progress,
 								  FTessellationData& TessData,
 								  FCompactMaps& CompactInfo,
 								  FDynamicMesh3* ResultMesh)
@@ -1101,6 +1165,11 @@ namespace AdaptiveTessellateLocals
 			}
 		}
 
+		if (Progress && Progress->Cancelled()) 
+		{
+			return false;
+		}
+
 		// Append all of the vertices along tessellated edges
 		for (const int EdgeID : TessData.EdgesToTessellate)
 		{
@@ -1122,6 +1191,11 @@ namespace AdaptiveTessellateLocals
 			}
 		}
 
+		if (Progress && Progress->Cancelled()) 
+		{
+			return false;
+		}
+
 		// Append all of the vertices inside tessellated triangles
 		for (const int TriangleID : TessData.TrianglesToTessellate)
 		{
@@ -1144,6 +1218,11 @@ namespace AdaptiveTessellateLocals
 			}
 		}
 
+		if (Progress && Progress->Cancelled()) 
+		{
+			return false;
+		}
+
 		if (Mesh->HasTriangleGroups()) 
 		{
 			ResultMesh->EnableTriangleGroups();
@@ -1181,28 +1260,36 @@ namespace AdaptiveTessellateLocals
 		//TODO: instead of generating the mesh from scratch, we could remove triangles we are tessellating from the 
 		// input mesh and append the new ones. This can be an option when the number of triangles tessellated is small,
 		// compared to the total number of triangles.
+
+		if (Progress && Progress->Cancelled()) 
+		{
+			return false;
+		}
+
+		return true;
 	}
 
 
 	/**
 	 * Construct a new overlay from the tessellation data and compact information from the geometry tessellation.
+	 * 
+	 * @return false if fails or the user cancells the operation
 	 */
 	template<typename RealType, int ElementSize> 
-	void ConstructTessellatedOverlay(const FDynamicMesh3* Mesh, 
+	bool ConstructTessellatedOverlay(const FDynamicMesh3* Mesh, 
 									 const TDynamicMeshOverlay<RealType, ElementSize>* Overlay,
 									 FTessellationData& MeshTessData, 
 									 FOverlayTessellationData<RealType, ElementSize>& OverlayTessData, 
 									 const FCompactMaps& CompactInfo,
+									 FProgressCancel* Progress, 
 								     TDynamicMeshOverlay<RealType, ElementSize>* ResultOverlay)  
 	{	
-
 		ResultOverlay->ClearElements();
 
 		// Need to track the ID of the appended elements to handle non-compact meshes
 		TMap<int, int> MapE; 
 		MapE.Reserve(Overlay->ElementCount() + OverlayTessData.EdgeIDOffsets.Num());
-		
-		
+				
 		// Buffers to be reused
 		RealType Element1[ElementSize];
 		RealType Element2[ElementSize];
@@ -1216,6 +1303,11 @@ namespace AdaptiveTessellateLocals
 			MapE.Add(ElementID, ResultOverlay->AppendElement(Out));
 		}
 		
+		if (Progress && Progress->Cancelled()) 
+		{
+			return false;
+		}
+		
 		// Add all the new elements inserted along the edges by the tessellator
 		for (const int EdgeID : OverlayTessData.EdgesToTessellate)
 		{
@@ -1262,6 +1354,11 @@ namespace AdaptiveTessellateLocals
 			}
 		}
 
+		if (Progress && Progress->Cancelled()) 
+		{
+			return false;
+		}
+
 		// Add all the elements added inside of the triangles by the tessellator
 		for (const int TriangleID : OverlayTessData.TrianglesToTessellate)
 		{
@@ -1286,6 +1383,11 @@ namespace AdaptiveTessellateLocals
 			}
 		}
 
+		if (Progress && Progress->Cancelled()) 
+		{
+			return false;
+		}
+
 		// Set the overlay triangles and point them to the correct element ids
 		for (const int TriangleID : Mesh->TriangleIndicesItr())
 		{
@@ -1314,6 +1416,14 @@ namespace AdaptiveTessellateLocals
 				ResultOverlay->SetTriangle(ToTID, FIndex3i(MapE[Tri.A], MapE[Tri.B], MapE[Tri.C]));
 			}
 		}
+
+		
+		if (Progress && Progress->Cancelled()) 
+		{
+			return false;
+		}
+
+		return true;
 	}
 
 	
@@ -1597,15 +1707,19 @@ namespace AdaptiveTessellateLocals
 	 * Run the main tessellation logic for the geometry, overlays and attributes. 
 	 * 
 	 * @param OutMesh The resulting tessellated mesh.
+	 * @return false if the tessellation failed or the user cancelled it
 	 */
-	void Tessellate(const FDynamicMesh3* InMesh, const FTessellationPattern* Pattern, const bool bUseParallel, FDynamicMesh3* OutMesh)
+	bool Tessellate(const FDynamicMesh3* InMesh, const FTessellationPattern* Pattern, const bool bUseParallel, FProgressCancel* Progress, FDynamicMesh3* OutMesh)
 	{
 		OutMesh->Clear();
 
 		FCompactMaps CompactInfo;
 
 		AdaptiveTessellateLocals::FTessellationData TessData(InMesh);
-		TessellateGeometry(InMesh, Pattern, bUseParallel, CompactInfo, TessData, OutMesh);
+		if (TessellateGeometry(InMesh, Pattern, bUseParallel, Progress, CompactInfo, TessData, OutMesh) == false) 
+		{
+			return false;
+		}
 
 		if (InMesh->HasAttributes()) 
 		{	
@@ -1619,7 +1733,10 @@ namespace AdaptiveTessellateLocals
 				OutAttributes->SetNumNormalLayers(InAttributes->NumNormalLayers());
 				for (int Idx = 0; Idx < InAttributes->NumNormalLayers(); ++Idx) 
 				{	
-					TessellateOverlay(InMesh, InAttributes->GetNormalLayer(Idx), Pattern, TessData, bUseParallel, CompactInfo,  OutAttributes->GetNormalLayer(Idx));
+					if (TessellateOverlay(InMesh, InAttributes->GetNormalLayer(Idx), Pattern, TessData, bUseParallel, Progress, CompactInfo,  OutAttributes->GetNormalLayer(Idx)) == false) 
+					{
+						return false;
+					}
 				} 
 			}
 
@@ -1628,14 +1745,20 @@ namespace AdaptiveTessellateLocals
 				OutAttributes->SetNumUVLayers(InAttributes->NumUVLayers());
 				for (int Idx = 0; Idx < InAttributes->NumUVLayers(); ++Idx) 
 				{	
-					TessellateOverlay(InMesh, InAttributes->GetUVLayer(Idx), Pattern, TessData, bUseParallel, CompactInfo, OutAttributes->GetUVLayer(Idx));
+					if (TessellateOverlay(InMesh, InAttributes->GetUVLayer(Idx), Pattern, TessData, bUseParallel, Progress, CompactInfo, OutAttributes->GetUVLayer(Idx)) == false)
+					{
+						return false;
+					}
 				}
 			}
 
 			if (InAttributes->HasPrimaryColors()) 
 			{
 				OutAttributes->EnablePrimaryColors();
-				TessellateOverlay(InMesh, InAttributes->PrimaryColors(), Pattern, TessData, bUseParallel, CompactInfo, OutAttributes->PrimaryColors());
+				if (TessellateOverlay(InMesh, InAttributes->PrimaryColors(), Pattern, TessData, bUseParallel, Progress, CompactInfo, OutAttributes->PrimaryColors()) == false) 
+				{
+					return false;
+				}
 			}
 
 			if (InAttributes->HasMaterialID())	 
@@ -1643,8 +1766,14 @@ namespace AdaptiveTessellateLocals
 				OutAttributes->EnableMaterialID();
 				OutAttributes->GetMaterialID()->SetName(InAttributes->GetMaterialID()->GetName());
 				ConstructTriangleAttribute(InMesh, InAttributes->GetMaterialID(), TessData, bUseParallel, CompactInfo, OutAttributes->GetMaterialID());
+			
+				if (Progress && Progress->Cancelled()) 
+				{
+					return false;
+				}
 			}
 
+			
 			if (InAttributes->NumPolygroupLayers())  
 			{
 				OutAttributes->SetNumPolygroupLayers(InAttributes->NumPolygroupLayers());
@@ -1653,6 +1782,11 @@ namespace AdaptiveTessellateLocals
 					OutAttributes->GetPolygroupLayer(Idx)->SetName(InAttributes->GetPolygroupLayer(Idx)->GetName());
 					ConstructTriangleAttribute(InMesh, InAttributes->GetPolygroupLayer(Idx), TessData, bUseParallel, CompactInfo, OutAttributes->GetPolygroupLayer(Idx));
 				}
+
+				if (Progress && Progress->Cancelled()) 
+				{
+					return false;
+				}
 			}
 
 			for (const TTuple<FName, TUniquePtr<FDynamicMeshVertexSkinWeightsAttribute>>& AttributeInfo : InAttributes->GetSkinWeightsAttributes())
@@ -1663,6 +1797,11 @@ namespace AdaptiveTessellateLocals
 				OutAttributes->AttachSkinWeightsAttribute(AttributeInfo.Key, SkinAttribute);
 			}
 
+			if (Progress && Progress->Cancelled()) 
+			{
+				return false;
+			}
+
 			if (InAttributes->NumWeightLayers() > 0)
 			{
 				OutAttributes->SetNumWeightLayers(InAttributes->NumWeightLayers());
@@ -1673,6 +1812,11 @@ namespace AdaptiveTessellateLocals
 					OutAttributes->GetWeightLayer(Idx)->SetName(InAttributes->GetWeightLayer(Idx)->GetName());
 				}
 			}
+
+			if (Progress && Progress->Cancelled()) 
+			{
+				return false;
+			}
 		}
 		
 		if (InMesh->HasVertexNormals()) 
@@ -1681,17 +1825,34 @@ namespace AdaptiveTessellateLocals
 			ConstructPerVertexNormals(InMesh, TessData, bUseParallel, CompactInfo, OutMesh);
 		}
 		
+		if (Progress && Progress->Cancelled()) 
+		{
+			return false;
+		}
+
 		if (InMesh->HasVertexUVs()) 
 		{	
 			OutMesh->EnableVertexUVs(FVector2f::Zero());
 			ConstructPerVertexUVs(InMesh, TessData, bUseParallel, CompactInfo, OutMesh);
 		}
 
+		if (Progress && Progress->Cancelled()) 
+		{
+			return false;
+		}
+
 		if (InMesh->HasVertexColors()) 
 		{	
 			OutMesh->EnableVertexColors(FVector4f::Zero());
 			ConstructPerVertexColors(InMesh, TessData, bUseParallel, CompactInfo, OutMesh);
 		}
+
+		if (Progress && Progress->Cancelled()) 
+		{
+			return false;
+		}
+
+		return true;
 	}
 }
 
@@ -1996,9 +2157,9 @@ bool FAdaptiveTessellate::Compute()
 		Mesh = ResultMeshCopy.Get();
 	}
 
-	AdaptiveTessellateLocals::Tessellate(Mesh, Pattern, bUseParallel, ResultMesh);
+	bool bTessResult = AdaptiveTessellateLocals::Tessellate(Mesh, Pattern, bUseParallel, Progress, ResultMesh);
 
-	if (Cancelled()) 
+	if (Cancelled() || bTessResult == false)
 	{
 		if (bInPlace) 
 		{ 
