Updated several modules to have UnsafeTypeCastWarningLevel be an error so it matches other modules. This helps compile times because there will be less shared PCHs created.

diff --git a/Engine/Plugins/Runtime/GeometryProcessing/Source/DynamicMesh/DynamicMesh.Build.cs b/Engine/Plugins/Runtime/GeometryProcessing/Source/DynamicMesh/DynamicMesh.Build.cs
--- a/Engine/Plugins/Runtime/GeometryProcessing/Source/DynamicMesh/DynamicMesh.Build.cs
+++ b/Engine/Plugins/Runtime/GeometryProcessing/Source/DynamicMesh/DynamicMesh.Build.cs
@@ -25,6 +25,6 @@ public class DynamicMesh : ModuleRules
 				"Eigen"
 			}
 		);
-		UnsafeTypeCastWarningLevel = WarningLevel.Warning;
+		UnsafeTypeCastWarningLevel = WarningLevel.Error;
 	}
 }
diff --git a/Engine/Plugins/Runtime/GeometryProcessing/Source/GeometryAlgorithms/GeometryAlgorithms.Build.cs b/Engine/Plugins/Runtime/GeometryProcessing/Source/GeometryAlgorithms/GeometryAlgorithms.Build.cs
--- a/Engine/Plugins/Runtime/GeometryProcessing/Source/GeometryAlgorithms/GeometryAlgorithms.Build.cs
+++ b/Engine/Plugins/Runtime/GeometryProcessing/Source/GeometryAlgorithms/GeometryAlgorithms.Build.cs
@@ -18,6 +18,6 @@ public class GeometryAlgorithms : ModuleRules
 				"GeometryCore"
 			}
 			);
-		UnsafeTypeCastWarningLevel = WarningLevel.Warning;
+		UnsafeTypeCastWarningLevel = WarningLevel.Error;
 	}
 }
diff --git a/Engine/Source/Runtime/Experimental/ChaosCore/ChaosCore.Build.cs b/Engine/Source/Runtime/Experimental/ChaosCore/ChaosCore.Build.cs
--- a/Engine/Source/Runtime/Experimental/ChaosCore/ChaosCore.Build.cs
+++ b/Engine/Source/Runtime/Experimental/ChaosCore/ChaosCore.Build.cs
@@ -34,7 +34,7 @@ namespace UnrealBuildTool.Rules
 			{
 				PublicDefinitions.Add("CHAOS_DEBUG_NAME=0");
 			}
-			UnsafeTypeCastWarningLevel = WarningLevel.Warning;
+			UnsafeTypeCastWarningLevel = WarningLevel.Error;
         }
     }
 }
diff --git a/Engine/Source/Runtime/GeometryCore/GeometryCore.Build.cs b/Engine/Source/Runtime/GeometryCore/GeometryCore.Build.cs
--- a/Engine/Source/Runtime/GeometryCore/GeometryCore.Build.cs
+++ b/Engine/Source/Runtime/GeometryCore/GeometryCore.Build.cs
@@ -23,6 +23,6 @@ public class GeometryCore : ModuleRules
 			}
 		);
 
-		UnsafeTypeCastWarningLevel = WarningLevel.Warning;
+		UnsafeTypeCastWarningLevel = WarningLevel.Error;
 	}
 }
diff --git a/Engine/Source/Runtime/Json/Json.Build.cs b/Engine/Source/Runtime/Json/Json.Build.cs
--- a/Engine/Source/Runtime/Json/Json.Build.cs
+++ b/Engine/Source/Runtime/Json/Json.Build.cs
@@ -19,7 +19,7 @@ namespace UnrealBuildTool.Rules
 				}
 			);
 
-			UnsafeTypeCastWarningLevel = WarningLevel.Warning;
+			UnsafeTypeCastWarningLevel = WarningLevel.Error;
 		}
 	}
 }
