Fixed a bug causing an error when tessellating non-compact meshes.

#rnx

diff --git a/Engine/Plugins/Runtime/GeometryProcessing/Source/DynamicMesh/Private/Operations/AdaptiveTessellate.cpp b/Engine/Plugins/Runtime/GeometryProcessing/Source/DynamicMesh/Private/Operations/AdaptiveTessellate.cpp
--- a/Engine/Plugins/Runtime/GeometryProcessing/Source/DynamicMesh/Private/Operations/AdaptiveTessellate.cpp
+++ b/Engine/Plugins/Runtime/GeometryProcessing/Source/DynamicMesh/Private/Operations/AdaptiveTessellate.cpp
@@ -1432,29 +1432,32 @@ namespace AdaptiveTessellateLocals
 		// Set the overlay triangles and point them to the correct element ids
 		for (const int TriangleID : Mesh->TriangleIndicesItr())
 		{
-			if (OverlayTessData.TrianglesToTessellate.Contains(TriangleID) == false)
+			if (Overlay->IsSetTriangle(TriangleID) && OverlayTessData.TrianglesToTessellate.Contains(TriangleID) == false)
 			{
 				const FIndex3i ElementTri = Overlay->GetTriangle(TriangleID);
 				const int ToTID = CompactInfo.GetTriangleMapping(TriangleID);
-				ResultOverlay->SetTriangle(ToTID, ElementTri);
+				ResultOverlay->SetTriangle(ToTID, FIndex3i(MapE[ElementTri.A], MapE[ElementTri.B], MapE[ElementTri.C]));
 			}
 		}
 
 		for (const int TriangleID : OverlayTessData.TrianglesToTessellate)
 		{
-			TArrayView<FIndex3i> TrianglesBlock = OverlayTessData.MapTrianglesBufferBlock(TriangleID);
-			TArrayView<int> TriangleIDBlock = MeshTessData.MapTriangleIDBufferBlock(TriangleID);
-			
-			checkSlow(TrianglesBlock.Num() == TriangleIDBlock.Num());
-			
-			for (int Index = 0; Index < TrianglesBlock.Num(); ++Index)
-			{
-				const FIndex3i Tri = TrianglesBlock[Index];
-			
-				const int NewTriangleID = TriangleIDBlock[Index];
-				const int ToTID = CompactInfo.GetTriangleMapping(NewTriangleID);
+			if (ensure(Overlay->IsSetTriangle(TriangleID))) // TrianglesToTessellate set should already contain only triangle IDs set in the overlay
+			{ 
+				TArrayView<FIndex3i> TrianglesBlock = OverlayTessData.MapTrianglesBufferBlock(TriangleID);
+				TArrayView<int> TriangleIDBlock = MeshTessData.MapTriangleIDBufferBlock(TriangleID);
+
+				checkSlow(TrianglesBlock.Num() == TriangleIDBlock.Num());
+
+				for (int Index = 0; Index < TrianglesBlock.Num(); ++Index)
+				{
+					const FIndex3i Tri = TrianglesBlock[Index];
+
+					const int NewTriangleID = TriangleIDBlock[Index];
+					const int ToTID = CompactInfo.GetTriangleMapping(NewTriangleID);
 
-				ResultOverlay->SetTriangle(ToTID, FIndex3i(MapE[Tri.A], MapE[Tri.B], MapE[Tri.C]));
+					ResultOverlay->SetTriangle(ToTID, FIndex3i(MapE[Tri.A], MapE[Tri.B], MapE[Tri.C]));
+				}
 			}
 		}
 
@@ -1922,10 +1925,10 @@ TUniquePtr<FTessellationPattern> FAdaptiveTessellate::CreateConcentricRingsTesse
 												   											   const int InTessellationLevel)
 {
 	TArray<int> InEdgeTessLevels;
-	InEdgeTessLevels.Init(InTessellationLevel, InMesh->EdgeCount());
+	InEdgeTessLevels.Init(InTessellationLevel, InMesh->MaxEdgeID());
 	
 	TArray<int> InInnerTessLevels;
-	InInnerTessLevels.Init(InTessellationLevel, InMesh->TriangleCount());
+	InInnerTessLevels.Init(InTessellationLevel, InMesh->MaxTriangleID());
 	
 	return FAdaptiveTessellate::CreateConcentricRingsTessellationPattern(InMesh, InEdgeTessLevels, InInnerTessLevels);
 }
