Renaming AdaptiveTessellate to SelectiveTessellate to better reflect its actual functionality.

#rnx

diff --git a/Engine/Plugins/Experimental/GeometryScripting/Source/GeometryScriptingCore/Private/MeshSubdivideFunctions.cpp b/Engine/Plugins/Experimental/GeometryScripting/Source/GeometryScriptingCore/Private/MeshSubdivideFunctions.cpp
--- a/Engine/Plugins/Experimental/GeometryScripting/Source/GeometryScriptingCore/Private/MeshSubdivideFunctions.cpp
+++ b/Engine/Plugins/Experimental/GeometryScripting/Source/GeometryScriptingCore/Private/MeshSubdivideFunctions.cpp
@@ -129,12 +129,12 @@ UDynamicMesh* UGeometryScriptLibrary_MeshSubdivideFunctions::ApplyUniformTessell
 }
 
 
-UDynamicMesh* UGeometryScriptLibrary_MeshSubdivideFunctions::ApplyAdaptiveTessellation(
+UDynamicMesh* UGeometryScriptLibrary_MeshSubdivideFunctions::ApplySelectiveTessellation(
 	UDynamicMesh* TargetMesh,
 	FGeometryScriptMeshSelection Selection,
-	FGeometryScriptAdaptiveTessellateOptions Options,
+	FGeometryScriptSelectiveTessellateOptions Options,
 	int TessellationLevel,
-	EAdaptiveTessellatePatternType PatternType,
+	ESelectiveTessellatePatternType PatternType,
 	UGeometryScriptDebug* Debug)
 {
 	if (TargetMesh == nullptr)
@@ -146,7 +146,7 @@ UDynamicMesh* UGeometryScriptLibrary_MeshSubdivideFunctions::ApplyAdaptiveTessel
 	{
 		return TargetMesh;
 	}
-	if (PatternType != EAdaptiveTessellatePatternType::ConcentricRings)
+	if (PatternType != ESelectiveTessellatePatternType::ConcentricRings)
 	{
 		UE::Geometry::AppendError(Debug, EGeometryScriptErrorType::InvalidInputs, LOCTEXT("ApplyAdapativeTessellation_InvalidPatternTypeNew", "ApplyAdapativeTessellation: Only ConcentricRings pattern is currently supported"));
 		return TargetMesh;
@@ -159,7 +159,7 @@ UDynamicMesh* UGeometryScriptLibrary_MeshSubdivideFunctions::ApplyAdaptiveTessel
 	TargetMesh->EditMesh([&](FDynamicMesh3& EditMesh) 
 	{	
 		FDynamicMesh3 TessellatedMesh;
-		FAdaptiveTessellate Tessellator(&EditMesh, &TessellatedMesh);
+		FSelectiveTessellate Tessellator(&EditMesh, &TessellatedMesh);
 
 		TArray<int32> Triangles;
 		Selection.ConvertToMeshIndexArray(EditMesh, Triangles, EGeometryScriptIndexType::Triangle);
@@ -169,11 +169,11 @@ UDynamicMesh* UGeometryScriptLibrary_MeshSubdivideFunctions::ApplyAdaptiveTessel
 		}
 
 		TUniquePtr<FTessellationPattern> Pattern; 
-		if (PatternType == EAdaptiveTessellatePatternType::ConcentricRings)
+		if (PatternType == ESelectiveTessellatePatternType::ConcentricRings)
 		{
 			Pattern = (Triangles.Num() > 0) ?
-				FAdaptiveTessellate::CreateConcentricRingsTessellationPattern(&EditMesh, TessellationLevel, Triangles)
-				: FAdaptiveTessellate::CreateConcentricRingsTessellationPattern(&EditMesh, TessellationLevel);
+				FSelectiveTessellate::CreateConcentricRingsTessellationPattern(&EditMesh, TessellationLevel, Triangles)
+				: FSelectiveTessellate::CreateConcentricRingsTessellationPattern(&EditMesh, TessellationLevel);
 		}
 		else
 		{
diff --git a/Engine/Plugins/Experimental/GeometryScripting/Source/GeometryScriptingCore/Public/GeometryScript/MeshSubdivideFunctions.h b/Engine/Plugins/Experimental/GeometryScripting/Source/GeometryScriptingCore/Public/GeometryScript/MeshSubdivideFunctions.h
--- a/Engine/Plugins/Experimental/GeometryScripting/Source/GeometryScriptingCore/Public/GeometryScript/MeshSubdivideFunctions.h
+++ b/Engine/Plugins/Experimental/GeometryScripting/Source/GeometryScriptingCore/Public/GeometryScript/MeshSubdivideFunctions.h
@@ -24,17 +24,17 @@ public:
 
 
 //
-// Adaptive Tessellate options
+// Selective Tessellate options
 //
 UENUM(BlueprintType)
-enum class EAdaptiveTessellatePatternType : uint8
+enum class ESelectiveTessellatePatternType : uint8
 {
 	ConcentricRings = 0
 };
 
 
-USTRUCT(BlueprintType, meta = (DisplayName = "Adaptive Tessellate Options"))
-struct GEOMETRYSCRIPTINGCORE_API FGeometryScriptAdaptiveTessellateOptions
+USTRUCT(BlueprintType, meta = (DisplayName = "Selective Tessellate Options"))
+struct GEOMETRYSCRIPTINGCORE_API FGeometryScriptSelectiveTessellateOptions
 {
 	GENERATED_BODY()
 public:
@@ -79,12 +79,12 @@ public:
 
 	UFUNCTION(BlueprintCallable, Category = "GeometryScript|Subdivide", meta=(ScriptMethod))
 	static UPARAM(DisplayName = "Target Mesh") UDynamicMesh* 
-	ApplyAdaptiveTessellation(
+	ApplySelectiveTessellation(
 		UDynamicMesh* TargetMesh,
 		FGeometryScriptMeshSelection Selection,
-		FGeometryScriptAdaptiveTessellateOptions Options,
+		FGeometryScriptSelectiveTessellateOptions Options,
 		int TessellationLevel = 1,
-		EAdaptiveTessellatePatternType PatternType = EAdaptiveTessellatePatternType::ConcentricRings,
+		ESelectiveTessellatePatternType PatternType = ESelectiveTessellatePatternType::ConcentricRings,
 		UGeometryScriptDebug* Debug = nullptr);
 
 };
\ No newline at end of file
diff --git a/Engine/Plugins/Experimental/MeshModelingToolsetExp/Source/MeshModelingToolsExp/Private/DisplaceMeshTool.cpp b/Engine/Plugins/Experimental/MeshModelingToolsetExp/Source/MeshModelingToolsExp/Private/DisplaceMeshTool.cpp
--- a/Engine/Plugins/Experimental/MeshModelingToolsetExp/Source/MeshModelingToolsExp/Private/DisplaceMeshTool.cpp
+++ b/Engine/Plugins/Experimental/MeshModelingToolsetExp/Source/MeshModelingToolsExp/Private/DisplaceMeshTool.cpp
@@ -375,10 +375,10 @@ namespace DisplaceMeshToolLocals{
 			{
 				if (Parameters.ActiveMaterialID.IsSet())
 				{
-					TUniquePtr<FTessellationPattern> Pattern = FAdaptiveTessellate::CreateConcentricRingsPatternFromMaterial(ResultMesh.Get(), SubdivisionsCount, Parameters.ActiveMaterialID.GetValue());;
+					TUniquePtr<FTessellationPattern> Pattern = FSelectiveTessellate::CreateConcentricRingsPatternFromMaterial(ResultMesh.Get(), SubdivisionsCount, Parameters.ActiveMaterialID.GetValue());;
 					
 					FDynamicMesh3 OutMesh;
-					FAdaptiveTessellate Tessellator(ResultMesh.Get(), &OutMesh);
+					FSelectiveTessellate Tessellator(ResultMesh.Get(), &OutMesh);
 					Tessellator.Progress = ProgressCancel;
 					Tessellator.SetPattern(Pattern.Get());
 
diff --git a/Engine/Plugins/Runtime/GeometryProcessing/Source/DynamicMesh/Private/Operations/AdaptiveTessellate.cpp b/Engine/Plugins/Runtime/GeometryProcessing/Source/DynamicMesh/Private/Operations/AdaptiveTessellate.cpp
--- a/Engine/Plugins/Runtime/GeometryProcessing/Source/DynamicMesh/Private/Operations/AdaptiveTessellate.cpp
+++ b/Engine/Plugins/Runtime/GeometryProcessing/Source/DynamicMesh/Private/Operations/AdaptiveTessellate.cpp
@@ -13,7 +13,7 @@
 
 using namespace UE::Geometry;
 
-namespace AdaptiveTessellateLocals 
+namespace SelectiveTessellateLocals 
 {	
 	// Forward declare
 	class FTessellationData;
@@ -26,7 +26,7 @@ namespace AdaptiveTessellateLocals
 								  FTessellationData& TessData,
 								  FCompactMaps& CompactInfo,
 								  FDynamicMesh3* ResultMesh,
-								  FAdaptiveTessellate::FTessellationInformation& TessInfo);
+								  FSelectiveTessellate::FTessellationInformation& TessInfo);
 
 	template<typename RealType, int ElementSize> 
 	bool ConstructTessellatedOverlay(const FDynamicMesh3* Mesh, 
@@ -861,7 +861,7 @@ namespace AdaptiveTessellateLocals
 							FCompactMaps& OutCompactInfo, 
 							FTessellationData& TessData,
 							FDynamicMesh3* OutMesh,
-							FAdaptiveTessellate::FTessellationInformation& TessInfo) 
+							FSelectiveTessellate::FTessellationInformation& TessInfo) 
 	{	
 
 		TessData.Init(Pattern);
@@ -953,12 +953,12 @@ namespace AdaptiveTessellateLocals
 			return false;
 		}
 
-		if (AdaptiveTessellateLocals::ConstructTessellatedMesh(Mesh, 
-														       Progress, 
-															   TessData, 
-															   OutCompactInfo, 
-															   OutMesh, 
-															   TessInfo) == false)
+		if (SelectiveTessellateLocals::ConstructTessellatedMesh(Mesh, 
+														        Progress, 
+															    TessData, 
+															    OutCompactInfo, 
+															    OutMesh, 
+															    TessInfo) == false)
 		{
 			return false;
 		}
@@ -985,7 +985,7 @@ namespace AdaptiveTessellateLocals
 		
 		//TODO: The overlay tessellation data should only compute and contain the connectivity information since we 
 	 	// can simply  reuse coordinates from the geometry tessellation (i.e. overlay elements live on the vertices). 
-		AdaptiveTessellateLocals::FOverlayTessellationData OverlayTessData(Mesh, Overlay);
+		SelectiveTessellateLocals::FOverlayTessellationData OverlayTessData(Mesh, Overlay);
 		OverlayTessData.Init(Pattern);
 
 		// Tessellate edge patches
@@ -1088,13 +1088,13 @@ namespace AdaptiveTessellateLocals
 			return false;
 		}
 
-		if (AdaptiveTessellateLocals::ConstructTessellatedOverlay(Mesh, 
-															  	  Overlay,
-															  	  MeshTessData, 
-															  	  OverlayTessData, 
-															  	  CompactInfo, 
-																  Progress, 
-															  	  OutOverlay) == false) 
+		if (SelectiveTessellateLocals::ConstructTessellatedOverlay(Mesh, 
+															  	   Overlay,
+															  	   MeshTessData, 
+															  	   OverlayTessData, 
+															  	   CompactInfo, 
+																   Progress, 
+															  	   OutOverlay) == false) 
 		{
 			return false;
 		}
@@ -1150,7 +1150,7 @@ namespace AdaptiveTessellateLocals
 								  FTessellationData& TessData,
 								  FCompactMaps& CompactInfo,
 								  FDynamicMesh3* ResultMesh,
-								  FAdaptiveTessellate::FTessellationInformation& TessInfo)
+								  FSelectiveTessellate::FTessellationInformation& TessInfo)
 	{
 		ResultMesh->Clear();
 
@@ -1758,13 +1758,13 @@ namespace AdaptiveTessellateLocals
 					const bool bUseParallel, 
 					FProgressCancel* Progress, 
 					FDynamicMesh3* OutMesh,
-					FAdaptiveTessellate::FTessellationInformation& TessInfo)
+					FSelectiveTessellate::FTessellationInformation& TessInfo)
 	{
 		OutMesh->Clear();
 
 		FCompactMaps CompactInfo;
 
-		AdaptiveTessellateLocals::FTessellationData TessData(InMesh);
+		SelectiveTessellateLocals::FTessellationData TessData(InMesh);
 		if (TessellateGeometry(InMesh, Pattern, bUseParallel, Progress, CompactInfo, TessData, OutMesh, TessInfo) == false) 
 		{
 			return false;
@@ -1910,19 +1910,19 @@ namespace AdaptiveTessellateLocals
 // GLSL pattern
 //
 
-TUniquePtr<FTessellationPattern> FAdaptiveTessellate::CreateConcentricRingsTessellationPattern(const FDynamicMesh3* InMesh,	
-												   											   const TArray<int>& InEdgeTessLevels,
-                                                   											   const TArray<int>& InInnerTessLevels)
+TUniquePtr<FTessellationPattern> FSelectiveTessellate::CreateConcentricRingsTessellationPattern(const FDynamicMesh3* InMesh,	
+												   											    const TArray<int>& InEdgeTessLevels,
+                                                   											    const TArray<int>& InInnerTessLevels)
 {
-	TUniquePtr<AdaptiveTessellateLocals::FConcentricRingsTessellationPattern> Pattern = MakeUnique<AdaptiveTessellateLocals::FConcentricRingsTessellationPattern>(InMesh);
+	TUniquePtr<SelectiveTessellateLocals::FConcentricRingsTessellationPattern> Pattern = MakeUnique<SelectiveTessellateLocals::FConcentricRingsTessellationPattern>(InMesh);
 	Pattern->Init(InEdgeTessLevels, InInnerTessLevels);
 	
 	return Pattern;
 }
 
 
-TUniquePtr<FTessellationPattern> FAdaptiveTessellate::CreateConcentricRingsTessellationPattern(const FDynamicMesh3* InMesh,
-												   											   const int InTessellationLevel)
+TUniquePtr<FTessellationPattern> FSelectiveTessellate::CreateConcentricRingsTessellationPattern(const FDynamicMesh3* InMesh,
+												   											    const int InTessellationLevel)
 {
 	TArray<int> InEdgeTessLevels;
 	InEdgeTessLevels.Init(InTessellationLevel, InMesh->MaxEdgeID());
@@ -1930,13 +1930,13 @@ TUniquePtr<FTessellationPattern> FAdaptiveTessellate::CreateConcentricRingsTesse
 	TArray<int> InInnerTessLevels;
 	InInnerTessLevels.Init(InTessellationLevel, InMesh->MaxTriangleID());
 	
-	return FAdaptiveTessellate::CreateConcentricRingsTessellationPattern(InMesh, InEdgeTessLevels, InInnerTessLevels);
+	return FSelectiveTessellate::CreateConcentricRingsTessellationPattern(InMesh, InEdgeTessLevels, InInnerTessLevels);
 }
 
  
-TUniquePtr<FTessellationPattern> FAdaptiveTessellate::CreateConcentricRingsTessellationPattern(const FDynamicMesh3* InMesh,  
-								   							  					    		   const int InTessellationLevel,
-								   							  					    		   const TArray<int>& InTriangleList) 
+TUniquePtr<FTessellationPattern> FSelectiveTessellate::CreateConcentricRingsTessellationPattern(const FDynamicMesh3* InMesh,  
+								   							  					    		    const int InTessellationLevel,
+								   							  					    		    const TArray<int>& InTriangleList) 
 {
 	if (InTriangleList.IsEmpty()) 
 	{
@@ -1961,12 +1961,12 @@ TUniquePtr<FTessellationPattern> FAdaptiveTessellate::CreateConcentricRingsTesse
 		TriangleTessLevels[TriangleID] = InTessellationLevel; 
 	}
 
-	return FAdaptiveTessellate::CreateConcentricRingsTessellationPattern(InMesh, EdgeTessLevels, TriangleTessLevels);	
+	return FSelectiveTessellate::CreateConcentricRingsTessellationPattern(InMesh, EdgeTessLevels, TriangleTessLevels);	
 }
 
 
-TUniquePtr<FTessellationPattern> FAdaptiveTessellate::CreateConcentricRingsTessellationPattern(const TFunctionRef<int(const int EdgeID)> InEdgeFunc, 
-                                                               					    		   const TFunctionRef<int(const int TriangleID)> InTriFunc) 
+TUniquePtr<FTessellationPattern> FSelectiveTessellate::CreateConcentricRingsTessellationPattern(const TFunctionRef<int(const int EdgeID)> InEdgeFunc, 
+                                                               					    		    const TFunctionRef<int(const int TriangleID)> InTriFunc) 
 {
 
 	// TODO: not implemented yet
@@ -1975,9 +1975,9 @@ TUniquePtr<FTessellationPattern> FAdaptiveTessellate::CreateConcentricRingsTesse
 }
 
 
-TUniquePtr<FTessellationPattern> FAdaptiveTessellate::CreateConcentricRingsPatternFromTriangleGroup(const FDynamicMesh3* InMesh,
-																					 			    const int InTessellationLevel,
-                                                                                     			    const int InPolygroupID) 
+TUniquePtr<FTessellationPattern> FSelectiveTessellate::CreateConcentricRingsPatternFromTriangleGroup(const FDynamicMesh3* InMesh,
+																					 			     const int InTessellationLevel,
+                                                                                     			     const int InPolygroupID) 
 {
 	if (InMesh->HasTriangleGroups() == false)
 	{
@@ -1995,14 +1995,14 @@ TUniquePtr<FTessellationPattern> FAdaptiveTessellate::CreateConcentricRingsPatte
 		}
 	}
 
-	return FAdaptiveTessellate::CreateConcentricRingsTessellationPattern(InMesh, InTessellationLevel, TriangleList);
+	return FSelectiveTessellate::CreateConcentricRingsTessellationPattern(InMesh, InTessellationLevel, TriangleList);
 }
 
 
-TUniquePtr<FTessellationPattern> FAdaptiveTessellate::CreateConcentricRingsPatternFromPolyGroup(const FDynamicMesh3* InMesh,
-                                                                           						const int InTessellationLevel,
-                                                                           						const FString& InLayerName,
-                                                                           						const int InPolygroupID) 
+TUniquePtr<FTessellationPattern> FSelectiveTessellate::CreateConcentricRingsPatternFromPolyGroup(const FDynamicMesh3* InMesh,
+                                                                           						 const int InTessellationLevel,
+                                                                           						 const FString& InLayerName,
+                                                                           						 const int InPolygroupID) 
 {
 	if (InMesh->HasAttributes() == false)
 	{
@@ -2036,13 +2036,13 @@ TUniquePtr<FTessellationPattern> FAdaptiveTessellate::CreateConcentricRingsPatte
 		}
 	}
 
-	return FAdaptiveTessellate::CreateConcentricRingsTessellationPattern(InMesh, InTessellationLevel, TriangleList);
+	return FSelectiveTessellate::CreateConcentricRingsTessellationPattern(InMesh, InTessellationLevel, TriangleList);
 }
 
 
-TUniquePtr<FTessellationPattern> FAdaptiveTessellate::CreateConcentricRingsPatternFromMaterial(const FDynamicMesh3* InMesh,
-																					 		   const int InTessellationLevel,
-                                                                                     		   const int MaterialID) 
+TUniquePtr<FTessellationPattern> FSelectiveTessellate::CreateConcentricRingsPatternFromMaterial(const FDynamicMesh3* InMesh,
+																					 		    const int InTessellationLevel,
+                                                                                     		    const int MaterialID) 
 {
 	if (InMesh->HasAttributes() == false)
 	{
@@ -2065,16 +2065,16 @@ TUniquePtr<FTessellationPattern> FAdaptiveTessellate::CreateConcentricRingsPatte
 		}
 	}
 
-	return FAdaptiveTessellate::CreateConcentricRingsTessellationPattern(InMesh, InTessellationLevel, TriangleList);
+	return FSelectiveTessellate::CreateConcentricRingsTessellationPattern(InMesh, InTessellationLevel, TriangleList);
 }
 
 //
 // Inner uniform pattern
 //
 
-TUniquePtr<FTessellationPattern> FAdaptiveTessellate::CreateInnerUnifromTessellationPattern(const FDynamicMesh3* InMesh,
-													  	   						 		    const TArray<int>& InEdgeTessLevels,
-                                                      	   									const TArray<int>& InInnerTessLevels)
+TUniquePtr<FTessellationPattern> FSelectiveTessellate::CreateInnerUnifromTessellationPattern(const FDynamicMesh3* InMesh,
+													  	   						 		     const TArray<int>& InEdgeTessLevels,
+                                                      	   									 const TArray<int>& InInnerTessLevels)
 {
 	// TODO: not implemented yet
 	checkNoEntry();
@@ -2172,13 +2172,13 @@ void FTessellationPattern::ComputeInnerConcentricTriangle(const FVector3d& V1,
 }
 
 
-bool FAdaptiveTessellate::Cancelled()
+bool FSelectiveTessellate::Cancelled()
 {
 	return (Progress == nullptr) ? false : Progress->Cancelled();
 }
 
 
-bool FAdaptiveTessellate::Compute()
+bool FSelectiveTessellate::Compute()
 {	
 	if (Validate() != EOperationValidationResult::Ok) 
 	{
@@ -2206,7 +2206,7 @@ bool FAdaptiveTessellate::Compute()
 		Mesh = ResultMeshCopy.Get();
 	}
 
-	bool bTessResult = AdaptiveTessellateLocals::Tessellate(Mesh, Pattern, bUseParallel, Progress, ResultMesh, TessInfo);
+	bool bTessResult = SelectiveTessellateLocals::Tessellate(Mesh, Pattern, bUseParallel, Progress, ResultMesh, TessInfo);
 
 	if (Cancelled() || bTessResult == false)
 	{
diff --git a/Engine/Plugins/Runtime/GeometryProcessing/Source/DynamicMesh/Public/Operations/AdaptiveTessellate.h b/Engine/Plugins/Runtime/GeometryProcessing/Source/DynamicMesh/Public/Operations/AdaptiveTessellate.h
--- a/Engine/Plugins/Runtime/GeometryProcessing/Source/DynamicMesh/Public/Operations/AdaptiveTessellate.h
+++ b/Engine/Plugins/Runtime/GeometryProcessing/Source/DynamicMesh/Public/Operations/AdaptiveTessellate.h
@@ -201,7 +201,7 @@ protected:
  * 
  * @note Currently does not support interpolation of the GenericAttributes besides Skin Weights.
  */
-class DYNAMICMESH_API FAdaptiveTessellate
+class DYNAMICMESH_API FSelectiveTessellate
 {
 public:
 
@@ -273,7 +273,7 @@ public:
 	 * be changed.
      * //TODO: Add an option to turn off making the backup copy of the input mesh in case the computation fails.
 	 */
-	FAdaptiveTessellate(FDynamicMesh3* OutMesh) 
+	FSelectiveTessellate(FDynamicMesh3* OutMesh) 
 	:
 	ResultMesh(OutMesh), bInPlace(true)
 	{
@@ -282,13 +282,13 @@ public:
 	/** 
 	 * Tessellate the mesh and write the result into another mesh. This will overwrite any data stored in the OutMesh. 
 	 */
-	FAdaptiveTessellate(const FDynamicMesh3* Mesh, FDynamicMesh3* OutMesh) 
+	FSelectiveTessellate(const FDynamicMesh3* Mesh, FDynamicMesh3* OutMesh) 
 	:
 	ResultMesh(OutMesh), Mesh(Mesh), bInPlace(false)
 	{
 	}
 	
-	virtual ~FAdaptiveTessellate() 
+	virtual ~FSelectiveTessellate() 
 	{
 	}
 
