MeshModelingTools: Give ability to add extra corners in PolyEd for manipulation of some common open-boundary cases (like a disconnected rectangle, for instance). Also show selectable corners by default (but smaller).

diff --git a/Engine/Plugins/Runtime/GeometryProcessing/Source/DynamicMesh/Private/GroupTopology.cpp b/Engine/Plugins/Runtime/GeometryProcessing/Source/DynamicMesh/Private/GroupTopology.cpp
--- a/Engine/Plugins/Runtime/GeometryProcessing/Source/DynamicMesh/Private/GroupTopology.cpp
+++ b/Engine/Plugins/Runtime/GeometryProcessing/Source/DynamicMesh/Private/GroupTopology.cpp
@@ -1,6 +1,8 @@
 // Copyright Epic Games, Inc. All Rights Reserved.
 
 #include "GroupTopology.h"
+
+#include "Algo/MinElement.h"
 #include "MeshRegionBoundaryLoops.h"
 
 using namespace UE::Geometry;
@@ -49,6 +51,8 @@ bool FGroupTopology::RebuildTopology()
 	Edges.Reset();
 	Corners.Reset();
 
+	CurrentExtraCornerVids.Reset();
+
 	int32 MaxGroupID = 0;
 	for (int32 tid : Mesh->TriangleIndicesItr())
 	{
@@ -85,30 +89,15 @@ bool FGroupTopology::RebuildTopology()
 		Groups[GroupIDToGroupIndexMap[GroupID]].Triangles.Add(tid);
 	}
 
-	// precompute junction vertices set
-	CornerVerticesFlags.Init(false, Mesh->MaxVertexID());
-	for (int vid : Mesh->VertexIndicesItr())
-	{
-		if (IsCornerVertex(vid))
-		{
-			CornerVerticesFlags[vid] = true;
-			FCorner Corner = { vid };
-			int NewCornerIndex = Corners.Num();
-			VertexIDToCornerIDMap.Add(vid, NewCornerIndex);
-			Corners.Add(Corner);
-		}
-	}
-	for (FCorner& Corner : Corners)
-	{
-		GetAllVertexGroups(Corner.VertexID, Corner.NeighbourGroupIDs);
-	}
-
+	VertexIDToCornerIDMap.Reset();
+	TMap<int32, int32> GroupEdgeMinEidToGroupEdgeID;
+	TBitArray<> VertCheckedForCorner;
+	VertCheckedForCorner.Init(false, Mesh->MaxVertexID());
 
 	// construct boundary loops
 	for (FGroup& Group : Groups)
 	{
-		// finds FGroupEdges and uses to populate Group.Boundaries
-		bool bOK = ExtractGroupEdges(Group);
+		bool bOK = GenerateBoundaryAndGroupEdges(Group, GroupEdgeMinEidToGroupEdgeID, VertCheckedForCorner);
 		if (!bOK)
 		{
 			return false;
@@ -157,31 +146,58 @@ void FGroupTopology::RetargetOnClonedMesh(const FDynamicMesh3* NewMesh)
 	}
 }
 
-
 bool FGroupTopology::IsCornerVertex(int VertexID) const
+{ 
+	return ShouldVertBeCorner(VertexID);
+}
+
+bool FGroupTopology::ShouldVertBeCorner(int VertexID, bool* bIsAnExtraCornerOut) const
 {
-	FIndex3i UniqueGroups;
-	int UniqueCount = 0;
-	for (int tid : Mesh->VtxTrianglesItr(VertexID))
+	if (bIsAnExtraCornerOut)
 	{
-		int GroupID = GetGroupID(tid);
-		if (UniqueCount == 0)
-		{
-			UniqueGroups[0] = GroupID;
-			UniqueCount++;
-		}
-		else if (UniqueCount == 1 && GroupID != UniqueGroups[0])
+		*bIsAnExtraCornerOut = false;
+	}
+
+	// If the vertex has at least three attached groupedge edges (ie three edges that separate two groups
+	// or are boundary edges), then it has to be a corner. Otherwise, we let the user supply a function that
+	// determines whether the vert should be a corner anyway, passing in the group edges that do pass through the
+	// vertex.
+
+	// only ends up being used if there are less than 3, so that we can pass in the
+	// relevant edges to the user function.
+	FIndex2i AttachedGroupEdgeEids(IndexConstants::InvalidID, IndexConstants::InvalidID);
+
+	for (int32 Eid : Mesh->VtxEdgesItr(VertexID))
+	{
+		FIndex2i EdgeTris = Mesh->GetEdgeT(Eid);
+		if (EdgeTris.B == IndexConstants::InvalidID
+			|| Mesh->GetTriangleGroup(EdgeTris.A) != Mesh->GetTriangleGroup(EdgeTris.B))
 		{
-			UniqueGroups[1] = GroupID;
-			UniqueCount++;
-		}
-		else if (UniqueCount == 2 && GroupID != UniqueGroups[0] && GroupID != UniqueGroups[1])
+			// This is a group edge.
+			if (AttachedGroupEdgeEids[0] == IndexConstants::InvalidID)
+			{
+				AttachedGroupEdgeEids[0] = Eid;
+			}
+			else if (AttachedGroupEdgeEids[1] == IndexConstants::InvalidID)
+			{
+				AttachedGroupEdgeEids[1] = Eid;
+			}
+			else
+			{
+				// This is the third such edge. Has to be corner.
+				return true;
+			}
+		}//end if part of group edge
+	}//end looking through attached edges.
+
+	// If we got to here, then the vert isn't a corner just based on group topology, but can be made
+	// into a corner anyway by user function (for instance, at sharp corners, or at user-defined points).
+	if (ShouldAddExtraCornerAtVert && ShouldAddExtraCornerAtVert(*this, VertexID, AttachedGroupEdgeEids))
+	{
+		if (bIsAnExtraCornerOut)
 		{
-			return true;
+			*bIsAnExtraCornerOut = true;
 		}
-	}
-	if (UniqueCount == 2 && Mesh->IsBoundaryVertex(VertexID))
-	{
 		return true;
 	}
 	return false;
@@ -497,7 +513,8 @@ void FGroupTopology::CollectGroupBoundaryVertices(int GroupID, TSet<int>& Vertic
 
 
 
-bool FGroupTopology::ExtractGroupEdges(FGroup& Group)
+bool FGroupTopology::GenerateBoundaryAndGroupEdges(FGroup& Group, 
+	TMap<int32, int32>& GroupEdgeMinEidToGroupEdgeID, TBitArray<>& VertCheckedForCorner)
 {
 	FMeshRegionBoundaryLoops BdryLoops(Mesh, Group.Triangles, true);
 
@@ -509,6 +526,44 @@ bool FGroupTopology::ExtractGroupEdges(FGroup& Group)
 
 	int NumLoops = BdryLoops.Loops.Num();
 
+	// Function we use to see if a vertex is a corner, creating a corner object if so
+	auto CheckForCornerAndCreateIfNeeded = [this, &GroupEdgeMinEidToGroupEdgeID, &VertCheckedForCorner](int32 Vid)
+	{
+		// See if we have a cached result
+		if (VertCheckedForCorner[Vid])
+		{
+			return VertexIDToCornerIDMap.Contains(Vid);
+		}
+		else
+		{
+			VertCheckedForCorner[Vid] = true;
+			bool bIsExtraCorner = false;
+			if (ShouldVertBeCorner(Vid, &bIsExtraCorner))
+			{
+				// New corner vertex found. Create a corner for this vertex.
+				int32 CornerID = Corners.Emplace();
+				FCorner& Corner = Corners[CornerID];
+				Corner.VertexID = Vid;
+				GetAllVertexGroups(Vid, Corner.NeighbourGroupIDs);
+
+				VertexIDToCornerIDMap.Add(Vid, CornerID);
+
+				if (bIsExtraCorner)
+				{
+					CurrentExtraCornerVids.Add(Vid);
+				}
+				return true;
+			}
+			else
+			{
+				return false;
+			}
+		}
+		return ensure(false); // Shouldn't be able to get here
+	};
+
+	// Go through the boundary, find the corners, and use the intervening edges to create
+	// group edges.
 	Group.Boundaries.SetNum(NumLoops);
 	for ( int li = 0; li < NumLoops; ++li )
 	{
@@ -520,7 +575,7 @@ bool FGroupTopology::ExtractGroupEdges(FGroup& Group)
 		int NumV = Loop.Vertices.Num();
 		for (int i = 0; i < NumV; ++i)
 		{
-			if (CornerVerticesFlags[Loop.Vertices[i]])
+			if (CheckForCornerAndCreateIfNeeded(Loop.Vertices[i]))
 			{
 				CornerIndices.Add(i);
 			}
@@ -529,12 +584,12 @@ bool FGroupTopology::ExtractGroupEdges(FGroup& Group)
 		// if we had no indices then this is like the cap of a cylinder, just one single long edge
 		if ( CornerIndices.Num() == 0 )
 		{ 
-			FIndex2i EdgeID = MakeEdgeID(Loop.Edges[0]);
-			int OtherGroupID = (EdgeID.A == Group.GroupID) ? EdgeID.B : EdgeID.A;
-			int EdgeIndex = FindExistingGroupEdge(Group.GroupID, OtherGroupID, Loop.Vertices[0], Loop.Vertices[1]);
-			if (EdgeIndex == -1)
+			int32 MinEid = *Algo::MinElement(Loop.Edges);
+			int32* ExistingGroupEdgeID = GroupEdgeMinEidToGroupEdgeID.Find(MinEid);
+			int EdgeIndex = ExistingGroupEdgeID ? *ExistingGroupEdgeID : IndexConstants::InvalidID;
+			if (EdgeIndex == IndexConstants::InvalidID)
 			{
-				FGroupEdge Edge = { EdgeID };
+				FGroupEdge Edge = { MakeEdgeGroupsPair(Loop.Edges[0]) };
 				Edge.Span = FEdgeSpan(Mesh);
 				Edge.Span.InitializeFromEdges(Loop.Edges);
 				Edge.EndpointCorners = FIndex2i::Invalid();
@@ -552,27 +607,37 @@ bool FGroupTopology::ExtractGroupEdges(FGroup& Group)
 		// add each span
 		for (int k = 0; k < NumSpans; ++k)
 		{
-			int i0 = CornerIndices[k];
+			int32 StartIndex = CornerIndices[k];
+			int32 EndIndex = CornerIndices[k + 1];		// note: StartIndex == EndIndex on a closed loop, ie NumSpans == 1
+
+			int32 NumSpanEdges = (EndIndex + NumV - StartIndex) % NumV;
 
-			FIndex2i EdgeID = MakeEdgeID(Loop.Edges[i0]);
-			int OtherGroupID = (EdgeID.A == Group.GroupID) ? EdgeID.B : EdgeID.A;
-			int EdgeIndex = FindExistingGroupEdge(Group.GroupID, OtherGroupID, Loop.Vertices[i0], Loop.Vertices[(i0+1)%Loop.Vertices.Num()]);
-			if (EdgeIndex != -1)
+			// Find the min eid in this subsection so that we can see if we already have a group edge for it.
+			int32 MinEid = Loop.Edges[StartIndex];
+			for (int32 i = 1; i < NumSpanEdges; ++i)
+			{
+				int32 Index = (StartIndex + i) % NumV;
+				MinEid = FMath::Min(MinEid, Loop.Edges[Index]);
+			}
+
+			int32* ExistingGroupEdgeID = GroupEdgeMinEidToGroupEdgeID.Find(MinEid);
+			int EdgeIndex = ExistingGroupEdgeID ? *ExistingGroupEdgeID : IndexConstants::InvalidID;
+			if (EdgeIndex != IndexConstants::InvalidID)
 			{
 				FGroupEdge& Existing = Edges[EdgeIndex];
 				Boundary.GroupEdges.Add(EdgeIndex);
 				continue;
 			}
 
-			FGroupEdge Edge = { EdgeID };
+			FGroupEdge Edge = { MakeEdgeGroupsPair(Loop.Edges[StartIndex]) };
 
-			int i1 = CornerIndices[k+1];		// note: i1 == i0 on a closed loop, ie NumSpans == 1
 			TArray<int> SpanVertices;
-			do {
-				SpanVertices.Add(Loop.Vertices[i0]);
-				i0 = (i0 + 1) % NumV;
-			} while (i0 != i1);
-			SpanVertices.Add(Loop.Vertices[i0]);	// add last vertex
+			int32 NumSpanVertices = NumSpanEdges + 1;
+			for (int32 i = 0; i < NumSpanVertices; ++i)
+			{
+				int32 Index = (StartIndex + i) % NumV;
+				SpanVertices.Add(Loop.Vertices[Index]);
+			}
 
 			Edge.Span = FEdgeSpan(Mesh);
 			Edge.Span.InitializeFromVertices(SpanVertices);
@@ -580,6 +645,7 @@ bool FGroupTopology::ExtractGroupEdges(FGroup& Group)
 			check(Edge.EndpointCorners.A != IndexConstants::InvalidID && Edge.EndpointCorners.B != IndexConstants::InvalidID);
 			EdgeIndex = Edges.Add(Edge);
 			Boundary.GroupEdges.Add(EdgeIndex);
+			GroupEdgeMinEidToGroupEdgeID.Add(MinEid, EdgeIndex);
 		}
 	}
 
@@ -587,59 +653,21 @@ bool FGroupTopology::ExtractGroupEdges(FGroup& Group)
 }
 
 
+bool UE::Geometry::FGroupTopology::RebuildTopologyWithSpecificExtraCorners(const TSet<int32>& ExtraCornerVids)
+{
+	// We temporarily replace the vert corner-forcing function with one that checks against the given vids,
+	// then we revert after doing the build.
+	TFunction<bool(const FGroupTopology& GroupTopology, int32 Vid, const FIndex2i& AttachedGroupEdgeEids)> OldCornerFunction = MoveTemp(ShouldAddExtraCornerAtVert);
 
+	ShouldAddExtraCornerAtVert = [&ExtraCornerVids](const FGroupTopology&, int32 Vid, const FIndex2i&) { return ExtraCornerVids.Contains(Vid); };
 
-int FGroupTopology::FindExistingGroupEdge(int GroupID, int OtherGroupID, int FirstVertexID, int SecondVertexID)
-{
-	// if this is a boundary edge, we cannot have created it already
-	if (OtherGroupID < 0)
-	{
-		return -1;
-	}
+	bool bSuccess = RebuildTopology();
 
-	FGroup& OtherGroup = Groups[GroupIDToGroupIndexMap[OtherGroupID]];
-	FIndex2i EdgeID = MakeEdgeID(GroupID, OtherGroupID);
-	
-	for (FGroupBoundary& Boundary : OtherGroup.Boundaries)
-	{
-		for (int EdgeIndex : Boundary.GroupEdges)
-		{
-			if (Edges[EdgeIndex].Groups == EdgeID)
-			{
-				// Same EdgeID pair may occur multiple times in the same boundary loop
-				// (think of a cube with its side faces joined together on opposite corners).
-				// For non-loop edges, it is sufficient to check that one of the endpoints is the
-				// same vertex to know that the edges are the same.
-				TArray<int>& Vertices = Edges[EdgeIndex].Span.Vertices;
-				int32 NumVerts = Vertices.Num();
-				if (Edges[EdgeIndex].EndpointCorners.A != IndexConstants::InvalidID)
-				{
-					if (Vertices[0] == FirstVertexID || Vertices[NumVerts - 1] == FirstVertexID)
-					{
-						return EdgeIndex;
-					}
-				}
-				else
-				{
-					// For loop edges we're not guaranteed to have the loop start on any particular
-					// vertex. We have to make sure that the two loops share at least two adjacent 
-					// vertices, because of pathological cases with bowtie-shaped groups.
-					int32 FirstVertIndex = Vertices.IndexOfByKey(FirstVertexID);
-					if (FirstVertIndex != INDEX_NONE 
-						&& (Vertices[(FirstVertIndex + 1) % NumVerts] == SecondVertexID
-							|| Vertices[(FirstVertIndex + NumVerts - 1) % NumVerts] == SecondVertexID))
-					{
-						return EdgeIndex;
-					}
-				}
-			}//end if group pair matched
-		}
-	}//end looking through other group boundaries
-	return -1;
+	ShouldAddExtraCornerAtVert = MoveTemp(OldCornerFunction);
+	return bSuccess;
 }
 
 
-
 bool FGroupTopology::GetGroupEdgeTangent(int GroupEdgeID, FVector3d& TangentOut) const
 {
 	check(GroupEdgeID >= 0 && GroupEdgeID < Edges.Num());
@@ -881,10 +909,8 @@ bool FTriangleGroupTopology::RebuildTopology()
 		}
 	}
 
-	CornerVerticesFlags.Init(true, Mesh->MaxVertexID());
 	for (int vid : Mesh->VertexIndicesItr())
 	{
-		CornerVerticesFlags[vid] = true;
 		FCorner Corner = { vid };
 		int NewCornerIndex = Corners.Num();
 		VertexIDToCornerIDMap.Add(vid, NewCornerIndex);
@@ -912,7 +938,7 @@ bool FTriangleGroupTopology::RebuildTopology()
 			int& GroupEdgeIndex = MeshEdgeToGroupEdge[TriEdges[j]];
 			if (GroupEdgeIndex == INDEX_NONE)
 			{
-				FGroupEdge NewGroupEdge = { MakeEdgeID(TriEdges[j]) };
+				FGroupEdge NewGroupEdge = { MakeEdgeGroupsPair(TriEdges[j]) };
 				FIndex2i EdgeVerts = Mesh->GetEdgeV(TriEdges[j]);
 				NewGroupEdge.Span = FEdgeSpan(Mesh);
 				SpanVertices[0] = EdgeVerts.A; SpanVertices[1] = EdgeVerts.B;
diff --git a/Engine/Plugins/Runtime/GeometryProcessing/Source/DynamicMesh/Public/GroupTopology.h b/Engine/Plugins/Runtime/GeometryProcessing/Source/DynamicMesh/Public/GroupTopology.h
--- a/Engine/Plugins/Runtime/GeometryProcessing/Source/DynamicMesh/Public/GroupTopology.h
+++ b/Engine/Plugins/Runtime/GeometryProcessing/Source/DynamicMesh/Public/GroupTopology.h
@@ -152,9 +152,14 @@ public:
 	FGroupTopology() {}
 	FGroupTopology(const FDynamicMesh3* Mesh, bool bAutoBuild);
 	FGroupTopology(const FDynamicMesh3* Mesh, const FDynamicMeshPolygroupAttribute* GroupLayer, bool bAutoBuild);
-	
 
+	//~ If we allow deprecation warnings, we'll get a warning in destructor and assignment operator 
+	//~ about CornerVerticesFlags (which we deprecated here) whenever it's copied or destroyed.
+	PRAGMA_DISABLE_DEPRECATION_WARNINGS;
+	FGroupTopology(const FGroupTopology&) = default;
+	FGroupTopology& operator=(const FGroupTopology&) = default;
 	virtual ~FGroupTopology() {}
+	PRAGMA_ENABLE_DEPRECATION_WARNINGS;
 
 	/**
 	 * Keeps the structures of topology in place but points the mesh pointers to
@@ -174,6 +179,29 @@ public:
 	 */
 	virtual bool RebuildTopology();
 
+	/**
+	 * Optional function that will add extra corners (in addition to the ones that are always placed at junctures
+	 * of three or more group edges) during RebuildTopology whenever the function returns true for a vertex that
+	 * is part of a group edge.
+	 * For instance, this allows corners to be placed at sharp group edge bends even if the given vertex doesn't
+	 * connect 3 or more group edges.
+	 */
+	TFunction<bool(const FGroupTopology& GroupTopology, int32 Vid, const FIndex2i& AttachedGroupEdgeEids)> ShouldAddExtraCornerAtVert = nullptr;
+
+	/**
+	 * @return Set of vertex ID's that are currently considered corners despite not originally having 3 or more 
+	 *  attached group edges (due to the effects of ShouldAddExtraCornerAtVert or RebuildTopologyWithSpecificExtraCorners).
+	 */
+	const TSet<int32>& GetCurrentExtraCornerVids() const { return CurrentExtraCornerVids; }
+
+	/**
+	 * Rebuilds the topology graph while ignoring ShouldAddExtraCornerAtVert and only determining whether
+	 * an extra corner should be forced by seeing whether the vertex is in the provided set. This is used to rebuild
+	 * the topology with a particular group of extra corners during client undo/redo operations where the behavior
+	 * of ShouldAddExtraCornerAtVert may have changed in the interim.
+	 */
+	virtual bool RebuildTopologyWithSpecificExtraCorners(const TSet<int32>& ExtraCornerVids);
+
 	/** 
 	 * Adjacency of Per-Triangle integers are what define the triangle groups.
 	 * Override this function to provide an alternate group definition.
@@ -236,7 +264,7 @@ public:
 	 */
 	struct DYNAMICMESH_API FGroupEdge
 	{
-		/** Identifier for this edge, as a pair of groups, sorted by increasing value */
+		/** Groups adjoinging this edge, sorted by increasing value if both present, otherwise the second element may be InvalidID (if boundary edge). */
 		FIndex2i Groups;
 		/** Edge span for this edge */
 		FEdgeSpan Span;
@@ -406,24 +434,53 @@ protected:
 	const FDynamicMeshPolygroupAttribute* GroupLayer = nullptr;
 
 	TArray<int> GroupIDToGroupIndexMap;		// allow fast lookup of index in .Groups, given GroupID
+
+	UE_DEPRECATED(5.1, "Use VertexIDToCornerIDMap instead.")
 	TBitArray<> CornerVerticesFlags;		// bit array of corners for fast testing in ExtractGroupEdges  (redundant w/ VertexIDToCornerIDMap?)
+
 	TArray<int> EmptyArray;
 	TMap<int32, int32> VertexIDToCornerIDMap;
 
-	/** @return true if given mesh vertex is a Corner vertex */
+	TSet<int32> CurrentExtraCornerVids;
+
+	/**
+	 * @return true if given mesh vertex should be a Corner vertex. This looks at the actual mesh- it doesn't
+	 *  use the topology structures.
+	 */
+	UE_DEPRECATED(5.1, "Use ShouldVertBeCorner instead.")
 	virtual bool IsCornerVertex(int VertexID) const;
 
+	/**
+	 * Return true if the given vertex should be a corner based on the number of adjoining group edges or based on
+	 * ShouldAddExtraCornerAtVert, if provided.
+	 * 
+	 * @param bIsExtraCornerOut Set to true if the vertex didn't have 3 or more group adjoining edges, but was forced
+	 *  to be a corner by ShouldAddExtraCornerAtVert.
+	 */
+	virtual bool ShouldVertBeCorner(int VertexID, bool* bIsExtraCornerOut = nullptr) const;
+
 	/** Return false if we had a problem finding group boundaries */
-	bool ExtractGroupEdges(FGroup& Group);
+	UE_DEPRECATED(5.1, "Use GenerateBoundaryAndGroupEdges instead.")
+	bool ExtractGroupEdges(FGroup& Group) { return false; }
+
+	// Generates group edges and boundary structs for given group. Returns false if failed to find group boundaries.
+	bool GenerateBoundaryAndGroupEdges(FGroup& Group, TMap<int32, int32>& GroupEdgeMinEidToGroupEdgeID, TBitArray<>& VertCheckedForCorner);
+
+	UE_DEPRECATED(5.1, "Use MakeEdgeGroupsPair instead. The returned pair cannot uniquely identify a group edge.")
+	inline FIndex2i MakeEdgeID(int MeshEdgeID) const { return MakeEdgeGroupsPair(MeshEdgeID); }
+	UE_DEPRECATED(5.1, "Use MakeEdgeGroupsPair instead. The returned pair cannot uniquely identify a group edge.")
+	inline FIndex2i MakeEdgeID(int Group1, int Group2) const { return MakeEdgeGroupsPair(Group1, Group2); }
+
+	inline FIndex2i MakeEdgeGroupsPair(int MeshEdgeID) const;
+	inline FIndex2i MakeEdgeGroupsPair(int Group1, int Group2) const;
 
-	inline FIndex2i MakeEdgeID(int MeshEdgeID) const;
-	inline FIndex2i MakeEdgeID(int Group1, int Group2) const;
+	UE_DEPRECATED(5.1, "Group edges cannot be uniquely identified by a group pair.")
+	int FindExistingGroupEdge(int GroupID, int OtherGroupID, int FirstVertexID, int SecondVertexID) { return 0; }
 
-	int FindExistingGroupEdge(int GroupID, int OtherGroupID, int FirstVertexID, int SecondVertexID);
 	void GetAllVertexGroups(int32 VertexID, TArray<int32>& GroupsOut) const;
 };
 
-FIndex2i FGroupTopology::MakeEdgeID(int MeshEdgeID) const
+FIndex2i FGroupTopology::MakeEdgeGroupsPair(int MeshEdgeID) const
 {
 	FIndex2i EdgeTris = Mesh->GetEdgeT(MeshEdgeID);
 
@@ -433,11 +490,11 @@ FIndex2i FGroupTopology::MakeEdgeID(int MeshEdgeID) const
 	}
 	else
 	{
-		return MakeEdgeID(GetGroupID(EdgeTris.A), GetGroupID(EdgeTris.B));
+		return MakeEdgeGroupsPair(GetGroupID(EdgeTris.A), GetGroupID(EdgeTris.B));
 	}
 }
 
-FIndex2i FGroupTopology::MakeEdgeID(int Group1, int Group2) const
+FIndex2i FGroupTopology::MakeEdgeGroupsPair(int Group1, int Group2) const
 {
 	check(Group1 != Group2);
 	check(Group1 >= 0 && Group2 >= 0);
diff --git a/Engine/Plugins/Runtime/MeshModelingToolset/Source/MeshModelingTools/Private/EditMeshPolygonsTool.cpp b/Engine/Plugins/Runtime/MeshModelingToolset/Source/MeshModelingTools/Private/EditMeshPolygonsTool.cpp
--- a/Engine/Plugins/Runtime/MeshModelingToolset/Source/MeshModelingTools/Private/EditMeshPolygonsTool.cpp
+++ b/Engine/Plugins/Runtime/MeshModelingToolset/Source/MeshModelingTools/Private/EditMeshPolygonsTool.cpp
@@ -76,6 +76,38 @@ namespace EditMeshPolygonsToolLocals
 		TEXT("Maximal number of edges that PolyEd and TriEd support. Meshes that would require "
 			"more than this number of edges to be rendered in PolyEd or TriEd force the tools to "
 			"be disabled to avoid hanging the editor."));
+
+	// Allows undo/redo of addition of extra corners in the group topology based on user angle thresholds.
+	// Used after user-triggered topology corner changes where the mesh was not actually edited.
+	class FExtraCornerChange : public FToolCommandChange
+	{
+	public:
+		FExtraCornerChange(const TSet<int32>& BeforeIn, const TSet<int32>& AfterIn)
+			: Before(BeforeIn)
+			, After(AfterIn)
+		{
+		}
+		virtual void Apply(UObject* Object) override 
+		{
+			Cast<UEditMeshPolygonsTool>(Object)->RebuildTopologyWithGivenExtraCorners(After);
+		}
+		virtual void Revert(UObject* Object) override
+		{
+			Cast<UEditMeshPolygonsTool>(Object)->RebuildTopologyWithGivenExtraCorners(Before);
+		}
+		virtual bool HasExpired(UObject* Object) const override
+		{
+			return false;
+		}
+		virtual FString ToString() const override
+		{
+			return TEXT("FExtraCornerChange");
+		}
+
+	protected:
+		TSet<int32> Before;
+		TSet<int32> After;
+	};
 }
 
 /*
@@ -215,11 +247,67 @@ void UEditMeshPolygonsTool::Setup()
 				MaxEdges, CurrentMesh->EdgeCount()), EToolMessageLevel::UserError);
 			return;
 		}
+
+		Topology = MakeShared<FTriangleGroupTopology, ESPMode::ThreadSafe>(CurrentMesh.Get(), false);
 	}
+	else
+	{
+		Topology = MakeShared<FGroupTopology, ESPMode::ThreadSafe>(CurrentMesh.Get(), false);
+		
+		Topology->ShouldAddExtraCornerAtVert = 
+			[this](const FGroupTopology& GroupTopology, int32 Vid, const FIndex2i& AttachedGroupEdgeEids)
+		{
+			if (!TopologyProperties->bAddExtraCorners)
+			{
+				return false;
+			}
+
+			// Note: it's important that we don't use CurrentMesh here. It's possible that an activity might create a copy of 
+			// the topology that uses the same corner forcing function but points to a different mesh, so we want to use
+			// whatever mesh the passed-in topology uses.
+			const FDynamicMesh3* Mesh = GroupTopology.GetMesh();
+			
+			if (!ensure(Mesh->IsEdge(AttachedGroupEdgeEids.A) && Mesh->IsEdge(AttachedGroupEdgeEids.B)))
+			{
+				return false;
+			}
+
+			// Gets vector pointing from the Vid along the edge.
+			auto GetEdgeUnitVector = [Mesh, Vid](int32 Eid, FVector3d& VectorOut)->bool
+			{
+				FIndex2i EdgeVids = Mesh->GetEdgeV(Eid);
+				// Make sure that the Vid is at EdgeVids.A
+				if (EdgeVids.B == Vid)
+				{
+					Swap(EdgeVids.A, EdgeVids.B);
+				}
+				VectorOut = Mesh->GetVertex(EdgeVids.B) - Mesh->GetVertex(EdgeVids.A);
+				return VectorOut.Normalize(KINDA_SMALL_NUMBER);
+			};
+
+			FVector Edge1, Edge2;
+			if (!GetEdgeUnitVector(AttachedGroupEdgeEids.A, Edge1) || !GetEdgeUnitVector(AttachedGroupEdgeEids.B, Edge2))
+			{
+				// If either edge was degenerate, we won't consider this a corner because otherwise we will end up
+				// with two corners connected by the degenerate edge, which is not ideal.
+				return false;
+			}
+
+			return Edge1.Dot(Edge2) >= ExtraCornerDotProductThreshold;
+		};
+	}
+
+	TopologyProperties = NewObject<UPolyEditTopologyProperties>(this);
+	TopologyProperties->Initialize(this);
+	TopologyProperties->RestoreProperties(this, GetPropertyCacheIdentifier(bTriangleMode));
+
+	auto UpdateExtraCornerThreshold = [this]() { ExtraCornerDotProductThreshold = FMathd::Cos(TopologyProperties->ExtraCornerAngleThresholdDegrees * FMathd::DegToRad); };
+	UpdateExtraCornerThreshold();
+	TopologyProperties->WatchProperty(TopologyProperties->ExtraCornerAngleThresholdDegrees,
+		[this, UpdateExtraCornerThreshold](double) { UpdateExtraCornerThreshold(); });
+
+	Topology->RebuildTopology();
 
-	Topology = (bTriangleMode) ? MakeShared<FTriangleGroupTopology, ESPMode::ThreadSafe>(CurrentMesh.Get(), true)
-		: MakeShared<FGroupTopology, ESPMode::ThreadSafe>(CurrentMesh.Get(), true);
-	
 	if (!bTriangleMode)
 	{
 		int32 NumEdgesToRender = 0;
@@ -440,6 +528,15 @@ void UEditMeshPolygonsTool::Setup()
 	// CommonProps so that they are at the bottom.
 	AddToolPropertySource(SelectionMechanic->Properties);
 	AddToolPropertySource(CommonProps);
+	if (!bTriangleMode)
+	{
+		AddToolPropertySource(TopologyProperties);
+	}
+	else
+	{
+		// Not actually necessary since we don't use the forcing function in triangle mode, but might as well turn it off here too.
+		TopologyProperties->bAddExtraCorners = false;
+	}
 
 	// hide input StaticMeshComponent
 	UE::ToolTarget::HideSourceObject(Target);
@@ -552,6 +649,7 @@ void UEditMeshPolygonsTool::OnShutdown(EToolShutdownType ShutdownType)
 	}
 	CommonProps->SaveProperties(this, GetPropertyCacheIdentifier(bTriangleMode));
 	SelectionMechanic->Properties->SaveProperties(this, GetPropertyCacheIdentifier(bTriangleMode));
+	TopologyProperties->SaveProperties(this, GetPropertyCacheIdentifier(bTriangleMode));
 
 	GetToolManager()->GetContextObjectStore()->RemoveContextObjectsOfType(UPolyEditActivityContext::StaticClass());
 	ActivityContext = nullptr;
@@ -1114,6 +1212,9 @@ void UEditMeshPolygonsTool::OnTick(float DeltaTime)
 		case EEditMeshPolygonsToolActions::SimplifyByGroups:
 			SimplifyByGroups();
 			break;
+		case EEditMeshPolygonsToolActions::RegenerateExtraCorners:
+			ApplyRegenerateExtraCorners();
+			break;
 		}
 
 		PendingAction = EEditMeshPolygonsToolActions::NoAction;
@@ -1135,6 +1236,7 @@ void UEditMeshPolygonsTool::StartActivity(TObjectPtr<UInteractiveToolActivity> A
 		}
 		SelectionMechanic->SetIsEnabled(false);
 		SetToolPropertySourceEnabled(SelectionMechanic->Properties, false);
+		SetToolPropertySourceEnabled(TopologyProperties, false);
 		CurrentActivity = Activity;
 		if (CurrentActivity->HasAccept())
 		{
@@ -1164,6 +1266,7 @@ void UEditMeshPolygonsTool::EndCurrentActivity(EToolShutdownType ShutdownType)
 		SetToolPropertySourceEnabled(AcceptCancelAction, false);
 		SetActionButtonPanelsVisible(true);
 		SelectionMechanic->SetIsEnabled(true);
+		SetToolPropertySourceEnabled(TopologyProperties, true);
 		SetToolPropertySourceEnabled(SelectionMechanic->Properties, true);
 		UpdateGizmoVisibility();
 	}
@@ -1284,14 +1387,13 @@ void UEditMeshPolygonsTool::ApplyChange(const FMeshVertexChange* Change, bool bR
 }
 
 
-void UEditMeshPolygonsTool::UpdateFromCurrentMesh(bool bGroupTopologyModified)
+void UEditMeshPolygonsTool::UpdateFromCurrentMesh(bool bUpdateTopology)
 {
-	Preview->PreviewMesh->UpdatePreview(CurrentMesh.Get(),
-		bGroupTopologyModified ? UPreviewMesh::ERenderUpdateMode::FullUpdate : UPreviewMesh::ERenderUpdateMode::FastUpdate);
+	Preview->PreviewMesh->UpdatePreview(CurrentMesh.Get(), UPreviewMesh::ERenderUpdateMode::FullUpdate);
 	bSpatialDirty = true;
-	SelectionMechanic->NotifyMeshChanged(bGroupTopologyModified);
+	SelectionMechanic->NotifyMeshChanged(bUpdateTopology);
 
-	if (bGroupTopologyModified)
+	if (bUpdateTopology)
 	{
 		Topology->RebuildTopology();
 	}
@@ -1529,6 +1631,54 @@ void UEditMeshPolygonsTool::SimplifyByGroups()
 
 
 
+void UEditMeshPolygonsTool::ApplyRegenerateExtraCorners()
+{
+	if (!ensure(!bTriangleMode && Topology))
+	{
+		return;
+	}
+
+	// We need to remember the extra corners that get generated and put them into the undo system so that if we
+	// change the settings later, undoing still brings us back to the result we saw at that time.
+	TSet<int32> PreviousExtraCorners(Topology->GetCurrentExtraCornerVids());
+	Topology->RebuildTopology();
+	const TSet<int32>& NewExtraCorners = Topology->GetCurrentExtraCornerVids();
+
+	bool bCornersChanged = PreviousExtraCorners.Num() != NewExtraCorners.Num() || !PreviousExtraCorners.Includes(NewExtraCorners);
+	if (bCornersChanged)
+	{
+		const FText TransactionLabel = LOCTEXT("RegenerateCornersTransactionName", "Regenerate Corners");
+
+		GetToolManager()->BeginUndoTransaction(TransactionLabel);
+		if (SelectionMechanic && !SelectionMechanic->GetActiveSelection().IsEmpty())
+		{
+			SelectionMechanic->BeginChange();
+			SelectionMechanic->ClearSelection();
+			GetToolManager()->EmitObjectChange(SelectionMechanic, SelectionMechanic->EndChange(), TransactionLabel);
+		}
+
+		GetToolManager()->EmitObjectChange(this,
+			MakeUnique<EditMeshPolygonsToolLocals::FExtraCornerChange>(PreviousExtraCorners, NewExtraCorners),
+			TransactionLabel);
+
+		GetToolManager()->EndUndoTransaction();
+	}
+	
+	if (SelectionMechanic)
+	{
+		SelectionMechanic->NotifyMeshChanged(true);
+	}
+}
+
+
+void UEditMeshPolygonsTool::RebuildTopologyWithGivenExtraCorners(const TSet<int32>& Vids)
+{
+	Topology->RebuildTopologyWithSpecificExtraCorners(Vids);
+	SelectionMechanic->NotifyMeshChanged(true);
+}
+
+
+
 void UEditMeshPolygonsTool::ApplyDecompose()
 {
 	if (BeginMeshFaceEditChange() == false)
@@ -2165,8 +2315,7 @@ void UEditMeshPolygonsTool::EmitCurrentMeshChangeAndUpdate(const FText& Transact
 	FGroupTopologySelection TempSelection;
 	const FGroupTopologySelection* OutputSelectionToUse = &OutputSelection;
 
-	// Because we always consider group topology to be modified, the UpdateFromCurrentMesh() call down below
-	// will clear out our existing selection. However, we want to do this ourselves so that we issue a transaction.
+	// Emit a selection clear before emitting the mesh change, so that undo restores it properly.
 	if (!SelectionMechanic->GetActiveSelection().IsEmpty() /* && (bSelectionModified || bGroupTopologyModified) */)
 	{
 		if (bReferencingSameSelection)
@@ -2181,12 +2330,25 @@ void UEditMeshPolygonsTool::EmitCurrentMeshChangeAndUpdate(const FText& Transact
 		GetToolManager()->EmitObjectChange(SelectionMechanic, SelectionMechanic->EndChange(), LOCTEXT("ClearSelection", "Clear Selection"));
 	}
 
+	// Prep and emit the mesh change. This needs to be bookended by the change in extra corners, since
+	// those get regenerated in the topology rebuild.
+	TUniquePtr<FEditMeshPolygonsToolMeshChange> ChangeToEmit = MakeUnique<FEditMeshPolygonsToolMeshChange>(MoveTemp(MeshChangeIn));
+	if (!Topology->GetCurrentExtraCornerVids().IsEmpty())
+	{
+		ChangeToEmit->ExtraCornerVidsBefore = Topology->GetCurrentExtraCornerVids();
+	}
+	Topology->RebuildTopology();
+	if (!Topology->GetCurrentExtraCornerVids().IsEmpty())
+	{
+		ChangeToEmit->ExtraCornerVidsAfter = Topology->GetCurrentExtraCornerVids();
+	}
 	GetToolManager()->EmitObjectChange(this, 
-		MakeUnique<FEditMeshPolygonsToolMeshChange>(MoveTemp(MeshChangeIn)),
+		MoveTemp(ChangeToEmit),
 		TransactionLabel);
 
-	// Update related structures
-	UpdateFromCurrentMesh(bGroupTopologyModified);
+	// Update other related structures
+	UpdateFromCurrentMesh(false);
+	SelectionMechanic->NotifyMeshChanged(true); // This wasn't updated in UpdateFromCurrentMesh because we didn't ask to rebuild topology
 	ModifiedTopologyCounter += bGroupTopologyModified;
 
 	// Set output selection if there's a non-empty one. We know we've cleared the selection by
@@ -2329,28 +2491,26 @@ void FEditMeshPolygonsToolMeshChange::Apply(UObject* Object)
 {
 	UEditMeshPolygonsTool* Tool = Cast<UEditMeshPolygonsTool>(Object);
 	
-	// This function currently only supports FDynamicMeshChange but that should be issued only when the mesh changes
-	// topology. For now we use it even when eg vertex postions change. See :HandlePositionOnlyMeshChanges
-	bool bGroupTopologyModified = true;
-	
 	MeshChange->Apply(Tool->CurrentMesh.Get(), false);
-	Tool->UpdateFromCurrentMesh(bGroupTopologyModified);
-	Tool->ModifiedTopologyCounter += bGroupTopologyModified;
-	Tool->ActivityContext->OnUndoRedo.Broadcast(bGroupTopologyModified);
+	Tool->UpdateFromCurrentMesh(false);
+	++Tool->ModifiedTopologyCounter;
+
+	Tool->RebuildTopologyWithGivenExtraCorners(ExtraCornerVidsAfter);
+
+	Tool->ActivityContext->OnUndoRedo.Broadcast(true);
 }
 
 void FEditMeshPolygonsToolMeshChange::Revert(UObject* Object)
 {
 	UEditMeshPolygonsTool* Tool = Cast<UEditMeshPolygonsTool>(Object);
 	
-	// This function currently only supports FDynamicMeshChange but that should be issued only when the mesh changes
-	// topology. For now we use it even when eg vertex postions change. See :HandlePositionOnlyMeshChanges
-	bool bGroupTopologyModified = true;
-	
 	MeshChange->Apply(Tool->CurrentMesh.Get(), true);
-	Tool->UpdateFromCurrentMesh(bGroupTopologyModified);
-	Tool->ModifiedTopologyCounter -= bGroupTopologyModified;
-	Tool->ActivityContext->OnUndoRedo.Broadcast(bGroupTopologyModified);
+	Tool->UpdateFromCurrentMesh(false);
+	++Tool->ModifiedTopologyCounter;
+
+	Tool->RebuildTopologyWithGivenExtraCorners(ExtraCornerVidsBefore);
+
+	Tool->ActivityContext->OnUndoRedo.Broadcast(true);
 }
 
 FString FEditMeshPolygonsToolMeshChange::ToString() const
diff --git a/Engine/Plugins/Runtime/MeshModelingToolset/Source/MeshModelingTools/Public/EditMeshPolygonsTool.h b/Engine/Plugins/Runtime/MeshModelingToolset/Source/MeshModelingTools/Public/EditMeshPolygonsTool.h
--- a/Engine/Plugins/Runtime/MeshModelingToolset/Source/MeshModelingTools/Public/EditMeshPolygonsTool.h
+++ b/Engine/Plugins/Runtime/MeshModelingToolset/Source/MeshModelingTools/Public/EditMeshPolygonsTool.h
@@ -76,7 +76,7 @@ public:
 	bool bShowWireframe = false;
 
 	UPROPERTY(EditAnywhere, Category = Options)
-	bool bShowSelectableCorners = false;
+	bool bShowSelectableCorners = true;
 
 	/** When true, allows the transform gizmo to be rendered */
 	UPROPERTY(EditAnywhere, Category = Options)
@@ -135,6 +135,7 @@ enum class EEditMeshPolygonsToolActions
 	PlanarProjectionUV,
 
 	SimplifyByGroups,
+	RegenerateExtraCorners,
 
 	// triangle-specific edits
 	PokeSingleFace,
@@ -189,6 +190,36 @@ public:
 	void PostAction(EEditMeshPolygonsToolActions Action);
 };
 
+UCLASS()
+class MESHMODELINGTOOLS_API UPolyEditTopologyProperties : public UEditMeshPolygonsToolActionPropertySet
+{
+	GENERATED_BODY()
+
+public:
+	/** 
+	 * When true, adds extra corners at sharp group edge bends (in addition to the normal corners that
+	 * are placed at junctures of three or more group edges). For instance, a single disconnected quad-like group
+	 * would normally have a single group edge with no corners, since it has no neighboring groups, but this
+	 * setting will allow for the generation of corners at the quad corners, which is very useful for editing.
+	 * Note that the setting takes effect only after clicking Regenerate Extra Corners or performing some
+	 * operation that changes the group topology.
+	 */
+	UPROPERTY(EditAnywhere, Category = TopologyOptions)
+	bool bAddExtraCorners = true;
+
+	UFUNCTION(CallInEditor, Category = TopologyOptions)
+	void RegenerateExtraCorners() { PostAction(EEditMeshPolygonsToolActions::RegenerateExtraCorners); }
+
+	/** 
+	 * When generating extra corners, how sharp the angle needs to be to warrant an extra corner placement there. Lower values require
+	 * sharper corners, so are more tolerant of curved group edges. For instance, 180 will place corners at every vertex along a group
+	 * edge even if the edge is perfectly straight, and 135 will place a vertex only once the edge bends 45 degrees off the straight path
+	 * (i.e. 135 degrees to the previous edge). 
+	 * The setting is applied either when Regenerate Extra Corners is clicked, or after any operation that modifies topology.
+	 */
+	UPROPERTY(EditAnywhere, Category = TopologyOptions, meta = (ClampMin = "0", ClampMax = "180", EditCondition = "bAddExtraCorners"))
+	double ExtraCornerAngleThresholdDegrees = 135;
+};
 
 /** PolyEdit Actions */
 UCLASS()
@@ -452,6 +483,9 @@ public:
 
 	virtual void SetWorld(UWorld* World) { this->TargetWorld = World; }
 
+	// used by undo/redo
+	void RebuildTopologyWithGivenExtraCorners(const TSet<int32>& Vids);
+
 	virtual void Setup() override;
 	virtual void OnShutdown(EToolShutdownType ShutdownType) override;
 
@@ -526,6 +560,8 @@ protected:
 	UPROPERTY()
 	TObjectPtr<UEditMeshPolygonsToolAcceptCancelAction> AcceptCancelAction = nullptr;
 
+	UPROPERTY()
+	TObjectPtr<UPolyEditTopologyProperties> TopologyProperties = nullptr;
 
 	/**
 	 * Activity objects that handle multi-interaction operations
@@ -642,6 +678,9 @@ protected:
 
 	void SimplifyByGroups();
 
+	void ApplyRegenerateExtraCorners();
+	double ExtraCornerDotProductThreshold = -1;
+
 	FFrame3d ActiveSelectionFrameLocal;
 	FFrame3d ActiveSelectionFrameWorld;
 	TArray<int32> ActiveTriangleSelection;
@@ -679,7 +718,7 @@ protected:
 	bool BeginMeshBoundaryEdgeEditChange(bool bOnlySimple);
 	bool BeginMeshEdgeEditChange(TFunctionRef<bool(int32)> GroupEdgeIDFilterFunc);
 
-	void UpdateFromCurrentMesh(bool bGroupTopologyModified);
+	void UpdateFromCurrentMesh(bool bRebuildTopology);
 	int32 ModifiedTopologyCounter = 0;
 	bool bWasTopologyEdited = false;
 
@@ -696,7 +735,8 @@ protected:
 
 /**
  * Wraps a FDynamicMeshChange so that it can be expired and so that other data
- * structures in the tool can be updated.
+ * structures in the tool can be updated. On apply/revert, the topology is rebuilt
+ * using stored extra corner vids.
  */
 class MESHMODELINGTOOLS_API FEditMeshPolygonsToolMeshChange : public FToolCommandChange
 {
@@ -709,6 +749,8 @@ public:
 	virtual void Revert(UObject* Object) override;
 	virtual FString ToString() const override;
 
+	TSet<int32> ExtraCornerVidsBefore;
+	TSet<int32> ExtraCornerVidsAfter;
 protected:
 	TUniquePtr<UE::Geometry::FDynamicMeshChange> MeshChange;
 };
diff --git a/Engine/Plugins/Runtime/MeshModelingToolset/Source/ModelingComponents/Private/Selection/PolygonSelectionMechanic.cpp b/Engine/Plugins/Runtime/MeshModelingToolset/Source/ModelingComponents/Private/Selection/PolygonSelectionMechanic.cpp
--- a/Engine/Plugins/Runtime/MeshModelingToolset/Source/ModelingComponents/Private/Selection/PolygonSelectionMechanic.cpp
+++ b/Engine/Plugins/Runtime/MeshModelingToolset/Source/ModelingComponents/Private/Selection/PolygonSelectionMechanic.cpp
@@ -87,7 +87,7 @@ void UPolygonSelectionMechanic::Setup(UInteractiveTool* ParentToolIn)
 	PolyEdgesRenderer.LineColor = FLinearColor::Red;
 	PolyEdgesRenderer.LineThickness = 2.0;
 	PolyEdgesRenderer.PointColor = FLinearColor::Red;
-	PolyEdgesRenderer.PointSize = 10.0f;
+	PolyEdgesRenderer.PointSize = 8.0f;
 	HilightRenderer.LineColor = FLinearColor::Green;
 	HilightRenderer.LineThickness = 4.0f;
 	HilightRenderer.PointColor = FLinearColor::Green;
