remove DLL export symbol on header-defined function ComputeNumMatrixElements

diff --git a/Engine/Plugins/Runtime/GeometryProcessing/Source/DynamicMesh/Public/Solvers/PrecomputedMeshWeightData.h b/Engine/Plugins/Runtime/GeometryProcessing/Source/DynamicMesh/Public/Solvers/PrecomputedMeshWeightData.h
--- a/Engine/Plugins/Runtime/GeometryProcessing/Source/DynamicMesh/Public/Solvers/PrecomputedMeshWeightData.h
+++ b/Engine/Plugins/Runtime/GeometryProcessing/Source/DynamicMesh/Public/Solvers/PrecomputedMeshWeightData.h
@@ -13,7 +13,7 @@ namespace UE
 
 		// Utility to compute the number of elements in the sparse laplacian matrix
 		template<typename MeshT>
-		DYNAMICMESH_API int32 ComputeNumMatrixElements(const MeshT& DynamicMesh, const TArray<int32>& ToVtxId)
+		int32 ComputeNumMatrixElements(const MeshT& DynamicMesh, const TArray<int32>& ToVtxId)
 		{
 			const int32 NumVerts = ToVtxId.Num();
 			TArray<int32> OneRingSize;
