Add test coverage to ensure that Clipper's Union operation can also be used to resolve self-intersections via union-with-empty

diff --git a/Engine/Plugins/Runtime/GeometryProcessing/Source/GeometryAlgorithms/Private/ConstrainedDelaunay2.cpp b/Engine/Plugins/Runtime/GeometryProcessing/Source/GeometryAlgorithms/Private/ConstrainedDelaunay2.cpp
--- a/Engine/Plugins/Runtime/GeometryProcessing/Source/GeometryAlgorithms/Private/ConstrainedDelaunay2.cpp
+++ b/Engine/Plugins/Runtime/GeometryProcessing/Source/GeometryAlgorithms/Private/ConstrainedDelaunay2.cpp
@@ -3,7 +3,7 @@
 #include "ConstrainedDelaunay2.h"
 #include "CompGeom/Delaunay2.h"
 #include "Async/ParallelFor.h"
-
+#include "Curve/PolygonIntersectionUtils.h"
 
 using namespace UE::Geometry;
 
@@ -278,6 +278,39 @@ namespace ConstrainedDelaunay2Internal
 	}
 }
 
+template<class RealType>
+bool TConstrainedDelaunay2<RealType>::Add(const TPolygon2<RealType>& Polygon, bool bResolveSelfIntersections)
+{
+	if (!bResolveSelfIntersections)
+	{
+		Add(Polygon);
+		return true;
+	}
+	TGeneralPolygon2<RealType> GPolygon(Polygon);
+	return Add(GPolygon, bResolveSelfIntersections);
+}
+
+template<class RealType>
+bool TConstrainedDelaunay2<RealType>::Add(const TGeneralPolygon2<RealType>& GPolygon, bool bResolveSelfIntersections)
+{
+	if (!bResolveSelfIntersections)
+	{
+		Add(GPolygon);
+		return true;
+	}
+	TGeneralPolygon2<RealType> Empty;
+	TUnionPolygon2Polygon2<TGeneralPolygon2<RealType>, RealType> Union(GPolygon, Empty);
+	bool bResolveSuccess = Union.ComputeResult();
+	if (bResolveSuccess)
+	{
+		for (const TGeneralPolygon2<RealType>& ResultGPoly : Union.Result)
+		{
+			Add(ResultGPoly);
+		}
+	}
+	return bResolveSuccess;
+}
+
 
 template<class RealType>
 bool TConstrainedDelaunay2<RealType>::Triangulate(TFunctionRef<bool(const TArray<TVector2<RealType>>&, const FIndex3i&)> KeepTriangle)
diff --git a/Engine/Plugins/Runtime/GeometryProcessing/Source/GeometryAlgorithms/Public/ConstrainedDelaunay2.h b/Engine/Plugins/Runtime/GeometryProcessing/Source/GeometryAlgorithms/Public/ConstrainedDelaunay2.h
--- a/Engine/Plugins/Runtime/GeometryProcessing/Source/GeometryAlgorithms/Public/ConstrainedDelaunay2.h
+++ b/Engine/Plugins/Runtime/GeometryProcessing/Source/GeometryAlgorithms/Public/ConstrainedDelaunay2.h
@@ -126,6 +126,11 @@ struct TConstrainedDelaunay2
 		}
 	}
 
+	// Add Polygon with the option to resolve self-intersections. Returns false if resolving self-intersections failed.
+	bool GEOMETRYALGORITHMS_API Add(const TPolygon2<RealType>& GPolygon, bool bResolveSelfIntersections);
+	// Add General Polygon with the option to resolve self-intersections. Returns false if resolving self-intersections failed.
+	bool GEOMETRYALGORITHMS_API Add(const TGeneralPolygon2<RealType>& GPolygon, bool bResolveSelfIntersections);
+
 	//
 	// outputs
 	//
