Make Arrangment2d not fail/ensure on a degenerate input (collapsed to a single point)

diff --git a/Engine/Plugins/Runtime/GeometryProcessing/Source/GeometryAlgorithms/Public/Arrangement2d.h b/Engine/Plugins/Runtime/GeometryProcessing/Source/GeometryAlgorithms/Public/Arrangement2d.h
--- a/Engine/Plugins/Runtime/GeometryProcessing/Source/GeometryAlgorithms/Public/Arrangement2d.h
+++ b/Engine/Plugins/Runtime/GeometryProcessing/Source/GeometryAlgorithms/Public/Arrangement2d.h
@@ -43,12 +43,12 @@ struct FArrangement2d
 	double VertexSnapTol = 0.00001;
 
 	FArrangement2d(const FAxisAlignedBox2d& BoundsHint)
-		: PointHash(BoundsHint.MaxDim() / 64, -1)
+		: PointHash(FMath::Max(FMathd::ZeroTolerance, BoundsHint.MaxDim() / 64), -1)
 	{
 	}
 
 	FArrangement2d(double PointHashCellSize)
-		: PointHash(PointHashCellSize, -1)
+		: PointHash(FMath::Max(FMathd::ZeroTolerance, PointHashCellSize), -1)
 	{
 	}
 
