non-unity build fixes for polygonintersectionutils.h

diff --git a/Engine/Plugins/Runtime/GeometryProcessing/Source/GeometryAlgorithms/Private/Curve/PolygonIntersectionUtils.cpp b/Engine/Plugins/Runtime/GeometryProcessing/Source/GeometryAlgorithms/Private/Curve/PolygonIntersectionUtils.cpp
--- a/Engine/Plugins/Runtime/GeometryProcessing/Source/GeometryAlgorithms/Private/Curve/PolygonIntersectionUtils.cpp
+++ b/Engine/Plugins/Runtime/GeometryProcessing/Source/GeometryAlgorithms/Private/Curve/PolygonIntersectionUtils.cpp
@@ -1,4 +1,4 @@
-﻿// Copyright Epic Games, Inc. All Rights Reserved.
+// Copyright Epic Games, Inc. All Rights Reserved.
 
 #include "Curve/PolygonIntersectionUtils.h"
 
@@ -53,23 +53,26 @@ TBooleanPolygon2Polygon2<OperationType, GeometryType, RealType>::TBooleanPolygon
 }
 
 template <EPolygonBooleanOp OperationType, typename GeometryType, typename RealType>
-constexpr bool TBooleanPolygon2Polygon2<OperationType, GeometryType, RealType>::ComputeResult()
+bool TBooleanPolygon2Polygon2<OperationType, GeometryType, RealType>::ComputeResult()
 {
-	switch (OperationType)
+	if constexpr (OperationType == EPolygonBooleanOp::Union)
 	{
-	case EPolygonBooleanOp::Union:
 		return Private::Clip<RealType>(Clipper2Lib::ClipType::Union, PolygonA, PolygonB, Result);
-
-	case EPolygonBooleanOp::Difference:
+	}
+	else if constexpr (OperationType == EPolygonBooleanOp::Difference)
+	{
 		return Private::Clip<RealType>(Clipper2Lib::ClipType::Difference, PolygonA, PolygonB, Result);
-		
-	case EPolygonBooleanOp::Intersect:
+	}
+	else if constexpr (OperationType == EPolygonBooleanOp::Intersect)
+	{
 		return Private::Clip<RealType>(Clipper2Lib::ClipType::Intersection, PolygonA, PolygonB, Result);
-		
-	case EPolygonBooleanOp::ExclusiveOr:
+	}
+	else if constexpr (OperationType == EPolygonBooleanOp::ExclusiveOr)
+	{
 		return Private::Clip<RealType>(Clipper2Lib::ClipType::Xor, PolygonA, PolygonB, Result);
-
-	default:
+	}
+	else
+	{
 		return false;
 	}	
 }
diff --git a/Engine/Plugins/Runtime/GeometryProcessing/Source/GeometryAlgorithms/Public/Curve/PolygonIntersectionUtils.h b/Engine/Plugins/Runtime/GeometryProcessing/Source/GeometryAlgorithms/Public/Curve/PolygonIntersectionUtils.h
--- a/Engine/Plugins/Runtime/GeometryProcessing/Source/GeometryAlgorithms/Public/Curve/PolygonIntersectionUtils.h
+++ b/Engine/Plugins/Runtime/GeometryProcessing/Source/GeometryAlgorithms/Public/Curve/PolygonIntersectionUtils.h
@@ -1,4 +1,4 @@
-﻿// Copyright Epic Games, Inc. All Rights Reserved.
+// Copyright Epic Games, Inc. All Rights Reserved.
 
 #pragma once
 
@@ -17,7 +17,7 @@ namespace UE::Geometry
 
 	/** Produces the boolean'd result of PolygonA and PolygonB depending on the OperationType. Both polygons must be closed. */
 	template <EPolygonBooleanOp OperationType, typename GeometryType, typename RealType>
-	class TBooleanPolygon2Polygon2
+	class GEOMETRYALGORITHMS_API TBooleanPolygon2Polygon2
 	{
 	public:
 		// Input
@@ -31,7 +31,7 @@ namespace UE::Geometry
 			const GeometryType& InPolygonA,
 			const GeometryType& InPolygonB);
 
-		constexpr bool ComputeResult();
+		bool ComputeResult();
 	};
  
 	template <EPolygonBooleanOp OperationType>
