Fix ConstrainedDelaunay2 Add() function being overloaded in two ways; make the new Add be more-specifically-named AddWithIntersectionResolution

diff --git a/Engine/Plugins/Runtime/GeometryProcessing/Source/GeometryAlgorithms/Private/ConstrainedDelaunay2.cpp b/Engine/Plugins/Runtime/GeometryProcessing/Source/GeometryAlgorithms/Private/ConstrainedDelaunay2.cpp
--- a/Engine/Plugins/Runtime/GeometryProcessing/Source/GeometryAlgorithms/Private/ConstrainedDelaunay2.cpp
+++ b/Engine/Plugins/Runtime/GeometryProcessing/Source/GeometryAlgorithms/Private/ConstrainedDelaunay2.cpp
@@ -279,25 +279,15 @@ namespace ConstrainedDelaunay2Internal
 }
 
 template<class RealType>
-bool TConstrainedDelaunay2<RealType>::Add(const TPolygon2<RealType>& Polygon, bool bResolveSelfIntersections)
+bool TConstrainedDelaunay2<RealType>::AddWithIntersectionResolution(const TPolygon2<RealType>& Polygon)
 {
-	if (!bResolveSelfIntersections)
-	{
-		Add(Polygon);
-		return true;
-	}
 	TGeneralPolygon2<RealType> GPolygon(Polygon);
-	return Add(GPolygon, bResolveSelfIntersections);
+	return AddWithIntersectionResolution(GPolygon);
 }
 
 template<class RealType>
-bool TConstrainedDelaunay2<RealType>::Add(const TGeneralPolygon2<RealType>& GPolygon, bool bResolveSelfIntersections)
+bool TConstrainedDelaunay2<RealType>::AddWithIntersectionResolution(const TGeneralPolygon2<RealType>& GPolygon)
 {
-	if (!bResolveSelfIntersections)
-	{
-		Add(GPolygon);
-		return true;
-	}
 	TGeneralPolygon2<RealType> Empty;
 	TUnionPolygon2Polygon2<TGeneralPolygon2<RealType>, RealType> Union(GPolygon, Empty);
 	bool bResolveSuccess = Union.ComputeResult();
diff --git a/Engine/Plugins/Runtime/GeometryProcessing/Source/GeometryAlgorithms/Public/ConstrainedDelaunay2.h b/Engine/Plugins/Runtime/GeometryProcessing/Source/GeometryAlgorithms/Public/ConstrainedDelaunay2.h
--- a/Engine/Plugins/Runtime/GeometryProcessing/Source/GeometryAlgorithms/Public/ConstrainedDelaunay2.h
+++ b/Engine/Plugins/Runtime/GeometryProcessing/Source/GeometryAlgorithms/Public/ConstrainedDelaunay2.h
@@ -126,10 +126,10 @@ struct TConstrainedDelaunay2
 		}
 	}
 
-	// Add Polygon with the option to resolve self-intersections. Returns false if resolving self-intersections failed.
-	bool GEOMETRYALGORITHMS_API Add(const TPolygon2<RealType>& GPolygon, bool bResolveSelfIntersections);
-	// Add General Polygon with the option to resolve self-intersections. Returns false if resolving self-intersections failed.
-	bool GEOMETRYALGORITHMS_API Add(const TGeneralPolygon2<RealType>& GPolygon, bool bResolveSelfIntersections);
+	// Add Polygon with self-intersections resolved. Returns false if resolving self-intersections failed.
+	bool GEOMETRYALGORITHMS_API AddWithIntersectionResolution(const TPolygon2<RealType>& Polygon);
+	// Add General Polygon with self-intersections resolved. Returns false if resolving self-intersections failed.
+	bool GEOMETRYALGORITHMS_API AddWithIntersectionResolution(const TGeneralPolygon2<RealType>& GeneralPolygon);
 
 	//
 	// outputs
