Intrinsic Mesh Laplacian Construction - correctly handle case when both opposing angles are nearly 90 degrees.

This shouid improve the behaviour of the PolyDefTool when used in smooth mode.

diff --git a/Engine/Plugins/Runtime/GeometryProcessing/Source/DynamicMesh/Public/Solvers/LaplacianMatrixAssembly.h b/Engine/Plugins/Runtime/GeometryProcessing/Source/DynamicMesh/Public/Solvers/LaplacianMatrixAssembly.h
--- a/Engine/Plugins/Runtime/GeometryProcessing/Source/DynamicMesh/Public/Solvers/LaplacianMatrixAssembly.h
+++ b/Engine/Plugins/Runtime/GeometryProcessing/Source/DynamicMesh/Public/Solvers/LaplacianMatrixAssembly.h
@@ -728,7 +728,8 @@ void UE::MeshDeformation::ConstructIDTCotangentLaplacian(const FDynamicMesh3& Dy
 			{
 				continue;
 			}
-			EdgeCotanWeights[EdgeId] = IntrinsicMesh.EdgeCotanWeight(EdgeId);
+			const double CTWeight = IntrinsicMesh.EdgeCotanWeight(EdgeId);
+			EdgeCotanWeights[EdgeId] = FMathd::Max(0., CTWeight);  // roundoff errors can produce very, very small negative values when the angles opposite the edge are each 90-degrees. 
 		}
 	}
 
@@ -746,11 +747,12 @@ void UE::MeshDeformation::ConstructIDTCotangentLaplacian(const FDynamicMesh3& Dy
 		{
 			for (int32 EdgeId : IntrinsicMesh.VtxEdgesItr(IVertId))
 			{
+
 				const double EdgeLength = IntrinsicMesh.GetEdgeLength(EdgeId);
 				const double CTWeight = EdgeCotanWeights[EdgeId];
 				const double CTWeightLSqr = CTWeight * EdgeLength * EdgeLength;
 				
-				if (Uncorrected.Contains(EdgeId) || CTWeightLSqr < 1.e-1)
+				if (Uncorrected.Contains(EdgeId)) // this part of the intrinsic mesh is not Delaunay.
 				{
 					bUseUniform = true; 
 				}
@@ -758,7 +760,7 @@ void UE::MeshDeformation::ConstructIDTCotangentLaplacian(const FDynamicMesh3& Dy
 				WeightArea += CTWeightLSqr;
 			}
 			WeightArea *= 0.25;
-			WeightArea = FMath::Max(WeightArea, 1.e-5);
+			WeightArea = FMathd::Max(WeightArea, 1.e-5);
 		}		
 
 		double WeightII = 0.; // accumulate to equal and opposite the sum of the neighbor weights
