MeshModelingTools: Fix group topology not adding an edge if the boundary has a single corner.

diff --git a/Engine/Plugins/Runtime/GeometryProcessing/Source/DynamicMesh/Private/GroupTopology.cpp b/Engine/Plugins/Runtime/GeometryProcessing/Source/DynamicMesh/Private/GroupTopology.cpp
--- a/Engine/Plugins/Runtime/GeometryProcessing/Source/DynamicMesh/Private/GroupTopology.cpp
+++ b/Engine/Plugins/Runtime/GeometryProcessing/Source/DynamicMesh/Private/GroupTopology.cpp
@@ -610,8 +610,17 @@ bool FGroupTopology::GenerateBoundaryAndGroupEdges(FGroup& Group,
 			int32 StartIndex = CornerIndices[k];
 			int32 EndIndex = CornerIndices[k + 1];		// note: StartIndex == EndIndex on a closed loop, ie NumSpans == 1
 
+			// This calculation of number of edges is correct as long as we're not starting and stopping in the same place.
 			int32 NumSpanEdges = (EndIndex + NumV - StartIndex) % NumV;
 
+			// If StartIndex == EndIndex, then NumSpanEdges should actually be NumV, not 0, because we are on a loop with
+			// a single corner (either a bowtie or user-generated).
+			if (NumSpanEdges == 0)
+			{
+				ensure(NumSpans == 1);
+				NumSpanEdges = NumV;
+			}
+
 			// Find the min eid in this subsection so that we can see if we already have a group edge for it.
 			int32 MinEid = Loop.Edges[StartIndex];
 			for (int32 i = 1; i < NumSpanEdges; ++i)
