MeshModelingTools: Fix WalkMeshPlanar occasionally picking bad paths that resulted in holes and other errors in boolean operations.

diff --git a/Engine/Plugins/Runtime/GeometryProcessing/Source/DynamicMesh/Private/Operations/EmbedSurfacePath.cpp b/Engine/Plugins/Runtime/GeometryProcessing/Source/DynamicMesh/Private/Operations/EmbedSurfacePath.cpp
--- a/Engine/Plugins/Runtime/GeometryProcessing/Source/DynamicMesh/Private/Operations/EmbedSurfacePath.cpp
+++ b/Engine/Plugins/Runtime/GeometryProcessing/Source/DynamicMesh/Private/Operations/EmbedSurfacePath.cpp
@@ -124,16 +124,24 @@ FMeshSurfacePoint RelocateTrianglePointAfterRefinement(const FDynamicMesh3* Mesh
 struct FIndexDistance
 {
 	int Index;
-	double Distance;
+	double PathLength;
+	double DistanceToEnd;
 	bool operator<(const FIndexDistance& Other) const
 	{
-		return Distance < Other.Distance;
+		return PathLength + DistanceToEnd < Other.PathLength + Other.DistanceToEnd;
 	}
 };
 
 bool WalkMeshPlanar(const FDynamicMesh3* Mesh, int StartTri, int StartVID, FVector3d StartPt, int EndTri, int EndVertID, FVector3d EndPt, FVector3d WalkPlaneNormal, TFunction<FVector3d(const FDynamicMesh3*, int)> VertexToPosnFn,
 	bool bAllowBackwardsSearch, double AcceptEndPtOutsideDist, double PtOnPlaneThresholdSq, TArray<TPair<FMeshSurfacePoint, int>>& WalkedPath, double BackwardsTolerance)
 {
+	// Even when bAllowBackwardsSearch is false, there may be multiple paths from start point to end point. 
+	// The approach we take is a breadth-first-like one where we keep track of multiple paths and always extend the one
+	// that has the potential to be shortest, i.e. whose current path length + distance to destination is smallest.
+	// (we can't go by distance to destination alone because a sub optimal path can curve around the destination in the
+	// plane in such a way that it always seems closer than the next step in a more direct path that happens to pass
+	// through a more tesselated region).
+
 	auto SetTriVertPositions = [&VertexToPosnFn, &Mesh](FIndex3i TriVertIDs, FTriangle3d& Tri)
 	{
 		Tri.V[0] = VertexToPosnFn(Mesh, TriVertIDs.A);
@@ -162,12 +170,16 @@ bool WalkMeshPlanar(const FDynamicMesh3* Mesh, int StartTri, int StartVID, FVect
 
 	// TODO: vertex/edge snapping?
 	TArray<TPair<FMeshSurfacePoint, FWalkIndices>> ComputedPointsAndSources;
-	// TODO: switch this to a priority queue where distance to end is stored alongside, and we always pick the closest to goal ...
+
+	// This is treated as a heap, sorted to have the lowest potential path length (i.e. sum of the path length so far and
+	// the distance to end point) at the top, so that we can pop off of it to grow the best candidate path.
+	// The actual payloads are indices into ComputedPointsAndSources.
 	TArray<FIndexDistance> UnexploredEnds;
-	int BestKnownEnd = -1;
+
 	TSet<int> ExploredTriangles, CrossedVertices;
 
-	bool bHasArrived = false;
+	// When we've found a path, the index into ComputedPointsAndSources of the best path end.
+	int BestKnownEnd = -1;
 
 	FTriangle3d CurrentTri;
 	FIndex3i StartTriVertIDs = Mesh->GetTriangle(StartTri);
@@ -197,6 +209,12 @@ bool WalkMeshPlanar(const FDynamicMesh3* Mesh, int StartTri, int StartVID, FVect
 	FVector3d ForwardsDirection = EndPt - StartPt;
 
 	int CurrentEnd = 0;
+	double CurrentPathLength = 0;
+	double CurrentDistanceToEnd = ForwardsDirection.Length();
+
+	// Our start point is our first unexplored end
+	UnexploredEnds.Add({0, CurrentPathLength, CurrentDistanceToEnd });
+
 	int IterCountSafety = 0;
 	int NumTriangles = Mesh->TriangleCount();
 	while (true)
@@ -205,6 +223,22 @@ bool WalkMeshPlanar(const FDynamicMesh3* Mesh, int StartTri, int StartVID, FVect
 		{
 			return false;
 		}
+
+		// Grab the best potential path, if there is still a viable one.
+		if (UnexploredEnds.Num())
+		{
+			FIndexDistance TopEndWithDistance;
+			UnexploredEnds.HeapPop(TopEndWithDistance);
+
+			CurrentEnd = TopEndWithDistance.Index;
+			CurrentPathLength = TopEndWithDistance.PathLength;
+			CurrentDistanceToEnd = TopEndWithDistance.DistanceToEnd;
+		}
+		else
+		{
+			return false; // failed to find path
+		}
+
 		FMeshSurfacePoint FromPt = ComputedPointsAndSources[CurrentEnd].Key;
 		FWalkIndices CurrentWalk = ComputedPointsAndSources[CurrentEnd].Value;
 		int TriID = CurrentWalk.WalkingOnTri;
@@ -212,18 +246,26 @@ bool WalkMeshPlanar(const FDynamicMesh3* Mesh, int StartTri, int StartVID, FVect
 		FIndex3i TriVertIDs = Mesh->GetTriangle(TriID);
 		SetTriVertPositions(TriVertIDs, CurrentTri);
 
-		bool OnEndTri = EndTri == TriID;
+		// Note about ending the search: we happen to know that the final step of our path will be a direct line
+		// to the destination, which means that the final path length will be CurrentPathLength + DistanceToDestination
+		// at that point. Since we've been grabbing the minimal path (for extending) by the same criteria, we know that
+		// we don't need to try extending any other paths at that point (we have the shortest).
+		// This breaks down if AcceptEndPtOutsideDist is large enough to have multiple candidate end points and we want
+		// the absolute shortest one, or if we find ourselves in a world where the distance contributed by the last step
+		// is somehow potentially greater than the distance to destination at that time ("curved" triangles of some kind);
+		// In that case we would want to keep track of a PathLengthUpperBound, which we would update once we find a path,
+		// and only stop looking once we have no more unfinished paths, or if the next path we grab can't possibly beat 
+		// the upper bound.
 
 		// if we're on a triangle that is connected to the known final vertex, end the search!
 		if (EndVertID >= 0 && TriVertIDs.Contains(EndVertID))
 		{
-			OnEndTri = true;
 			CurrentEnd = ComputedPointsAndSources.Emplace(FMeshSurfacePoint(EndVertID), FWalkIndices(EndPt, CurrentEnd, TriID));
-			bHasArrived = true;
 			BestKnownEnd = CurrentEnd;
 			break;
 		}
 
+		bool OnEndTri = EndTri == TriID;
 		bool ComputedEndPtOnTri = false;
 		if (EndVertID < 0 && EndTri == -1) // if we need to check if this is the end tri, and it could be the end tri
 		{
@@ -249,7 +291,6 @@ bool WalkMeshPlanar(const FDynamicMesh3* Mesh, int StartTri, int StartVID, FVect
 			}
 			CurrentEnd = ComputedPointsAndSources.Emplace(FMeshSurfacePoint(TriID, CurrentTriDist.TriangleBaryCoords), FWalkIndices(EndPt, CurrentEnd, TriID));
 
-			bHasArrived = true;
 			BestKnownEnd = CurrentEnd;
 			break;
 		}
@@ -258,18 +299,7 @@ bool WalkMeshPlanar(const FDynamicMesh3* Mesh, int StartTri, int StartVID, FVect
 		{
 			// note we only add explored triangles to the search at all to handle the specific case where we go 'the long way' around and back to the start triangle ... otherwise we could have just not added them to the search at all
 			// currently that case is not even possible, but if it were, it would have been handled in the above `if (OnEndTri)` so we should be able to safely kill this branch of search here
-			// TODO: this code is a copy of the code at the end of the while(true) block!  consolidate?!
-			if (UnexploredEnds.Num())
-			{
-				FIndexDistance TopEndWithDistance;
-				UnexploredEnds.HeapPop(TopEndWithDistance);
-				CurrentEnd = TopEndWithDistance.Index;
-				continue;
-			}
-			else
-			{
-				return false; // failed to find
-			}
+			continue;
 		}
 		ExploredTriangles.Add(TriID);
 
@@ -373,26 +403,19 @@ bool WalkMeshPlanar(const FDynamicMesh3* Mesh, int StartTri, int StartVID, FVect
 			}
 		}
 
+		const FVector3d& PreviousPathPoint = ComputedPointsAndSources[CurrentEnd].Value.Position;
 		for (int32 NewComputedPtIdx = InitialComputedPointsNum; NewComputedPtIdx < ComputedPointsAndSources.Num(); NewComputedPtIdx++)
 		{
-			double DistSq = DistanceSquared(EndPt, ComputedPointsAndSources[NewComputedPtIdx].Value.Position);
-			// TODO: reject cases that move us backwards if !bAllowBackwardsSearch
-			bool bIsForward = ForwardsDirection.Dot(ComputedPointsAndSources[NewComputedPtIdx].Value.Position - StartPt) >= -BackwardsTolerance;
-			ensure(bAllowBackwardsSearch || bIsForward);
-			UnexploredEnds.HeapPush({ NewComputedPtIdx, DistSq });
-		}
-		
-		// TODO: this code is a copy of the code that handles terminating the search if we hit a triangle we've already seen, above; consolidate!?
-		if (UnexploredEnds.Num())
-		{
-			FIndexDistance TopEndWithDistance;
-			UnexploredEnds.HeapPop(TopEndWithDistance);
-			CurrentEnd = TopEndWithDistance.Index;
-			continue;
-		}
-		else
-		{
-			return false; // failed to find
+			const FVector3d& CurrentPathPoint = ComputedPointsAndSources[NewComputedPtIdx].Value.Position;
+			double PathLength = CurrentPathLength + Distance(PreviousPathPoint, CurrentPathPoint);
+			double DistanceToEnd = Distance(EndPt, CurrentPathPoint);
+
+			// Theoretically we've already verified the "forward" thing above while grabbing points, but we do it again here just in case.
+			bool bIsForward = ForwardsDirection.Dot(CurrentPathPoint - StartPt) >= -BackwardsTolerance;
+			if (ensure(bAllowBackwardsSearch || bIsForward))
+			{
+				UnexploredEnds.HeapPush({ NewComputedPtIdx, PathLength, DistanceToEnd });
+			}
 		}
 	}
 
