DynamicMesh: handle some failure cases in MeshBevel. These two fixes resolve many bevel failure cases, particularly those introduced by repeat bevels.

In some edge configurations two 'terminator' vertices (at the end of open bevel edge paths) are directly connected by a non-bevel edge that we will "open", and this creates a quad-shaped hole instead of a triangle-shaped hole (which was already handled). Detecting this case will arise is simpler before we change the topology, so it is done in BuildVertexSets(), stored in FBevelVertex::ConnectedBevelVertex, and used during mesh generation to call new AppendTerminatorVertexPairQuad function.

In BuildTerminatorVertex(), the search for a terminator split-edge may fail if the filtered triangle list ends up non-contiguous, preventing the vertex from being beveled (ie it is "stuck"). Code was added to both handle that case, and also try to force the list to stay contiguous.

diff --git a/Engine/Plugins/Runtime/GeometryProcessing/Source/DynamicMesh/Private/Operations/MeshBevel.cpp b/Engine/Plugins/Runtime/GeometryProcessing/Source/DynamicMesh/Private/Operations/MeshBevel.cpp
--- a/Engine/Plugins/Runtime/GeometryProcessing/Source/DynamicMesh/Private/Operations/MeshBevel.cpp
+++ b/Engine/Plugins/Runtime/GeometryProcessing/Source/DynamicMesh/Private/Operations/MeshBevel.cpp
@@ -384,6 +384,39 @@ void FMeshBevel::BuildVertexSets(const FDynamicMesh3& Mesh)
 			return;
 		}
 	}
+
+	// At a 'Terminator' vertex we are going to split the one-ring and fill the bevel-side with a quad, which
+	// will usually leave a single-triangle hole. However if two Terminator vertices are directly connected
+	// via a non-bevel mesh edge, the two triangles will be connected, ie the hole is quad-shaped. It's easier
+	// to detect this case here /before/ we split the mesh up into disconnected parts...
+	for (FBevelVertex& Vertex : Vertices)
+	{
+		if (Vertex.VertexType == EBevelVertexType::TerminatorVertex)
+		{
+			int32 OtherVertexID = Vertex.TerminatorInfo.B;
+			int32* OtherBevelVtxIdx = VertexIDToIndexMap.Find(OtherVertexID);
+			if (OtherBevelVtxIdx != nullptr)
+			{
+				// does other vertex have to be a terminator? or can this also happen w/ a junction?
+				FBevelVertex& OtherVertex = Vertices[*OtherBevelVtxIdx];
+				if (OtherVertex.VertexType == EBevelVertexType::TerminatorVertex)
+				{
+					// want to skip this if the ring-split edge is already a bevel edge
+					int32 MeshEdgeID = Mesh.FindEdge(Vertex.VertexID, OtherVertex.VertexID);
+					MESH_BEVEL_DEBUG_CHECK(MeshEdgeID >= 0);
+					if (Mesh.IsEdge(MeshEdgeID) && 
+						Vertex.TerminatorInfo.A == MeshEdgeID &&
+						OtherVertex.TerminatorInfo.A == MeshEdgeID &&		// do we need the other vertex to use the same edge here?  (is this actually a hard constraint on that edge that we should be enforcing??)
+						Vertex.IncomingBevelMeshEdges.Contains(MeshEdgeID) == false )
+					{
+						Vertex.ConnectedBevelVertex = *OtherBevelVtxIdx;
+					}
+				}
+
+			}
+		}
+	}
+
 }
 
 
@@ -480,12 +513,32 @@ void FMeshBevel::BuildTerminatorVertex(FBevelVertex& Vertex, const FDynamicMesh3
 	// in the "furthest" face (perpendicular to the edge)
 	FIndex2i IncomingEdgeT = Mesh.GetEdgeT(IncomingEdgeID);
 	FIndex2i IncomingEdgeGroups(Mesh.GetTriangleGroup(IncomingEdgeT.A), Mesh.GetTriangleGroup(IncomingEdgeT.B));
-	TArray<int32> VertexGroups;
-	Mesh.GetAllVertexGroups(Vertex.VertexID, VertexGroups);
+
+	// Find the first index of Vertex.SortedTriangles that we want to remove.
+	// This should not be able to fail but if it does, we will just start at 0 
+	// and maybe end up with some bevel failures below
+	int32 NumTriangles = Vertex.SortedTriangles.Num();
+	int32 StartIndex = 0;
+	for (int32 k = 0; k < NumTriangles; ++k)
+	{
+		int32 gid = Mesh.GetTriangleGroup(Vertex.SortedTriangles[k]);
+		if (IncomingEdgeGroups.Contains(gid))
+		{
+			StartIndex = k;
+			break;
+		}
+	}
+
 	TArray<int32> OtherGroupTris;	// sorted wedge of triangles that are not in either of the groups connected to incoming edge
 	TArray<int32> OtherGroups;		// list of group IDs encountered, in-order
-	for (int32 tid : Vertex.SortedTriangles)
-	{
+	for ( int32 k = 0; k < NumTriangles; ++k )
+	{
+		// Vertex.SortedTriangles was ordered, ie sequential triangles were connected, but if we filter
+		// out some triangles it may no longer be sequential in OtherGroupTris, leading to badness. 
+		// But the two groups we are removing should be contiguous, so if we start there, then the
+		// remaining tris should be contiguous. StartIndex found above should give us that triangle.
+		int32 ShiftedIndex = (StartIndex + k) % NumTriangles;
+		int32 tid = Vertex.SortedTriangles[ShiftedIndex];
 		int32 gid = Mesh.GetTriangleGroup(tid);
 		if (IncomingEdgeGroups.Contains(gid) == false)
 		{
@@ -493,6 +546,7 @@ void FMeshBevel::BuildTerminatorVertex(FBevelVertex& Vertex, const FDynamicMesh3
 			OtherGroups.AddUnique(gid);
 		}
 	}
+	int32 NumRemainingTris = OtherGroupTris.Num();
 
 	// Determine which edge to split at in the "other" group triangles. If we only have one group
 	// then we can try to pick the "middlest" edge. The worst case is when there is only one triangle,
@@ -521,18 +575,37 @@ void FMeshBevel::BuildTerminatorVertex(FBevelVertex& Vertex, const FDynamicMesh3
 		}
 		else
 		{
+			// try using the 'middlest' triangle as the 'middlest' edge
 			// TODO: should compute opening angles here and pick the edge closest to the middle of the angular span!!
 			int32 j = OtherGroupTris.Num() / 2;
 			RingSplitEdgeID = FindSharedEdgeInTriangles(Mesh, OtherGroupTris[j], OtherGroupTris[j+1]);
+
+			// If the OtherGroupTris list ended up being not contiguous (see above for how that can happen), then
+			// it's possible that (j) and (j+1) are not connected and the share-edge search will fail. In that case
+			// we will just linear-search for two connected triangles. If this fails then this vertex will not be bevelled.
+			if (RingSplitEdgeID == -1)
+			{
+				for (int32 k = 0; k < NumRemainingTris; ++k)
+				{
+					RingSplitEdgeID = FindSharedEdgeInTriangles(Mesh, OtherGroupTris[k], OtherGroupTris[(k+1)%NumRemainingTris]);
+					if (RingSplitEdgeID != -1)
+					{
+						break;
+					}
+				}
+			}
 		}
 	}
 	else
 	{
-		for (int32 k = 1; k < OtherGroupTris.Num(); ++k)
+		// Search for two adjacent triangles in different groups that have a shared edge. This search
+		// may need to wrap around if we got unlucky in the triangle ordering in OtherGroupTris
+		for (int32 k = 0; k < OtherGroupTris.Num(); ++k)
 		{
-			if (Mesh.GetTriangleGroup(OtherGroupTris[k-1]) != Mesh.GetTriangleGroup(OtherGroupTris[k]))
+			int32 TriangleA = OtherGroupTris[k], TriangleB = OtherGroupTris[(k+1)%NumRemainingTris];
+			if (Mesh.GetTriangleGroup(TriangleA) != Mesh.GetTriangleGroup(TriangleB))
 			{
-				RingSplitEdgeID = FindSharedEdgeInTriangles(Mesh, OtherGroupTris[k-1], OtherGroupTris[k]);
+				RingSplitEdgeID = FindSharedEdgeInTriangles(Mesh, TriangleA, TriangleB);
 				Vertex.NewGroupID = -1;		// allocate a new group for this triangle, this is usually what one would want
 				break;
 			}
@@ -1218,11 +1291,61 @@ void FMeshBevel::AppendTerminatorVertexTriangle(FDynamicMesh3& Mesh, FBevelVerte
 			{
 				Vertex.NewTriangles.Add(tid);
 			}
+			else
+			{
+				MESH_BEVEL_DEBUG_CHECK(false);
+			}
+		}
+		else
+		{
+			MESH_BEVEL_DEBUG_CHECK(false);
 		}
 	}
 }
 
 
+
+void FMeshBevel::AppendTerminatorVertexPairQuad(FDynamicMesh3& Mesh, FBevelVertex& Vertex0, FBevelVertex& Vertex1)
+{
+	MESH_BEVEL_DEBUG_CHECK(Vertex0.VertexType == EBevelVertexType::TerminatorVertex);
+	MESH_BEVEL_DEBUG_CHECK(Vertex1.VertexType == EBevelVertexType::TerminatorVertex);
+
+	// This is a variant of AppendTerminatorVertexTriangle that handles the case where basically two
+	// Terminator Vertices are directly connected by a non-beveled mesh edge that was used as the ring-split-edge.
+	// Since both sides were opened, we have a quad-shaped hole instead of a triangle-shaped hole, with a quad-edge
+	// at each end. The quad can be filled directly, we just need to sort out ordering/etc
+
+	// does not seem like we need to do anything w/ the TerminatorInfo here, we can get everything from wedges
+	int32 QuadEdgeID0 = Mesh.FindEdge(Vertex0.Wedges[0].WedgeVertex, Vertex0.Wedges[1].WedgeVertex);
+	int32 QuadEdgeID1 = Mesh.FindEdge(Vertex1.Wedges[0].WedgeVertex, Vertex1.Wedges[1].WedgeVertex);
+	MESH_BEVEL_DEBUG_CHECK(Mesh.IsEdge(QuadEdgeID0) && Mesh.IsEdge(QuadEdgeID1));
+	FIndex2i QuadEdgeV0 = Mesh.GetOrientedBoundaryEdgeV(QuadEdgeID0);
+	FIndex2i QuadEdgeV1 = Mesh.GetOrientedBoundaryEdgeV(QuadEdgeID1);
+
+	// make sure that the two opposing/connecting edges exist
+	MESH_BEVEL_DEBUG_CHECK(
+		Mesh.FindEdge(QuadEdgeV0.A, QuadEdgeV1.B) != IndexConstants::InvalidID &&
+		Mesh.FindEdge(QuadEdgeV0.B, QuadEdgeV1.A) != IndexConstants::InvalidID);
+
+	// BIASED? should have computed this GroupID in initial setup
+	int32 UseGroupID = (Vertex0.NewGroupID >= 0) ? Vertex0.NewGroupID : Mesh.AllocateTriangleGroup();
+
+	// quad order is V0.B, V0.A, V1.B, V1.A
+	int32 tid0 = Mesh.AppendTriangle(QuadEdgeV0.B, QuadEdgeV0.A, QuadEdgeV1.B, UseGroupID);
+	MESH_BEVEL_DEBUG_CHECK(tid0 >= 0);
+	if (Mesh.IsTriangle(tid0))
+	{
+		Vertex0.NewTriangles.Add(tid0);
+	}
+	int32 tid1 = Mesh.AppendTriangle(QuadEdgeV0.B, QuadEdgeV1.B, QuadEdgeV1.A, UseGroupID);
+	MESH_BEVEL_DEBUG_CHECK(tid1 >= 0);
+	if (Mesh.IsTriangle(tid1))
+	{
+		Vertex1.NewTriangles.Add(tid1);
+	}
+}
+
+
 void FMeshBevel::AppendEdgeQuads(FDynamicMesh3& Mesh, FBevelEdge& Edge)
 {
 	int32 NumEdges = Edge.MeshEdges.Num();
@@ -1357,12 +1480,27 @@ void FMeshBevel::CreateBevelMeshing(FDynamicMesh3& Mesh)
 		AppendLoopQuads(Mesh, Loop);
 	}
 
-	// easier to do these last so that we can use quad edge to orient the triangle
+	// easier to do terminators last so that we can use quad edge to orient the triangle
+	TSet<FIndex2i> HandledQuadVtxPairs;
 	for (FBevelVertex& Vertex : Vertices)
 	{
 		if (Vertex.VertexType == EBevelVertexType::TerminatorVertex)
 		{
-			AppendTerminatorVertexTriangle(Mesh, Vertex);
+			if (Vertex.ConnectedBevelVertex >= 0)
+			{
+				FBevelVertex& OtherVertex = Vertices[Vertex.ConnectedBevelVertex];
+				FIndex2i VtxPair(Vertex.VertexID, OtherVertex.VertexID);
+				VtxPair.Sort();
+				if (HandledQuadVtxPairs.Contains(VtxPair) == false)
+				{
+					AppendTerminatorVertexPairQuad(Mesh, Vertex, OtherVertex);
+					HandledQuadVtxPairs.Add(VtxPair);
+				}
+			}
+			else
+			{
+				AppendTerminatorVertexTriangle(Mesh, Vertex);
+			}
 		}
 	}
 }
diff --git a/Engine/Plugins/Runtime/GeometryProcessing/Source/DynamicMesh/Public/Operations/MeshBevel.h b/Engine/Plugins/Runtime/GeometryProcessing/Source/DynamicMesh/Public/Operations/MeshBevel.h
--- a/Engine/Plugins/Runtime/GeometryProcessing/Source/DynamicMesh/Public/Operations/MeshBevel.h
+++ b/Engine/Plugins/Runtime/GeometryProcessing/Source/DynamicMesh/Public/Operations/MeshBevel.h
@@ -220,6 +220,7 @@ public:
 		TArray<int32> NewTriangles;				// new triangles that make up the beveled polygon for this vertex (if NumEdges > 2)
 
 		FIndex2i TerminatorInfo;				// for TerminatorVertex type, store [EdgeID, FarVertexID] in one-ring, used to unlink/fill (see usage)
+		int32 ConnectedBevelVertex = -1;		// If set to another FBevelVertex index, then the TerminatorInfo.EdgeID directly connects to that vertex and special handling is needed
 	};
 
 
@@ -275,6 +276,7 @@ protected:
 	void CreateBevelMeshing(FDynamicMesh3& Mesh);
 	void AppendJunctionVertexPolygon(FDynamicMesh3& Mesh, FBevelVertex& Vertex);
 	void AppendTerminatorVertexTriangle(FDynamicMesh3& Mesh, FBevelVertex& Vertex);
+	void AppendTerminatorVertexPairQuad(FDynamicMesh3& Mesh, FBevelVertex& Vertex0, FBevelVertex& Vertex1);
 	void AppendEdgeQuads(FDynamicMesh3& Mesh, FBevelEdge& Edge);
 	void AppendLoopQuads(FDynamicMesh3& Mesh, FBevelLoop& Loop);
 
