Fix symmetry detection for scaled-up meshes, by making error tolerances also scale

diff --git a/Engine/Plugins/Runtime/GeometryProcessing/Source/DynamicMesh/Private/Parameterization/MeshPlanarSymmetry.cpp b/Engine/Plugins/Runtime/GeometryProcessing/Source/DynamicMesh/Private/Parameterization/MeshPlanarSymmetry.cpp
--- a/Engine/Plugins/Runtime/GeometryProcessing/Source/DynamicMesh/Private/Parameterization/MeshPlanarSymmetry.cpp
+++ b/Engine/Plugins/Runtime/GeometryProcessing/Source/DynamicMesh/Private/Parameterization/MeshPlanarSymmetry.cpp
@@ -75,6 +75,10 @@ bool FMeshPlanarSymmetry::Validate(const FDynamicMesh3* Mesh)
 	return true;
 }
 
+void FMeshPlanarSymmetry::SetErrorScale(const FAxisAlignedBox3d& Bounds)
+{
+	ErrorScale = FMath::Max(1, Bounds.MaxDim()); // Error tolerances scale up with larger meshes
+}
 
 bool FMeshPlanarSymmetry::FindPlaneAndInitialize(FDynamicMesh3* Mesh, const FAxisAlignedBox3d& Bounds, FFrame3d& SymmetryFrameOut, TArrayView<const FVector3d> PreferredNormals)
 {
@@ -84,6 +88,11 @@ bool FMeshPlanarSymmetry::FindPlaneAndInitialize(FDynamicMesh3* Mesh, const FAxi
 		return false;
 	}
 
+	SetErrorScale(Bounds);
+	double OnPlaneTolerance = OnPlaneToleranceFactor * ErrorScale;
+	double MatchVertexTolerance = MatchVertexToleranceFactor * ErrorScale;
+	double VertexHashCellSize = VertexHashCellSizeFactor * ErrorScale;
+
 	// If this mesh is a tiny mesh, then all points are on the symmetry plane ...
 	double BoundsMaxDim = Bounds.MaxDim();
 	if (BoundsMaxDim < OnPlaneTolerance)
@@ -185,7 +194,7 @@ bool FMeshPlanarSymmetry::FindPlaneAndInitialize(FDynamicMesh3* Mesh, const FAxi
 	
 	// Test a plane through MeshCentroid with the given Normal for symmetry.  While we're at it, accumulate an improved normal estimate from the matches.
 	// Returns false at the first failure to find a match, so often very fast to reject an incorrect plane
-	auto TestPlane = [Mesh, MeshCentroid, &VertexHash, this](FVector3d Normal, FVector3d& RefinedNormalOut) -> bool
+	auto TestPlane = [Mesh, MeshCentroid, &VertexHash, this, OnPlaneTolerance, MatchVertexTolerance](FVector3d Normal, FVector3d& RefinedNormalOut) -> bool
 	{
 		FVector3d SumOfPtToMirrorDeltas(0, 0, 0);
 		bool bAllOnPlane = true;
@@ -223,7 +232,8 @@ bool FMeshPlanarSymmetry::FindPlaneAndInitialize(FDynamicMesh3* Mesh, const FAxi
 	};
 
 	// Check whether the centroids already span a plane; if so, we already have the mirror plane candidate
-	// We use a relatively-large tolerance for this check -- if the points are too close together, the plane estimate could be unstable ...
+	// Note that the ExtremePoints algorithm internally scales its tolerance by the scale of the input,
+	// so we shouldn't use our ErrorScale-scaled tolerances with it.
 	constexpr double ExtremePointsTolerance = UE_DOUBLE_KINDA_SMALL_NUMBER;
 	FExtremePoints3d CentroidsExtremePoints(AllCentroids.Num(), [&AllCentroids](int32 Idx) {return AllCentroids[Idx];}, [](int32) {return true;}, ExtremePointsTolerance);
 	if (CentroidsExtremePoints.Dimension > 1)
@@ -364,7 +374,7 @@ bool FMeshPlanarSymmetry::FindPlaneAndInitialize(FDynamicMesh3* Mesh, const FAxi
 			{
 				int32 VID = PairCandidateComponentsArray[Cluster.StartIdx + WithinClusterIdx];
 				return Mesh->GetVertex(VID);
-			}, [](int32 Idx) { return true; }, UE_KINDA_SMALL_NUMBER);
+			}, [](int32 Idx) { return true; }, ExtremePointsTolerance);
 		// helper to avoid repeating evaluation of extreme points
 		auto IsPointExtreme = [ExtremePoints](int32 WithinClusterIdx, int32 NumExtremeToConsider)
 		{
@@ -476,6 +486,9 @@ bool FMeshPlanarSymmetry::AssignMatches(const FDynamicMesh3* Mesh, const TPointH
 	// Whether to reattempt matches w/out feature-based filtering, if a match-failure would break symmetry overall
 	constexpr bool bReattemptMatchesOnFailure = true;
 
+	double OnPlaneTolerance = OnPlaneToleranceFactor * ErrorScale;
+	double MatchVertexTolerance = MatchVertexToleranceFactor * ErrorScale;
+
 	// Compute unique matching vertices on mirror side (negative side of plane) for positive-side vertices
 	ParallelFor(PositiveVerts.Num(), [&](int32 k)
 		{
@@ -495,7 +508,7 @@ bool FMeshPlanarSymmetry::AssignMatches(const FDynamicMesh3* Mesh, const TPointH
 				, [Feature, VertexID, this, &InvariantFeatures, MatchFeaturesToleranceSqr](const int32& OtherVID) -> bool // Ignore filter fn returns true for vertices we want to skip
 				{
 					return OtherVID == VertexID // don't match to self
-						|| Vertices[OtherVID].bIsSourceVertex // don't match to on-plane or negative-side vertices
+						|| Vertices[OtherVID].bIsSourceVertex // don't match to other positive vertices
 						|| FVector3d::DistSquared(InvariantFeatures[OtherVID], Feature) > MatchFeaturesToleranceSqr; // don't match if the features don't match
 				}
 				);
@@ -515,7 +528,7 @@ bool FMeshPlanarSymmetry::AssignMatches(const FDynamicMesh3* Mesh, const TPointH
 					, [Feature, VertexID, this, &InvariantFeatures, MatchFeaturesToleranceSqr](const int32& OtherVID) -> bool // Ignore filter fn returns true for vertices we want to skip
 					{
 						return OtherVID == VertexID // don't match to self
-							|| Vertices[OtherVID].bIsSourceVertex; // don't match to on-plane or negative-side vertices
+							|| Vertices[OtherVID].bIsSourceVertex; // don't match to other positive vertices
 					}
 					);
 				NearestVID = FeaturelessNearestInfo.Key;
@@ -564,6 +577,8 @@ bool FMeshPlanarSymmetry::Initialize(FDynamicMesh3* Mesh, const FAxisAlignedBox3
 		return false;
 	}
 
+	SetErrorScale(Bounds);
+
 	// Compute vertex hash and features
 	FVector3d MeshCentroid;
 	TArray<FVector3d> Features;
@@ -571,7 +586,7 @@ bool FMeshPlanarSymmetry::Initialize(FDynamicMesh3* Mesh, const FAxisAlignedBox3
 
 	// Quick reject test: MeshCentroid must be on desired symmetry plane
 	double CentroidPlaneSignedDist = (MeshCentroid - SymmetryFrameIn.Origin).Dot(SymmetryFrameIn.Z());
-	if (FMath::Abs(CentroidPlaneSignedDist) > OnPlaneTolerance)
+	if (FMath::Abs(CentroidPlaneSignedDist) > OnPlaneToleranceFactor * ErrorScale)
 	{
 		return false;
 	}
@@ -579,7 +594,7 @@ bool FMeshPlanarSymmetry::Initialize(FDynamicMesh3* Mesh, const FAxisAlignedBox3
 	// Hash the vertices.
 	// TODO: For this case where we only test one plane, we could just hash the vertices w/ plane signed dist <= FMathd::ZeroTolerance ...
 	// (but that would be better done inside AssignMatches, where it's clearer exactly how the hash will be used, so code would need some restructuring)
-	TPointHashGrid3d<int32> VertexHash(VertexHashCellSize, INDEX_NONE);
+	TPointHashGrid3d<int32> VertexHash(VertexHashCellSizeFactor * ErrorScale, INDEX_NONE);
 	VertexHash.Reserve(Mesh->VertexCount());
 	for (int32 VID : Mesh->VertexIndicesItr())
 	{
diff --git a/Engine/Plugins/Runtime/GeometryProcessing/Source/DynamicMesh/Public/Parameterization/MeshPlanarSymmetry.h b/Engine/Plugins/Runtime/GeometryProcessing/Source/DynamicMesh/Public/Parameterization/MeshPlanarSymmetry.h
--- a/Engine/Plugins/Runtime/GeometryProcessing/Source/DynamicMesh/Public/Parameterization/MeshPlanarSymmetry.h
+++ b/Engine/Plugins/Runtime/GeometryProcessing/Source/DynamicMesh/Public/Parameterization/MeshPlanarSymmetry.h
@@ -150,11 +150,17 @@ private:
 	
 	// Point must be within this distance from the symmetry plane to be considered "on" the plane.
 	// Note: Vertices can still be matched below this tolerance; this is just the distance at which a lack of match is considered a lack of symmetry
-	constexpr static double OnPlaneTolerance = (double)FMathf::ZeroTolerance * 10;
-	constexpr static double MatchVertexTolerance = OnPlaneTolerance * 10;
+	// Note: Each Tolerance Factor will be multiplied by "ErrorScale", so that it can appropriately scale up for larger inputs.  SetErrorScale() must be called before using these tolerances.
+	constexpr static double OnPlaneToleranceFactor = (double)FMathf::ZeroTolerance * .5;
+	constexpr static double MatchVertexToleranceFactor = OnPlaneToleranceFactor * 2;
 	// Note performance of vertex hashing is much better when VertexHashCellSize is large enough that most hash lookups only need to look at a single cell
-	constexpr static double VertexHashCellSize = MatchVertexTolerance * 10;
+	constexpr static double VertexHashCellSizeFactor = MatchVertexToleranceFactor * 10;
+	// Note for features we specify a direct tolerance instead of a factor, as features are already in their own re-scaled space
 	constexpr static double MatchFeaturesTolerance = UE_DOUBLE_KINDA_SMALL_NUMBER;
+	// Internal error scale; for large meshes, we need larger tolerances because floating point values will be less accurate
+	double ErrorScale = 1;
+
+	void SetErrorScale(const FAxisAlignedBox3d& Bounds);
 
 };
 
