GeometryTools: MeshConstraintsUtil.  convert ensure to simple conditional.

diff --git a/Engine/Plugins/Runtime/GeometryProcessing/Source/DynamicMesh/Private/MeshConstraintsUtil.cpp b/Engine/Plugins/Runtime/GeometryProcessing/Source/DynamicMesh/Private/MeshConstraintsUtil.cpp
--- a/Engine/Plugins/Runtime/GeometryProcessing/Source/DynamicMesh/Private/MeshConstraintsUtil.cpp
+++ b/Engine/Plugins/Runtime/GeometryProcessing/Source/DynamicMesh/Private/MeshConstraintsUtil.cpp
@@ -322,7 +322,7 @@ void FMeshConstraintsUtil::SetBoundaryConstraintsWithProjection(
 		Topology = MakeUnique<FGroupTopology>(&Mesh, true);		
 		break;
 	case EBoundaryType::MaterialID:
-		if (!ensure(Mesh.HasAttributes() && Mesh.Attributes()->GetMaterialID() != nullptr))
+		if (!(Mesh.HasAttributes() && Mesh.Attributes()->GetMaterialID() != nullptr))
 		{
 			return;
 		}
