MeshModelingTools: Fix group topology creating duplicate group edges for loops.

diff --git a/Engine/Plugins/Runtime/GeometryProcessing/Source/DynamicMesh/Private/GroupTopology.cpp b/Engine/Plugins/Runtime/GeometryProcessing/Source/DynamicMesh/Private/GroupTopology.cpp
--- a/Engine/Plugins/Runtime/GeometryProcessing/Source/DynamicMesh/Private/GroupTopology.cpp
+++ b/Engine/Plugins/Runtime/GeometryProcessing/Source/DynamicMesh/Private/GroupTopology.cpp
@@ -594,6 +594,7 @@ bool FGroupTopology::GenerateBoundaryAndGroupEdges(FGroup& Group,
 				Edge.Span.InitializeFromEdges(Loop.Edges);
 				Edge.EndpointCorners = FIndex2i::Invalid();
 				EdgeIndex = Edges.Add(Edge);
+				GroupEdgeMinEidToGroupEdgeID.Add(MinEid, EdgeIndex);
 			}
 			Boundary.GroupEdges.Add(EdgeIndex);
 			continue;
