UVEditor: Performance tuning to improve runtime speed of rebuilding the UV overlays and meshes when edits or channel switches are performed. This significantly reduces the allocation and search churn when switching UV channels on high poly count meshes.

This review duplicates and replaces the original review https://p4-swarm.epicgames.net/reviews/21689505 on //UE5/Main, primarily to bring it directly into 5.1 instead of robomerging from Main so close to the hardlock date. The issues discussed in the original review should now be addressed here.

diff --git a/Engine/Plugins/Runtime/GeometryProcessing/Source/DynamicMesh/Private/Parameterization/UVUnwrapMeshUtil.cpp b/Engine/Plugins/Runtime/GeometryProcessing/Source/DynamicMesh/Private/Parameterization/UVUnwrapMeshUtil.cpp
--- a/Engine/Plugins/Runtime/GeometryProcessing/Source/DynamicMesh/Private/Parameterization/UVUnwrapMeshUtil.cpp
+++ b/Engine/Plugins/Runtime/GeometryProcessing/Source/DynamicMesh/Private/Parameterization/UVUnwrapMeshUtil.cpp
@@ -5,6 +5,7 @@
 #include "DynamicMesh/DynamicMesh3.h"
 #include "DynamicMesh/DynamicMeshAttributeSet.h"
 #include "Util/UniqueIndexSet.h"
+#include "Algo/Unique.h"
 
 using namespace UE::Geometry;
 
@@ -12,15 +13,15 @@ namespace UVUnwrapMeshUtilLocals
 {
 	/**
 	 * Helper function to update the triangles in an unwrap mesh based on an overlay.
-	 * 
+	 *
 	 * When updating a particular set of triangles, RemovalTriIterator and InsertionTriIterator
 	 * should be the same (the triangles to update). When updating all triangles, RemovalTriIterator
 	 * should go over all of the triangles in the destination (which may include ones not in the source)
 	 * and InsertionTriIterator should go over all of the triangles in the source (which may include ones
 	 * not in the destination).
-	 */ 
+	 */
 	template <typename TriIteratorType>
-	void UpdateUnwrapTriangles(const FDynamicMeshUVOverlay& UVOverlayIn, const TriIteratorType& RemovalTriIterator,
+	void UpdateUnwrapTriangles(const FDynamicMeshUVOverlay& UVOverlayIn, const TriIteratorType& RemovalTriIterator, int32 MaxRemoveCount,
 		const TriIteratorType& InsertionTriIterator, FDynamicMesh3& UnwrapMeshOut)
 	{
 		// Updating triangles is a little messy. To be able to handle arbitrary remeshing, we
@@ -32,7 +33,8 @@ namespace UVUnwrapMeshUtilLocals
 		// but ended up not isolated.
 
 		// Remove tris and keep track of potentially isolated elements
-		TSet<int32> PotentiallyIsolatedElements;
+		TArray<int32> PotentiallyIsolatedElements;
+		PotentiallyIsolatedElements.Reserve(MaxRemoveCount * 3);
 		for (int32 Tid : RemovalTriIterator)
 		{
 			if (!UnwrapMeshOut.IsTriangle(Tid))
@@ -49,7 +51,9 @@ namespace UVUnwrapMeshUtilLocals
 				PotentiallyIsolatedElements.Add(PrevTriangle[i]);
 			}
 			UnwrapMeshOut.RemoveTriangle(Tid, false);
-		}
+		}		
+		PotentiallyIsolatedElements.Sort();
+		PotentiallyIsolatedElements.SetNum(Algo::Unique(PotentiallyIsolatedElements));
 
 		// Reinsert new tris
 		UnwrapMeshOut.BeginUnsafeTrianglesInsert();
@@ -145,7 +149,7 @@ void UVUnwrapMeshUtil::GenerateUVUnwrapMesh(
 	UnwrapMeshOut.EndUnsafeTrianglesInsert();
 }
 
-void UVUnwrapMeshUtil::UpdateUVUnwrapMesh(const FDynamicMeshUVOverlay& UVOverlayIn, 
+void UVUnwrapMeshUtil::UpdateUVUnwrapMesh(const FDynamicMeshUVOverlay& UVOverlayIn,
 	FDynamicMesh3& UnwrapMeshOut, TFunctionRef<FVector3d(const FVector2f&)> UVToVertPosition,
 	const TArray<int32>* ChangedElementIDs, const TArray<int32>* ChangedTids)
 {
@@ -160,6 +164,8 @@ void UVUnwrapMeshUtil::UpdateUVUnwrapMesh(const FDynamicMeshUVOverlay& UVOverlay
 
 	auto UpdateVertPositions = [&UVOverlayIn, &UnwrapMeshOut, UVToVertPosition, UnwrapMeshUVOverlay](const auto& ElementIterator)
 	{
+		UnwrapMeshUVOverlay->BeginUnsafeElementsInsert();
+		UnwrapMeshOut.BeginUnsafeVerticesInsert();
 		for (int32 ElementID : ElementIterator)
 		{
 			if (!ensure(UVOverlayIn.IsElement(ElementID)))
@@ -181,7 +187,7 @@ void UVUnwrapMeshUtil::UpdateUVUnwrapMesh(const FDynamicMeshUVOverlay& UVOverlay
 			}
 			else
 			{
-				UnwrapMeshOut.InsertVertex(ElementID, UVToVertPosition(ElementValue));
+				UnwrapMeshOut.InsertVertex(ElementID, UVToVertPosition(ElementValue), true);
 			}
 
 			// Update the unwrap overlay.
@@ -191,9 +197,11 @@ void UVUnwrapMeshUtil::UpdateUVUnwrapMesh(const FDynamicMeshUVOverlay& UVOverlay
 			}
 			else
 			{
-				UnwrapMeshUVOverlay->InsertElement(ElementID, &ElementValue.X);
+				UnwrapMeshUVOverlay->InsertElement(ElementID, &ElementValue.X, true);
 			}
 		}
+		UnwrapMeshUVOverlay->EndUnsafeElementsInsert();
+		UnwrapMeshOut.EndUnsafeVerticesInsert();
 	};
 
 	if (ChangedElementIDs)
@@ -207,11 +215,11 @@ void UVUnwrapMeshUtil::UpdateUVUnwrapMesh(const FDynamicMeshUVOverlay& UVOverlay
 
 	if (ChangedTids)
 	{
-		UpdateUnwrapTriangles(UVOverlayIn, *ChangedTids, *ChangedTids, UnwrapMeshOut);
+		UpdateUnwrapTriangles(UVOverlayIn, *ChangedTids, ChangedTids->Num(), *ChangedTids, UnwrapMeshOut);
 	}
 	else
 	{
-		UpdateUnwrapTriangles(UVOverlayIn, UnwrapMeshOut.TriangleIndicesItr(), 
+		UpdateUnwrapTriangles(UVOverlayIn, UnwrapMeshOut.TriangleIndicesItr(), UnwrapMeshOut.TriangleCount(),
 			UVOverlayIn.GetParentMesh()->TriangleIndicesItr(), UnwrapMeshOut);
 	}
 }
@@ -232,6 +240,8 @@ void UVUnwrapMeshUtil::UpdateUVUnwrapMesh(const FDynamicMesh3& SourceUnwrapMesh,
 
 	auto UpdateVerts = [&SourceUnwrapMesh, &DestUnwrapMesh, SourceOverlay, DestOverlay](const auto& VidIterator)
 	{
+		DestOverlay->BeginUnsafeElementsInsert();
+		DestUnwrapMesh.BeginUnsafeVerticesInsert();
 		for (int32 Vid : VidIterator)
 		{
 			if (!ensure(SourceUnwrapMesh.IsVertex(Vid)))
@@ -247,11 +257,13 @@ void UVUnwrapMeshUtil::UpdateUVUnwrapMesh(const FDynamicMesh3& SourceUnwrapMesh,
 			}
 			else
 			{
-				DestUnwrapMesh.InsertVertex(Vid, SourceUnwrapMesh.GetVertex(Vid));
+				DestUnwrapMesh.InsertVertex(Vid, SourceUnwrapMesh.GetVertex(Vid), true);
 				FVector2f ElementValue = SourceOverlay->GetElement(Vid);
-				DestOverlay->InsertElement(Vid, &ElementValue.X);
+				DestOverlay->InsertElement(Vid, &ElementValue.X, true);
 			}
 		}
+		DestUnwrapMesh.EndUnsafeVerticesInsert();
+		DestOverlay->EndUnsafeElementsInsert();
 	};
 	if (ChangedVids)
 	{
@@ -264,11 +276,12 @@ void UVUnwrapMeshUtil::UpdateUVUnwrapMesh(const FDynamicMesh3& SourceUnwrapMesh,
 
 	if (ChangedConnectivityTids)
 	{
-		UpdateUnwrapTriangles(*SourceOverlay, *ChangedConnectivityTids, *ChangedConnectivityTids, DestUnwrapMesh);
+		UpdateUnwrapTriangles(*SourceOverlay, *ChangedConnectivityTids, ChangedConnectivityTids->Num(),
+			*ChangedConnectivityTids, DestUnwrapMesh);
 	}
 	else
 	{
-		UpdateUnwrapTriangles(*SourceOverlay, DestUnwrapMesh.TriangleIndicesItr(), 
+		UpdateUnwrapTriangles(*SourceOverlay, DestUnwrapMesh.TriangleIndicesItr(), DestUnwrapMesh.TriangleCount(),
 			SourceUnwrapMesh.TriangleIndicesItr(), DestUnwrapMesh);
 	}
 }
@@ -285,6 +298,7 @@ void UVUnwrapMeshUtil::UpdateUVOverlayFromUnwrapMesh(
 
 	auto UpdateElements = [&UnwrapMeshIn, &UVOverlayOut, VertPositionToUV](const auto& VidIterator)
 	{
+		UVOverlayOut.BeginUnsafeElementsInsert();
 		for (int32 Vid : VidIterator)
 		{
 			if (!ensure(UnwrapMeshIn.IsVertex(Vid)))
@@ -301,9 +315,10 @@ void UVUnwrapMeshUtil::UpdateUVOverlayFromUnwrapMesh(
 			}
 			else
 			{
-				UVOverlayOut.InsertElement(Vid, &UV.X);
+				UVOverlayOut.InsertElement(Vid, &UV.X, true);
 			}
 		}
+		UVOverlayOut.EndUnsafeElementsInsert();
 	};
 	if (ChangedVids)
 	{
@@ -344,7 +359,7 @@ void UVUnwrapMeshUtil::UpdateUVOverlayFromUnwrapMesh(
 			}
 		}
 
-		for (int32 Tid: InsertTriIterator)
+		for (int32 Tid : InsertTriIterator)
 		{
 			// Update the triangle
 			if (UnwrapMeshIn.IsTriangle(Tid))
@@ -379,7 +394,7 @@ void UVUnwrapMeshUtil::UpdateUVOverlayFromUnwrapMesh(
 
 void UVUnwrapMeshUtil::UpdateOverlayFromOverlay(
 	const FDynamicMeshUVOverlay& OverlayIn, FDynamicMeshUVOverlay& OverlayOut,
-	bool bMeshesHaveSameTopology, const TArray<int32>* ChangedElements, 
+	bool bMeshesHaveSameTopology, const TArray<int32>* ChangedElements,
 	const TArray<int32>* ChangedConnectivityTids)
 {
 	if (!ChangedElements && !ChangedConnectivityTids && bMeshesHaveSameTopology)
@@ -390,6 +405,7 @@ void UVUnwrapMeshUtil::UpdateOverlayFromOverlay(
 
 	auto UpdateElements = [&OverlayIn, &OverlayOut](const auto& ElementIterator)
 	{
+		OverlayOut.BeginUnsafeElementsInsert();
 		for (int32 ElementID : ElementIterator)
 		{
 			if (!ensure(OverlayIn.IsElement(ElementID)))
@@ -405,9 +421,10 @@ void UVUnwrapMeshUtil::UpdateOverlayFromOverlay(
 			}
 			else
 			{
-				OverlayOut.InsertElement(ElementID, &ElementValue.X);
+				OverlayOut.InsertElement(ElementID, &ElementValue.X, true);
 			}
 		}
+		OverlayOut.EndUnsafeElementsInsert();
 	};
 
 	if (ChangedElements)
@@ -537,7 +554,7 @@ bool UVUnwrapMeshUtil::DoesUnwrapMatchOverlay(const FDynamicMeshUVOverlay& Overl
 
 // Explicit instantiations
 template bool UVUnwrapMeshUtil::DoesUnwrapMatchOverlay<EValidityCheckFailMode::ReturnOnly>
-	(const FDynamicMeshUVOverlay& Overlay, const FDynamicMesh3& UnwrapMesh,
+(const FDynamicMeshUVOverlay& Overlay, const FDynamicMesh3& UnwrapMesh,
 	TFunctionRef<FVector3d(const FVector2f&)> UVToVertPosition, double Tolerance);
 template bool UVUnwrapMeshUtil::DoesUnwrapMatchOverlay<EValidityCheckFailMode::Check>
 (const FDynamicMeshUVOverlay& Overlay, const FDynamicMesh3& UnwrapMesh,
