GeometryProcessing: harden FMeshJoinLoops to handle some cases that users can create in geometry scriptingf

diff --git a/Engine/Plugins/Runtime/GeometryProcessing/Source/DynamicMesh/Private/Operations/JoinMeshLoops.cpp b/Engine/Plugins/Runtime/GeometryProcessing/Source/DynamicMesh/Private/Operations/JoinMeshLoops.cpp
--- a/Engine/Plugins/Runtime/GeometryProcessing/Source/DynamicMesh/Private/Operations/JoinMeshLoops.cpp
+++ b/Engine/Plugins/Runtime/GeometryProcessing/Source/DynamicMesh/Private/Operations/JoinMeshLoops.cpp
@@ -34,6 +34,11 @@ static int32 GetFirstEdgeTri(const FDynamicMesh3* Mesh, int32 VertexA, int32 Ver
 template<typename TriangleIDFuncType>
 static double CalculateAverageUVScale(const FDynamicMeshUVOverlay* UVOverlay, int32 NumTriangles, TriangleIDFuncType GetTriangleIDFunc)
 {
+	if (!UVOverlay)
+	{
+		return 1.0;
+	}
+
 	const FDynamicMesh3* Mesh = UVOverlay->GetParentMesh();
 
 	double AverageScale = 0;
@@ -91,13 +96,17 @@ bool FJoinMeshLoops::Apply()
 	for ( int32 qi = 0; qi < NumQuads; ++qi)
 	{
 		FIndex2i& Quad = JoinQuads[qi];
-
-		JoinTriangles.Add(Quad.A);
-		JoinTriangles.Add(Quad.B);
-
 		QuadGroups[qi] = NewGroupID;
-		Mesh->SetTriangleGroup(Quad.A, NewGroupID);
-		Mesh->SetTriangleGroup(Quad.B, NewGroupID);
+		if (Mesh->IsTriangle(Quad.A))
+		{
+			JoinTriangles.Add(Quad.A);
+			Mesh->SetTriangleGroup(Quad.A, NewGroupID);
+		}
+		if (Mesh->IsTriangle(Quad.B))
+		{
+			JoinTriangles.Add(Quad.B);
+			Mesh->SetTriangleGroup(Quad.B, NewGroupID);
+		}
 	}
 
 
@@ -108,10 +117,8 @@ bool FJoinMeshLoops::Apply()
 	}
 
 
-	if (Mesh->Attributes() && Mesh->Attributes()->PrimaryUV() != nullptr)
+	if (Mesh->Attributes() && UVOverlay != nullptr)
 	{
-		//FDynamicMeshUVOverlay* UVOverlay = Mesh->Attributes()->PrimaryUV();
-
 		TArray<int32> UVLoopA, UVLoopB;
 		UVLoopA.SetNum(NV + 1);
 		UVLoopB.SetNum(NV + 1);
@@ -159,9 +166,14 @@ bool FJoinMeshLoops::Apply()
 			int32 uvb = UVLoopA[qi+1];
 			int32 uvc = UVLoopB[qi];
 			int32 uvd = UVLoopB[qi + 1];
-
-			UVOverlay->SetTriangle(Quad.A, FIndex3i(uvb, uva, uvd));
-			UVOverlay->SetTriangle(Quad.B, FIndex3i(uva, uvc, uvd));
+			if (Mesh->IsTriangle(Quad.A))
+			{
+				UVOverlay->SetTriangle(Quad.A, FIndex3i(uvb, uva, uvd));
+			}
+			if (Mesh->IsTriangle(Quad.B))
+			{
+				UVOverlay->SetTriangle(Quad.B, FIndex3i(uva, uvc, uvd));
+			}
 		}
 
 	}
