ModelingTools: Fix ExtractSubloops in MeshRegionBoundaryLoops being unable to separate out some separatable subloops.

diff --git a/Engine/Plugins/Runtime/GeometryProcessing/Source/DynamicMesh/Private/MeshRegionBoundaryLoops.cpp b/Engine/Plugins/Runtime/GeometryProcessing/Source/DynamicMesh/Private/MeshRegionBoundaryLoops.cpp
--- a/Engine/Plugins/Runtime/GeometryProcessing/Source/DynamicMesh/Private/MeshRegionBoundaryLoops.cpp
+++ b/Engine/Plugins/Runtime/GeometryProcessing/Source/DynamicMesh/Private/MeshRegionBoundaryLoops.cpp
@@ -410,118 +410,71 @@ int FMeshRegionBoundaryLoops::FindLeftTurnEdge(int incoming_e, int bowtie_v, TAr
 
 
 
-// This is called when loopV contains one or more "bowtie" vertices.
-// These vertices *might* be duplicated : loopV (but not necessarily)
-// If they are, we have to break loopV into subloops that don't contain duplicates.
+// This is called when LoopV contains one or more "bowtie" vertices.
+// These vertices *might* be duplicated in LoopV (but not necessarily)
+// If they are, we have to break LoopV into subloops that don't contain duplicates.
 //
-// The list bowties contains all the possible duplicates 
-// (all v in bowties occur in loopV at least once)
+// The list Bowties contains all the possible duplicates 
+// (all v in Bowties occur in LoopV at least once)
 //
-// Currently loopE is not used, and the returned FEdgeLoop objects do not have their Edges
-// arrays initialized. Perhaps to improve : future.
-bool FMeshRegionBoundaryLoops::TryExtractSubloops(TArray<int>& loopV, const TArray<int>& loopE, const TArray<int>& bowties, TArray<FEdgeLoop>& SubLoopsOut)
+// Currently LoopE is not used.
+//
+// Note: the approach used here, which doesn't fail, should also be used in FMeshBoundaryLoops. However we're
+// going to do additional tests there with hole filling to make sure that not failing has an overall positive 
+// impact on hole filling, since we would rather fail there if the resulting subloops are more likely to be
+// bad for hole filling than good.
+bool FMeshRegionBoundaryLoops::TryExtractSubloops(TArray<int>& LoopV, const TArray<int>& LoopE, const TArray<int>& Bowties, TArray<FEdgeLoop>& SubloopsOut)
 {
-	SubLoopsOut.Reset();
-
-	// figure out which bowties we saw are actually duplicated : loopV
-	TArray<int> dupes;
-	for (int bv : bowties)
+	// We keep track of the last place we've seen a bowtie in our list. As soon as we see a bowtie a second
+	// time, we extract that subloop. As long as we do this immediately, the intervening vids will not be 
+	// duplicates even if they are bowties.
+	TMap<int32, int32> BowtieVidToLoopIndex;
+	TSet<int32> BowtieVids(Bowties);
+	for (int LoopIndex = 0; LoopIndex < LoopV.Num(); ++LoopIndex)
 	{
-		if (FMeshBoundaryLoops::CountInList(loopV, bv) > 1)
+		int32 Vid = LoopV[LoopIndex];
+		int32* SeenBowtieIndex = BowtieVidToLoopIndex.Find(Vid);
+		if (SeenBowtieIndex)
 		{
-			dupes.Add(bv);
-		}
-	}
-
-	// we might not actually have any duplicates, if we got luck. Early out : that case
-	if (dupes.Num() == 0)
-	{
-		FEdgeLoop NewLoop(Mesh);
-		NewLoop.Vertices = loopV;
-		NewLoop.Edges = loopE;
-		NewLoop.BowtieVertices = bowties;
-		SubLoopsOut.Add(NewLoop);
-		return true;
-	}
-
-	// This loop extracts subloops until we have dealt with all the
-	// duplicate vertices : loopV
-	while (dupes.Num() > 0)
-	{
-
-		// Find shortest "simple" loop, ie a loop from a bowtie to itself that
-		// does not contain any other bowties. This is an independent loop.
-		// We're doing a lot of extra work here if we only have one element : dupes...
-		int bi = 0, bv = 0;
-		int start_i = -1, end_i = -1;
-		int bv_shortest = -1; int shortest = TNumericLimits<int>::Max();
-		for (; bi < dupes.Num(); ++bi)
-		{
-			bv = dupes[bi];
-			if (FMeshBoundaryLoops::IsSimpleBowtieLoop(loopV, dupes, bv, start_i, end_i))
+			// Check that the place we saw this wasn't in an extracted subloop (in which case the vert is set to -1 in LoopV)
+			if (LoopV[*SeenBowtieIndex] >= 0)
 			{
-				int len = FMeshBoundaryLoops::CountSpan(loopV, start_i, end_i);
-				if (len < shortest)
-				{
-					bv_shortest = bv;
-					shortest = len;
-				}
+				// Now that we've seen the bowtie twice, extract the subloop
+				TArray<int32> SubloopVertices;
+				FMeshBoundaryLoops::ExtractSpan(LoopV, *SeenBowtieIndex, LoopIndex, true, SubloopVertices);
+
+				FEdgeLoop& NewLoop = SubloopsOut[SubloopsOut.Emplace()];
+				NewLoop.InitializeFromVertices(Mesh, SubloopVertices, false);
+				NewLoop.SetBowtieVertices(Bowties);
 			}
-		}
-		if (bv_shortest == -1)
-		{
-			// Cannot find a valid simple loop -- unrecoverable failure
-			return false;
-		}
 
-		if (bv != bv_shortest)
-		{
-			bv = bv_shortest;
-			// running again just to get start_i and end_i...
-			FMeshBoundaryLoops::IsSimpleBowtieLoop(loopV, dupes, bv, start_i, end_i);
+			// Update the last place we saw the bowtie, in case a subsequent loop goes through here.
+			BowtieVidToLoopIndex[Vid] = LoopIndex;
 		}
-
-		check(loopV[start_i] == bv && loopV[end_i] == bv);
-
-		FEdgeLoop loop(Mesh);
-		FMeshBoundaryLoops::ExtractSpan(loopV, start_i, end_i, true, loop.Vertices);
-		FEdgeLoop::VertexLoopToEdgeLoop(Mesh, loop.Vertices, loop.Edges);
-		loop.BowtieVertices = bowties;
-		SubLoopsOut.Add(loop);
-
-		// If there are no more duplicates of this bowtie, we can treat
-		// it like a regular vertex now
-		if (FMeshBoundaryLoops::CountInList(loopV, bv) < 2)
+		else if (BowtieVids.Contains(Vid))
 		{
-			dupes.Remove(bv);
+			BowtieVidToLoopIndex.Add(Vid, LoopIndex);
 		}
 	}
 
-	// Should have one loop left that contains duplicates. 
-	// Extract this as a separate loop
-	int nLeft = 0;
-	for (int i = 0; i < loopV.Num(); ++i)
+	// Should have one loop left.
+	TArray<int32> RemainingLoopVids;
+	for (int32 Vid : LoopV)
 	{
-		if (loopV[i] != -1)
+		if (Vid >= 0)
 		{
-			nLeft++;
+			RemainingLoopVids.Add(Vid);
 		}
 	}
-	if (nLeft > 0)
+	if (ensure(RemainingLoopVids.Num() > 2))
 	{
-		FEdgeLoop loop(Mesh);
-		loop.Vertices.SetNum(nLeft);
-		int vi = 0;
-		for (int i = 0; i < loopV.Num(); ++i)
-		{
-			if (loopV[i] != -1)
-			{
-				loop.Vertices[vi++] = loopV[i];
-			}
-		}
-		FEdgeLoop::VertexLoopToEdgeLoop(Mesh, loop.Vertices, loop.Edges);
-		loop.BowtieVertices = bowties;
-		SubLoopsOut.Add(loop);
+		FEdgeLoop& NewLoop = SubloopsOut[SubloopsOut.Emplace()];
+		NewLoop.InitializeFromVertices(Mesh, RemainingLoopVids, false);
+		NewLoop.SetBowtieVertices(Bowties);
+	}
+	else
+	{
+		return false;
 	}
 
 	return true;
