make the world-z alignment logic in dynamic mesh's grouptopology class also apply when the normal is aligned to -Z

diff --git a/Engine/Plugins/Runtime/GeometryProcessing/Source/DynamicMesh/Private/GroupTopology.cpp b/Engine/Plugins/Runtime/GeometryProcessing/Source/DynamicMesh/Private/GroupTopology.cpp
--- a/Engine/Plugins/Runtime/GeometryProcessing/Source/DynamicMesh/Private/GroupTopology.cpp
+++ b/Engine/Plugins/Runtime/GeometryProcessing/Source/DynamicMesh/Private/GroupTopology.cpp
@@ -788,14 +788,26 @@ FFrame3d FGroupTopology::GetSelectionFrame(const FGroupTopologySelection& Select
 	if (AccumCount > 0)
 	{
 		AccumulatedOrigin /= (double)AccumCount;
-		Normalize(AccumulatedNormal);
+		if (!AccumulatedNormal.Normalize())
+		{
+			AccumulatedNormal = FVector3d::UnitZ();
+		}
 
 		// We set our frame Z to be accumulated normal, and the other two axes are unconstrained, so
 		// we want to set them to something that will make our frame generally more useful. If the normal
 		// is aligned with world Z, then the entire frame might as well be aligned with world.
-		if (1 - AccumulatedNormal.Dot(FVector3d::UnitZ()) < KINDA_SMALL_NUMBER)
+		double ZAlign = AccumulatedNormal.Dot(FVector3d::UnitZ());
+		if (FMath::Abs(ZAlign) > 1 - KINDA_SMALL_NUMBER)
 		{
-			AccumulatedFrame = FFrame3d(AccumulatedOrigin, FQuaterniond::Identity());
+			if (ZAlign < 0)
+			{
+				// Rotate frame 180 degrees around X
+				AccumulatedFrame = FFrame3d(AccumulatedOrigin, FVector3d::UnitX(), -FVector3d::UnitY(), -FVector3d::UnitZ());
+			}
+			else
+			{
+				AccumulatedFrame = FFrame3d(AccumulatedOrigin, FQuaterniond::Identity());
+			}
 		}
 		else
 		{
