GeometryProcessing: Fix error with local namespace name throwing duplicate symbol warnings during builds.

#rnx

diff --git a/Engine/Plugins/Runtime/GeometryProcessing/Source/DynamicMesh/Private/Parameterization/MeshUDIMClassifier.cpp b/Engine/Plugins/Runtime/GeometryProcessing/Source/DynamicMesh/Private/Parameterization/MeshUDIMClassifier.cpp
--- a/Engine/Plugins/Runtime/GeometryProcessing/Source/DynamicMesh/Private/Parameterization/MeshUDIMClassifier.cpp
+++ b/Engine/Plugins/Runtime/GeometryProcessing/Source/DynamicMesh/Private/Parameterization/MeshUDIMClassifier.cpp
@@ -6,7 +6,7 @@
 
 using namespace UE::Geometry;
 
-namespace UVLayoutOpLocals
+namespace UDIMClassifierLocals
 {
 	FVector2f InternalUVToExternalUV(const FVector2f& UV)
 	{
@@ -48,7 +48,7 @@ FVector2i FDynamicMeshUDIMClassifier::ClassifyTrianglesToUDIM(const FDynamicMesh
 		FVector2f Vertex0, Vertex1, Vertex2;
 		UVOverlay->GetTriElements(Tid, Vertex0, Vertex1, Vertex2);
 		FVector2f BaryCenter = (Vertex0 + Vertex1 + Vertex2) / 3.0f;
-		BaryCenter = UVLayoutOpLocals::InternalUVToExternalUV(BaryCenter);
+		BaryCenter = UDIMClassifierLocals::InternalUVToExternalUV(BaryCenter);
 		UDIM = ClassifyPointToUDIM(BaryCenter);		
 		return UDIM;
 	};
