Fixed invalid array access in FDynamicMeshUVEditor::SetTriangleUVsFromConformal()

diff --git a/Engine/Plugins/Runtime/GeometryProcessing/Source/DynamicMesh/Private/Parameterization/DynamicMeshUVEditor.cpp b/Engine/Plugins/Runtime/GeometryProcessing/Source/DynamicMesh/Private/Parameterization/DynamicMeshUVEditor.cpp
--- a/Engine/Plugins/Runtime/GeometryProcessing/Source/DynamicMesh/Private/Parameterization/DynamicMeshUVEditor.cpp
+++ b/Engine/Plugins/Runtime/GeometryProcessing/Source/DynamicMesh/Private/Parameterization/DynamicMeshUVEditor.cpp
@@ -586,11 +586,12 @@ bool FDynamicMeshUVEditor::SetTriangleUVsFromConformal(const TArray<int32>& Tria
 	TUniquePtr<UE::Solvers::IConstrainedMeshUVSolver> Solver = nullptr;
 
 	FMeshBoundaryLoops Loops(&Submesh, true);
-	if (Loops.GetLoopCount() == 0) 
+	int32 LongestLoopIndex = Loops.GetLongestLoopIndex();
+	if (LongestLoopIndex == INDEX_NONE)
 	{ 
 		return false; 
 	}
-	const TArray<int32>& ConstrainLoop = Loops[Loops.GetLongestLoopIndex()].Vertices;
+	const TArray<int32>& ConstrainLoop = Loops[LongestLoopIndex].Vertices;
 	int32 LoopNum = ConstrainLoop.Num();
 
 	if (bUseSpectral) 
