GeometryCore and GeometryProcessing additions to support selection and new extrude operator
- Add TMeshQueries::GetTrianglesBounds() and ::GetVerticesBounds()
- Add FMeshNormals::RecomputeOverlayTriNormals() and ::RecomputeOverlayElementNormals(). Update ComputeOverlayNormal() to take area and angle weight parameters.
- Add new FMeshTriOrderedEdgeID type, and util functions ConvertLoopToTriOrderedEdgeLoop / ConvertTriOrderedEdgeLoopToLoop. This type is similar to the TPair-tuple used by FDynamicMeshEditor::ConvertLoopToTriVidPairSequence, but formalized and providing a function to decode back into a vertex/edge ID loop

diff --git a/Engine/Plugins/Runtime/GeometryProcessing/Source/DynamicMesh/Private/EdgeLoop.cpp b/Engine/Plugins/Runtime/GeometryProcessing/Source/DynamicMesh/Private/EdgeLoop.cpp
--- a/Engine/Plugins/Runtime/GeometryProcessing/Source/DynamicMesh/Private/EdgeLoop.cpp
+++ b/Engine/Plugins/Runtime/GeometryProcessing/Source/DynamicMesh/Private/EdgeLoop.cpp
@@ -250,4 +250,78 @@ void FEdgeLoop::VertexLoopToEdgeLoop(const FDynamicMesh3* Mesh, const TArray<int
 		int v1 = VertexLoop[(i + 1) % NV];
 		OutEdgeLoop[i] = Mesh->FindEdge(v0, v1);
 	}
-}
\ No newline at end of file
+}
+
+
+bool UE::Geometry::ConvertLoopToTriOrderedEdgeLoop(const FDynamicMesh3& Mesh, const TArray<int32>& VertexLoop,
+	const TArray<int32>& EdgeLoop, TArray<FMeshTriOrderedEdgeID>& TriOrderedEdgesLoopOut)
+{
+	return ConvertLoopToTriOrderedEdgeLoop(Mesh, VertexLoop, EdgeLoop, [](int, int32 A, int32 B) { return A;}, TriOrderedEdgesLoopOut);
+}
+
+
+bool UE::Geometry::ConvertLoopToTriOrderedEdgeLoop(const FDynamicMesh3& Mesh, const TArray<int32>& VertexLoop,
+	const TArray<int32>& EdgeLoop, 
+	TFunctionRef<int(int, int,int)> SelectEdgeTriangleFunc,
+	TArray<FMeshTriOrderedEdgeID>& TriOrderedEdgesLoopOut )
+{
+	if (!ensure(EdgeLoop.Num() == VertexLoop.Num()))
+	{
+		return false;
+	}
+
+	TriOrderedEdgesLoopOut.Reset();
+	TriOrderedEdgesLoopOut.Reserve(EdgeLoop.Num());
+	for (int32 LoopIndex = 0; LoopIndex < EdgeLoop.Num(); ++LoopIndex)
+	{
+		FIndex2i EdgeT = Mesh.GetEdgeT(EdgeLoop[LoopIndex]);
+		int32 TriangleID = SelectEdgeTriangleFunc(EdgeLoop[LoopIndex], EdgeT.A, EdgeT.B);
+		int32 FirstVertexID = VertexLoop[LoopIndex];
+		int32 SecondVertexID = VertexLoop[(LoopIndex + 1) % VertexLoop.Num()];
+
+		FIndex3i Triangle = Mesh.GetTriangle(TriangleID);
+		int32 FirstVertexIndex = (int32)IndexUtil::FindTriIndex(FirstVertexID, Triangle);
+		int32 SecondVertexIndex = (int32)IndexUtil::FindTriIndex(SecondVertexID, Triangle);
+		if (ensure(FirstVertexIndex >= 0 && SecondVertexIndex >= 0))
+		{
+			TriOrderedEdgesLoopOut.Emplace( FMeshTriOrderedEdgeID(TriangleID, FirstVertexIndex, SecondVertexIndex) );
+		}
+		else
+		{
+			return false;
+		}
+	}
+	return true;
+}
+
+
+bool UE::Geometry::ConvertTriOrderedEdgeLoopToLoop(const FDynamicMesh3& Mesh,
+	const TArray<FMeshTriOrderedEdgeID>& TriOrderedEdgesLoopOut,
+	TArray<int32>& VertexLoop, TArray<int32>* EdgeLoop)
+{
+	int32 N = TriOrderedEdgesLoopOut.Num();
+	VertexLoop.Reset();
+	VertexLoop.Reserve(N-1);
+
+	FIndex3i FirstTri = Mesh.GetTriangle(TriOrderedEdgesLoopOut[0].TriangleID);
+	VertexLoop.Add(FirstTri[TriOrderedEdgesLoopOut[0].VertIndexA]);
+	VertexLoop.Add(FirstTri[TriOrderedEdgesLoopOut[0].VertIndexB]);
+	for (int32 k = 1; k < N-1; ++k)
+	{
+		FIndex3i Tri = Mesh.GetTriangle(TriOrderedEdgesLoopOut[k].TriangleID);
+		check( Tri[TriOrderedEdgesLoopOut[k].VertIndexA] == VertexLoop.Last() );
+		VertexLoop.Add( Tri[TriOrderedEdgesLoopOut[k].VertIndexB] );
+	}
+	if ( EdgeLoop != nullptr )
+	{
+		EdgeLoop->Reset();
+		EdgeLoop->Reserve(N);
+		for (int32 k = 0; k < N; ++k)
+		{
+			int32 EdgeID = Mesh.FindEdgeFromTri( TriOrderedEdgesLoopOut[k].TriangleID, VertexLoop[k], VertexLoop[(k+1)%N] );
+			EdgeLoop->Add(EdgeID);
+		}
+	}
+	return true;
+}
+
diff --git a/Engine/Plugins/Runtime/GeometryProcessing/Source/DynamicMesh/Public/EdgeLoop.h b/Engine/Plugins/Runtime/GeometryProcessing/Source/DynamicMesh/Public/EdgeLoop.h
--- a/Engine/Plugins/Runtime/GeometryProcessing/Source/DynamicMesh/Public/EdgeLoop.h
+++ b/Engine/Plugins/Runtime/GeometryProcessing/Source/DynamicMesh/Public/EdgeLoop.h
@@ -240,5 +240,43 @@ public:
 };
 
 
+//
+// Utility functions for converting/processing edge loops
+//
+
+/**
+ * Convert the input Vertex/Edge loop into an edge loop representation that is
+ * independent of Vertex/Edge IDs. See description of FMeshTriOrderedEdgeID for
+ * more information on how this works. But essentially the resulting TriOrderedEdgesLoopOut
+ * is stable across some topological changes, in particular if the vertices of
+ * the loop change but the triangle topology remains the same (eg if a bowtie vertex
+ * is disconnected, or if the mesh is disconnected along the edge loop), the encoded
+ * loop remains valid, while an explicit vertex or edge loop would not.
+ * 
+ * @param SelectEdgeTriangleFunc This function is used to select which triangle to use for encoding at each edge
+ */
+DYNAMICMESH_API bool ConvertLoopToTriOrderedEdgeLoop(const FDynamicMesh3& Mesh,
+	const TArray<int32>& VertexLoop, const TArray<int32>& EdgeLoop,
+	TFunctionRef<int(int EdgeID, int TriangleA, int TriangleB)> SelectEdgeTriangleFunc,
+	TArray<FMeshTriOrderedEdgeID>& TriOrderedEdgesLoopOut );
+
+/**
+ * Variant of ConvertLoopToTriOrderedEdgeLoop that always encodes using EdgeTriangles.A, ie suitable for encoding open border loops
+ */
+DYNAMICMESH_API bool ConvertLoopToTriOrderedEdgeLoop(const FDynamicMesh3& Mesh,
+	const TArray<int32>& VertexLoop, const TArray<int32>& EdgeLoop,
+	TArray<FMeshTriOrderedEdgeID>& TriOrderedEdgesLoopOut );
+
+/**
+ * Recover the current Vertex and Edge loops given an encoded loop of FMeshTriOrderedEdgeID.
+ * See FMeshTriOrderedEdgeID for details on this encoding, this function basically reverses
+ * ConvertLoopToTriOrderedEdgeLoop() for the current mesh topology
+ */
+DYNAMICMESH_API bool ConvertTriOrderedEdgeLoopToLoop(const FDynamicMesh3& Mesh,
+	const TArray<FMeshTriOrderedEdgeID>& TriOrderedEdgesLoopOut,
+	TArray<int32>& VertexLoop, TArray<int32>* EdgeLoop = nullptr);
+
+
+
 } // end namespace UE::Geometry
 } // end namespace UE
diff --git a/Engine/Source/Runtime/GeometryCore/Private/DynamicMesh/MeshNormals.cpp b/Engine/Source/Runtime/GeometryCore/Private/DynamicMesh/MeshNormals.cpp
--- a/Engine/Source/Runtime/GeometryCore/Private/DynamicMesh/MeshNormals.cpp
+++ b/Engine/Source/Runtime/GeometryCore/Private/DynamicMesh/MeshNormals.cpp
@@ -397,6 +397,47 @@ bool FMeshNormals::QuickRecomputeOverlayNormals(FDynamicMesh3& Mesh, bool bInver
 	return false;
 }
 
+bool FMeshNormals::RecomputeOverlayTriNormals(FDynamicMesh3& Mesh, const TArray<int32>& Triangles, bool bWeightByArea, bool bWeightByAngle)
+{
+	if (Mesh.HasAttributes() && Mesh.Attributes()->PrimaryNormals() != nullptr )
+	{
+		FDynamicMeshNormalOverlay* NormalOverlay = Mesh.Attributes()->PrimaryNormals();
+		TSet<int32> UniqueElementIDs;
+		for ( int32 tid : Triangles)
+		{
+			if (NormalOverlay->IsSetTriangle(tid))
+			{
+				FIndex3i NormalTri = NormalOverlay->GetTriangle(tid);
+				UniqueElementIDs.Add(NormalTri.A);
+				UniqueElementIDs.Add(NormalTri.B);
+				UniqueElementIDs.Add(NormalTri.C);
+			}
+		}
+
+		return RecomputeOverlayElementNormals(Mesh, UniqueElementIDs.Array(), bWeightByArea, bWeightByAngle);
+	}
+	return false;
+}
+
+bool FMeshNormals::RecomputeOverlayElementNormals(FDynamicMesh3& Mesh, const TArray<int32>& ElementIDs, bool bWeightByArea, bool bWeightByAngle)
+{
+	if (Mesh.HasAttributes() && Mesh.Attributes()->PrimaryNormals() != nullptr )
+	{
+		FDynamicMeshNormalOverlay* NormalOverlay = Mesh.Attributes()->PrimaryNormals();
+		ParallelFor(ElementIDs.Num(), [&](int32 k) 
+		{
+			int32 ElementID = ElementIDs[k];
+			if ( NormalOverlay->IsElement(ElementID) )
+			{
+				FVector3d NewNormal = FMeshNormals::ComputeOverlayNormal(Mesh, NormalOverlay, ElementID, bWeightByArea, bWeightByAngle);
+				NormalOverlay->SetElement(ElementID, (FVector3f)NewNormal);
+			}
+		});
+		return true;
+	}
+	return false;
+}
+
 
 FVector3d FMeshNormals::ComputeVertexNormal(const FDynamicMesh3& Mesh, int VertIdx, bool bWeightByArea, bool bWeightByAngle)
 {
@@ -436,7 +477,7 @@ FVector3d FMeshNormals::ComputeVertexNormal(const FDynamicMesh3& Mesh, int32 Ver
 
 
 
-FVector3d FMeshNormals::ComputeOverlayNormal(const FDynamicMesh3& Mesh, const FDynamicMeshNormalOverlay* NormalOverlay, int ElemIdx)
+FVector3d FMeshNormals::ComputeOverlayNormal(const FDynamicMesh3& Mesh, const FDynamicMeshNormalOverlay* NormalOverlay, int ElemIdx, bool bWeightByArea, bool bWeightByAngle)
 {
 	int ParentVertexID = NormalOverlay->GetParentVertex(ElemIdx);
 	FVector3d SumNormal = FVector3d::Zero();
@@ -447,7 +488,10 @@ FVector3d FMeshNormals::ComputeOverlayNormal(const FDynamicMesh3& Mesh, const FD
 		{
 			FVector3d Normal, Centroid; double Area;
 			Mesh.GetTriInfo(TriIdx, Normal, Area, Centroid);
-			SumNormal += Area * Normal;
+			FVector3d TriNormalWeights = GetVertexWeightsOnTriangle(&Mesh, TriIdx, Area, bWeightByArea, bWeightByAngle);
+			FIndex3i Triangle = NormalOverlay->GetTriangle(TriIdx);
+			int32 j = IndexUtil::FindTriIndex(ElemIdx, Triangle);	// todo: we computed already in TriangleHasElement...
+			SumNormal += Normal * TriNormalWeights[j];
 			Count++;
 		}
 	});
diff --git a/Engine/Source/Runtime/GeometryCore/Public/DynamicMesh/InfoTypes.h b/Engine/Source/Runtime/GeometryCore/Public/DynamicMesh/InfoTypes.h
--- a/Engine/Source/Runtime/GeometryCore/Public/DynamicMesh/InfoTypes.h
+++ b/Engine/Source/Runtime/GeometryCore/Public/DynamicMesh/InfoTypes.h
@@ -95,6 +95,50 @@ struct FMeshTriEdgeID
 };
 
 
+/**
+ * FMeshTriOrderedEdgeID identifies an oriented edge in a triangle mesh based on indices
+ * into the triangle vertices. IE if a triangle has vertices [A,B,C], then an
+ * oriented edge could be (A,B) or (B,A), or any of the other 4 permutations.
+ * So the ordered edge in the triangle can be represented as two vertex indices, 
+ * and the full encoding is (TriangleID, J, K) where J and K are in range 0/1/2.
+ * 
+ * This type of edge identifier is applicable on any indexed mesh, even if 
+ * the mesh does not store explicit edge IDs. In addition, this identifier is stable
+ * across mesh topological changes, ie if two connected triangles are unlinked (ie
+ * the shared edge becomes two edges, and the 2 vertices become 4), the FMeshTriOrderedEdgeID
+ * will still refer to the correct oriented edge, as it does not explicitly depend on
+ * the Vertex or Edge IDs, only the ordering within the triangle.
+ * 
+ * Note that cycling or permuting the vertices of a triangle will change/break these indices.
+ */
+struct FMeshTriOrderedEdgeID
+{
+	/** The index of the mesh Triangle */
+	int32 TriangleID;
+	/** The 0/1/2 index of the first vertex in the triangles tuple of vertices */
+	unsigned VertIndexA : 2;
+	/** The 0/1/2 index of the second vertex in the triangles tuple of vertices */
+	unsigned VertIndexB : 2;
+
+
+	FMeshTriOrderedEdgeID()
+	{
+		TriangleID = IndexConstants::InvalidID;
+		VertIndexA = 0;
+		VertIndexB = 0;
+	}
+
+	FMeshTriOrderedEdgeID(int32 TriangleIDIn, int32 VertexIndexA, int32 VertexIndexB)
+	{
+		checkSlow(VertexIndexA >= 0 && VertexIndexA <= 2);
+		checkSlow(VertexIndexB >= 0 && VertexIndexB <= 2);
+		TriangleID = TriangleIDIn;
+		VertIndexA = VertexIndexA;
+		VertIndexB = VertexIndexB;
+	}
+};
+
+
 
 
 } // end namespace UE::Geometry
diff --git a/Engine/Source/Runtime/GeometryCore/Public/DynamicMesh/MeshNormals.h b/Engine/Source/Runtime/GeometryCore/Public/DynamicMesh/MeshNormals.h
--- a/Engine/Source/Runtime/GeometryCore/Public/DynamicMesh/MeshNormals.h
+++ b/Engine/Source/Runtime/GeometryCore/Public/DynamicMesh/MeshNormals.h
@@ -153,8 +153,10 @@ public:
 
 	/**
 	 * @return the computed overlay normal at an element of the overlay (ie based on normals of triangles connected to this element)
+	 * @param bWeightByArea weight neighbor triangles by area
+	 * @param bWeightByAngle weight neighbor triangles by angle
 	 */
-	static FVector3d ComputeOverlayNormal(const FDynamicMesh3& Mesh, const FDynamicMeshNormalOverlay* NormalOverlay, int ElemIdx);
+	static FVector3d ComputeOverlayNormal(const FDynamicMesh3& Mesh, const FDynamicMeshNormalOverlay* NormalOverlay, int ElemIdx, bool bWeightByArea = true, bool bWeightByAngle = true);
 
 	/**
 	 * Initialize the given NormalOverlay with per-vertex normals, ie single overlay element for each mesh vertex.
@@ -193,6 +195,15 @@ public:
 	 */
 	static bool QuickRecomputeOverlayNormals(FDynamicMesh3& Mesh, bool bInvert = false, bool bWeightByArea = true, bool bWeightByAngle = true);
 
+	/**
+	 * Compute overlay normals for the given mesh, for the given set of triangles
+	 */
+	static bool RecomputeOverlayTriNormals(FDynamicMesh3& Mesh, const TArray<int32>& Triangles, bool bWeightByArea = true, bool bWeightByAngle = true);
+
+	/**
+	 * Compute overlay normals for the given mesh, for the given set of element IDs
+	 */
+	static bool RecomputeOverlayElementNormals(FDynamicMesh3& Mesh, const TArray<int32>& ElementIDs, bool bWeightByArea = true, bool bWeightByAngle = true);
 
 	/**
 	 * Retrieve the area and/or angle weights for each vertex of a triangle
diff --git a/Engine/Source/Runtime/GeometryCore/Public/MeshQueries.h b/Engine/Source/Runtime/GeometryCore/Public/MeshQueries.h
--- a/Engine/Source/Runtime/GeometryCore/Public/MeshQueries.h
+++ b/Engine/Source/Runtime/GeometryCore/Public/MeshQueries.h
@@ -210,7 +210,7 @@ public:
 
 	static FAxisAlignedBox3d GetBounds(const TriangleMeshType& Mesh)
 	{
-		FAxisAlignedBox3d Bounds;
+		FAxisAlignedBox3d Bounds = FAxisAlignedBox3d::Empty();
 		for (int i = 0; i < Mesh.MaxVertexID(); ++i)
 		{
 			if (Mesh.IsVertex(i))
@@ -221,6 +221,46 @@ public:
 		return Bounds;
 	}
 
+	/** @return bounding box of subset of triangles of Mesh */
+	template<typename EnumerableTriListType>
+	static FAxisAlignedBox3d GetTrianglesBounds(
+		const TriangleMeshType& Mesh, 
+		const EnumerableTriListType& Triangles, 
+		const FTransform& Transform = FTransform::Identity)
+	{
+		FAxisAlignedBox3d Bounds = FAxisAlignedBox3d::Empty();
+		for (int32 tid : Triangles)
+		{
+			if (Mesh.IsTriangle(tid))
+			{
+				FVector3d A,B,C;
+				Mesh.GetTriVertices(tid, A,B,C);		// cannot use GetTriBounds here unless it is a FDynamicMesh3!
+				Bounds.Contain(Transform.TransformPosition(A));
+				Bounds.Contain(Transform.TransformPosition(B));
+				Bounds.Contain(Transform.TransformPosition(C));
+			}
+		}
+		return Bounds;
+	}
+
+	/** @return bounding box of subset of vertices of Mesh */
+	template<typename EnumerableTriListType>
+	static FAxisAlignedBox3d GetVerticesBounds(
+		const TriangleMeshType& Mesh, 
+		const EnumerableTriListType& Vertices,
+		const FTransform& Transform = FTransform::Identity)
+	{
+		FAxisAlignedBox3d Bounds = FAxisAlignedBox3d::Empty();
+		for (int32 vid : Vertices)
+		{
+			if (Mesh.IsVertex(vid))
+			{
+				Bounds.Contain( Transform.TransformPosition(Mesh.GetVertex(vid)) );
+			}
+		}
+		return Bounds;
+	}
+
 	// brute force search for nearest triangle to Point
 	static int FindNearestTriangle_LinearSearch(const TriangleMeshType& Mesh, const FVector3d& P)
 	{
