Expose TetWild parameters.

diff --git a/Engine/Plugins/Runtime/GeometryProcessing/Source/GeometryAlgorithms/Private/FTetWildWrapper.cpp b/Engine/Plugins/Runtime/GeometryProcessing/Source/GeometryAlgorithms/Private/FTetWildWrapper.cpp
--- a/Engine/Plugins/Runtime/GeometryProcessing/Source/GeometryAlgorithms/Private/FTetWildWrapper.cpp
+++ b/Engine/Plugins/Runtime/GeometryProcessing/Source/GeometryAlgorithms/Private/FTetWildWrapper.cpp
@@ -64,7 +64,9 @@ namespace
 				Params.values_sizing_field[Idx] = InParams.SizingFieldValues[Idx];
 			}
 		}
+		Params.ideal_edge_length = InParams.IdealEdgeLength;
 		Params.eps_rel = InParams.EpsRel;
+		Params.max_its = InParams.MaxIts;
 		Params.stop_energy = InParams.StopEnergy;
 		Params.coarsen = InParams.bCoarsen;
 		Params.manifold_surface = InParams.bExtractManifoldBoundarySurface;
diff --git a/Engine/Plugins/Runtime/GeometryProcessing/Source/GeometryAlgorithms/Public/FTetWildWrapper.h b/Engine/Plugins/Runtime/GeometryProcessing/Source/GeometryAlgorithms/Public/FTetWildWrapper.h
--- a/Engine/Plugins/Runtime/GeometryProcessing/Source/GeometryAlgorithms/Public/FTetWildWrapper.h
+++ b/Engine/Plugins/Runtime/GeometryProcessing/Source/GeometryAlgorithms/Public/FTetWildWrapper.h
@@ -33,8 +33,14 @@ namespace UE::Geometry::FTetWild
 		TArray<FIntVector4> SizingFieldTets;
 		TArray<double> SizingFieldValues;
 
+		// Ideal edge length
+		double IdealEdgeLength = 0.05;
+
 		// Relative tolerance, controlling how closely the mesh must follow the input surface
 		double EpsRel = 1e-3;
+
+		// Maximum number of optimization iterations
+		int32 MaxIts = 80;
 		// Energy at which to stop optimizing tet quality and accept the result
 		double StopEnergy = 10;
 
