GeometryProcessing: OffsetMeshRegion improvements
Refactor core of FDynamicMeshEditor::DisconnectTriangles into an overload that takes precomputed triangle TSet and BoundaryLoops array, as some call sites already have computed this info
Refactor new-polygroup computation out of OffsetMeshRegion into UE::Geometry::ComputeNewGroupIDsAlongEdgeLoop util function in PolyEditingEdgeUtil.h/cpp
Add new function UE::Geometry::ComputeAverageUVScaleRatioAlongVertexPath in PolyEditingUVUtil.h/cpp

Introduce new type FQuadGridPatch, this is a book-keeping data structure for keeping track of "grid of quads" subregion of a triangle mesh, that provides easy access to vertex and quad/trainagle rows/columns
Add util functions ComputeNormalsForQuadPatch() and ComputeUVIslandForQuadPatch() in QuadGridPatchUtil.h

Add new version of OffsetMeshRegion geometric operation, inside FOffsetMeshRegion class. EVersion parameter determines which version of operation to run. Code for the "Legacy" version has not been modified, so back-compat is maintained. New implementation supports variable number of subdivisions along the Offset (via .NumSubdivisions), and computes normal/UV-islands "per group" in extrusion region, instead of per-quad. New version also splits bowties in extrude region before any additional processing, which simplifies some of the book-keeping/etc.

FExtrudeOp currently hardcoded to use use the FOffsetMeshRegion implementation

diff --git a/Engine/Plugins/Runtime/GeometryProcessing/Source/DynamicMesh/Private/DynamicMeshEditor.cpp b/Engine/Plugins/Runtime/GeometryProcessing/Source/DynamicMesh/Private/DynamicMeshEditor.cpp
--- a/Engine/Plugins/Runtime/GeometryProcessing/Source/DynamicMesh/Private/DynamicMeshEditor.cpp
+++ b/Engine/Plugins/Runtime/GeometryProcessing/Source/DynamicMesh/Private/DynamicMeshEditor.cpp
@@ -462,10 +462,6 @@ void FDynamicMeshEditor::DuplicateTriangles(const TArray<int>& Triangles, FMeshI
 
 }
 
-
-
-
-
 bool FDynamicMeshEditor::DisconnectTriangles(const TArray<int>& Triangles, TArray<FLoopPairSet>& LoopSetOut, bool bHandleBoundaryVertices)
 {
 	// find the region boundary loops
@@ -477,14 +473,13 @@ bool FDynamicMeshEditor::DisconnectTriangles(const TArray<int>& Triangles, TArra
 	}
 	TArray<FEdgeLoop>& Loops = RegionLoops.Loops;
 
-	// need to test Contains() many times
-	TSet<int> TriangleSet;
-	TriangleSet.Reserve(Triangles.Num() * 3);
-	for (int TriID : Triangles)
-	{
-		TriangleSet.Add(TriID);
-	}
+	TSet<int> TriangleSet(Triangles);
+
+	return DisconnectTriangles(TriangleSet, Loops, LoopSetOut, bHandleBoundaryVertices);
+}
 
+bool FDynamicMeshEditor::DisconnectTriangles(const TSet<int>& TriangleSet, const TArray<FEdgeLoop>& Loops, TArray<FLoopPairSet>& LoopSetOut, bool bHandleBoundaryVertices)
+{
 	// process each loop island
 	int NumLoops = Loops.Num();
 	LoopSetOut.SetNum(NumLoops);
@@ -492,7 +487,7 @@ bool FDynamicMeshEditor::DisconnectTriangles(const TArray<int>& Triangles, TArra
 	TMap<int32, int32> OldVidsToNewVids; 
 	for ( int li = 0; li < NumLoops; ++li)
 	{
-		FEdgeLoop& Loop = Loops[li];
+		const FEdgeLoop& Loop = Loops[li];
 		FLoopPairSet& LoopPair = LoopSetOut[li];
 		LoopPair.OuterVertices = Loop.Vertices;
 		LoopPair.OuterEdges = Loop.Edges;
diff --git a/Engine/Plugins/Runtime/GeometryProcessing/Source/DynamicMesh/Private/Operations/OffsetMeshRegion.cpp b/Engine/Plugins/Runtime/GeometryProcessing/Source/DynamicMesh/Private/Operations/OffsetMeshRegion.cpp
--- a/Engine/Plugins/Runtime/GeometryProcessing/Source/DynamicMesh/Private/Operations/OffsetMeshRegion.cpp
+++ b/Engine/Plugins/Runtime/GeometryProcessing/Source/DynamicMesh/Private/Operations/OffsetMeshRegion.cpp
@@ -5,31 +5,20 @@
 #include "Algo/Reverse.h"
 #include "DynamicMesh/MeshNormals.h"
 #include "DynamicMeshEditor.h"
+#include "Parameterization/DynamicMeshUVEditor.h"
 #include "Selections/MeshVertexSelection.h"
 #include "DynamicMesh/DynamicMeshChangeTracker.h"
 #include "Selections/MeshConnectedComponents.h"
 #include "Operations/ExtrudeMesh.h"
+#include "EdgeLoop.h"
 #include "DynamicSubmesh3.h"
+#include "Selections/QuadGridPatch.h"
+#include "Operations/PolyEditingUVUtil.h"
+#include "Operations/PolyEditingEdgeUtil.h"
+#include "Operations/QuadGridPatchUtil.h"
 
 using namespace UE::Geometry;
 
-namespace OffsetMeshRegionLocals
-{
-	bool EdgesAreParallel(FDynamicMesh3* Mesh, int32 Eid1, int32 Eid2)
-	{
-		FIndex2i Vids1 = Mesh->GetEdgeV(Eid1);
-		FIndex2i Vids2 = Mesh->GetEdgeV(Eid2);
-
-		FVector3d Vec1 = Mesh->GetVertex(Vids1.A) - Mesh->GetVertex(Vids1.B);
-		FVector3d Vec2 = Mesh->GetVertex(Vids2.A) - Mesh->GetVertex(Vids2.B);
-		if (!Vec1.Normalize(KINDA_SMALL_NUMBER) || !Vec2.Normalize(KINDA_SMALL_NUMBER))
-		{
-			// A degenerate edge is parallel enough for our purposes
-			return true;
-		}
-		return FMath::Abs(Vec1.Dot(Vec2)) >= 1 - KINDA_SMALL_NUMBER;
-	}
-}
 
 FOffsetMeshRegion::FOffsetMeshRegion(FDynamicMesh3* mesh) : Mesh(mesh)
 {
@@ -60,11 +49,615 @@ bool FOffsetMeshRegion::Apply()
 	}
 
 	return bAllOK;
+}
+
+
+namespace OffsetMeshRegionLocals
+{
+
+typedef TPair<int32, TPair<int8, int8>> TriVertPair;
+
+/**
+ * @return true if mesh edges are parallel or either is degenerate (the degenerate tolerance is quite large though, 0.0001).
+ */
+bool EdgesAreParallel(FDynamicMesh3* Mesh, int32 Eid1, int32 Eid2)
+{
+	FIndex2i Vids1 = Mesh->GetEdgeV(Eid1);
+	FIndex2i Vids2 = Mesh->GetEdgeV(Eid2);
+
+	FVector3d Vec1 = Mesh->GetVertex(Vids1.A) - Mesh->GetVertex(Vids1.B);
+	FVector3d Vec2 = Mesh->GetVertex(Vids2.A) - Mesh->GetVertex(Vids2.B);
+	if (!Vec1.Normalize(KINDA_SMALL_NUMBER) || !Vec2.Normalize(KINDA_SMALL_NUMBER))
+	{
+		// A degenerate edge is parallel enough for our purposes
+		return true;
+	}
+	return FMath::Abs(Vec1.Dot(Vec2)) >= 1 - KINDA_SMALL_NUMBER;
+}
+
+/**
+ * Assuming GroupIDs is a list of integer IDs for a contiguous "loop", ie last
+ * element is adjacent to first element, this function returns the index of the first
+ * value in the first non-broken contiguous span of duplicate values. IE if the array 
+ * is [0,0,0,1,1,2,2] it will return "0" but if it were [0,0,1,1,2,2,0] it would return "2", 
+ * ie the inde of the first "1". The purpose is to allow iterating through the array
+ * in the order [1,1,2,2,0,0,0], ie the "0" values are not split across the N-1/0 transition
+ */
+static int32 FindLoopShiftFromGroupIDs(const TArray<int32>& GroupIDs)
+{
+	int32 N = GroupIDs.Num();
+	if (GroupIDs[0] != GroupIDs[N - 1])
+	{
+		return 0;
+	}
+	for (int32 k = 0; k < N-1; ++k)
+	{
+		if (GroupIDs[k] != GroupIDs[k + 1])
+		{
+			return k+1;
+		}
+	}
+	return 0;	// all values are the same, no shift
+}
+
+/**
+ * Cycle/Shift the values in the Values array to the left ShiftNum times.
+ * Uses a temporary array internally.
+ */
+template<typename ValueType>
+static void LeftShiftArray(TArray<ValueType>& Values, int ShiftNum)
+{
+	if (ShiftNum == 0 ) return;
+
+	int32 N = Values.Num();
+
+	// there is probably some clever way to do this w/o a temporary array...
+	TArray<ValueType> Tmp;
+	Tmp.SetNum(N);
+	for (int32 k = 0; k < N; ++k)
+	{
+		Tmp[k] = Values[(ShiftNum+k)%N];
+	}
+	Values = MoveTemp(Tmp);
+}
+
+
+
+/**
+ * Stitch pairs of border loops of Mesh defined by LoopPairs.
+ * 
+ */
+static bool StitchRegionBorderLoopPairs_Version1(
+	FDynamicMesh3& Mesh,
+	const TArray<FDynamicMeshEditor::FLoopPairSet>& LoopPairs,
+	const TArray<TArray<int32>>& LoopsPerEdgeNewGroupIDs,
+	const TArray<TArray<FMeshTriOrderedEdgeID>>& InnerTriOrderedEdgeLoops,
+	double UVScaleFactor,
+	int NumSubdivisions,
+	TArray<bool>& bLoopOK,
+	TArray<TArray<int32>>& PerLoopTrianglesOut,
+	TArray<TArray<int32>>& PerLoopGroupsOut )
+{
+	FDynamicMeshEditor Editor(&Mesh);
+
+	bool bAllSuccess = true;
+	int NumInitialLoops = LoopPairs.Num();
+
+	PerLoopTrianglesOut.SetNum(NumInitialLoops);
+	PerLoopGroupsOut.SetNum(NumInitialLoops);
+	bLoopOK.Init(false, NumInitialLoops);
+
+	for (int32 LoopIndex = 0; LoopIndex < LoopPairs.Num(); ++LoopIndex)
+	{
+		// note we are assuming here that the loops are aligned w/ a group transition
+
+		const TArray<int32>& OuterLoopV = LoopPairs[LoopIndex].OuterVertices;
+		const TArray<int32>& InnerLoopV = LoopPairs[LoopIndex].InnerVertices;
+		const TArray<FMeshTriOrderedEdgeID>& InnerTriOrderedEdgeLoop = InnerTriOrderedEdgeLoops[LoopIndex];
+		int32 NV = OuterLoopV.Num();
+
+		// LoopStack is the set of NumSubdivisions+1 loops that will be stitched,
+		// where the "last" loop is the Outer loop. The "first" Inner loop is not 
+		// included in LoopStack, it has special handling below
+		TArray<TArray<int32>> LoopStack;
+		LoopStack.SetNum(NumSubdivisions+1);
+		LoopStack[NumSubdivisions] = OuterLoopV;
+		int LastTemp = LoopStack[NumSubdivisions][0];
+
+		// build interior loops if there are any
+		for ( int32 StackIdx = 0; StackIdx < NumSubdivisions; ++StackIdx )
+		{
+			LoopStack[StackIdx].Reserve(NV);
+			double t = (double)(StackIdx+1) / (double)(NumSubdivisions+1);
+			for ( int32 k = 0; k < NV; ++k )
+			{
+				FVector3d A = Mesh.GetVertex(InnerLoopV[k]);
+				FVector3d B = Mesh.GetVertex(OuterLoopV[k]);
+				FVector3d StackPos = Lerp(A, B, t);
+				LoopStack[StackIdx].Add( Mesh.AppendVertex(StackPos) );
+			}
+			LastTemp = LoopStack[StackIdx][0];
+		}
+
+		// build the first strip from the Inner loop to the next loop
+		FDynamicMeshEditResult StitchResult;
+		TArray<int32> InnerVertexLoop, TempEdgeLoop;
+		UE::Geometry::ConvertTriOrderedEdgeLoopToLoop(Mesh, InnerTriOrderedEdgeLoop, InnerVertexLoop, &TempEdgeLoop);
+		bLoopOK[LoopIndex] = Editor.StitchVertexLoopsMinimal(InnerVertexLoop, LoopStack[0], StitchResult);
+		if (!bLoopOK[LoopIndex])
+		{
+			bAllSuccess = false;
+			continue;
+		}
+
+		FQuadGridPatch QuadPatch;
+		TArray<int32> TempSpan = LoopStack[0];		// annoying to have to make a copy like this...need to perserve for LoopStack[k-1] below though
+		TempSpan.Add(LoopStack[0][0]);
+		QuadPatch.InitializeFromQuadStrip(Mesh, StitchResult.NewQuads, TempSpan);
+		QuadPatch.ReverseRows();	// need LoopStack[0] to be the second span but it's currently the first one...
+
+		// incrementally stitch intermediate strips and accumulate into the Quad Patch
+		for ( int32 k = 1; k <= NumSubdivisions; ++k )
+		{
+			FDynamicMeshEditResult StripResult;
+			bLoopOK[LoopIndex] = Editor.StitchVertexLoopsMinimal(LoopStack[k-1], LoopStack[k], StripResult);
+
+			FQuadGridPatch NextQuadStrip;
+			TempSpan = LoopStack[k-1];		// must preserve LoopStack arrays currently
+			TempSpan.Add(LoopStack[k-1][0]);
+			NextQuadStrip.InitializeFromQuadStrip(Mesh, StripResult.NewQuads, TempSpan);
+
+			QuadPatch.AppendQuadPatchRows(MoveTemp(NextQuadStrip), true);
+		}
+
+		QuadPatch.ReverseRows();	// for an extrude/offset this makes sense, quads will "start" at the transition from base loop
+
+		// set group IDs on patch, propagating up the subdivisions columns
+		const TArray<int32>& PerEdgeNewGroupIDs = LoopsPerEdgeNewGroupIDs[LoopIndex];		// is it always guaranteed to be the same length??
+		QuadPatch.ForEachQuad( [&](int32 QuadRow, int32 QuadIndex, FIndex2i QuadTris)
+		{
+			int32 GroupID = PerEdgeNewGroupIDs[QuadIndex];
+			Mesh.SetTriangleGroup(QuadTris.A, GroupID);
+			Mesh.SetTriangleGroup(QuadTris.B, GroupID);
+		});
+
+		// split quad patch by group ID into a set of quad patches
+		TArray<FQuadGridPatch> GroupStrips;
+		QuadPatch.SplitColumnsByPredicate(
+			[&](int32 ColumnIdx, int32 NextColumnIdx) { return PerEdgeNewGroupIDs[ColumnIdx] != PerEdgeNewGroupIDs[NextColumnIdx]; },
+			GroupStrips);
+
+		// compute normals and UVs for each group-quad-patch
+		if (Mesh.HasAttributes())
+		{
+			for (FQuadGridPatch& GroupPatch : GroupStrips)
+			{
+				UE::Geometry::ComputeNormalsForQuadPatch(Mesh, GroupPatch);
+				UE::Geometry::ComputeUVIslandForQuadPatch(Mesh, GroupPatch, UVScaleFactor);
+			}
+		}
+
+		// Assign a new group to each row of group-quad-patch
+		for (FQuadGridPatch& GroupPatch : GroupStrips)
+		{
+			int32 NumQuadRows = GroupPatch.NumVertexRowsV-1;
+			TArray<int32> RowGroups;
+			RowGroups.SetNum(NumQuadRows);
+			for ( int32 k = 1; k < NumQuadRows; ++k )
+			{
+				RowGroups[k] = Mesh.AllocateTriangleGroup();
+			}
+			GroupPatch.ForEachQuad([&](int32 Row, int32 Col, FIndex2i QuadTris)
+			{
+				if ( Row > 0 )
+				{
+					Mesh.SetTriangleGroup(QuadTris.A, RowGroups[Row]);
+					Mesh.SetTriangleGroup(QuadTris.B, RowGroups[Row]);
+				}
+			});
+		}
+
+		// save the stitch triangles set and associated group IDs
+		QuadPatch.GetAllTriangles(PerLoopTrianglesOut[LoopIndex]);
+		for (int32 tid : PerLoopTrianglesOut[LoopIndex])
+		{
+			PerLoopGroupsOut[LoopIndex].AddUnique( Mesh.GetTriangleGroup(tid) );
+		}
+	}
+
+	return bAllSuccess;
+}
 
 
+
+
+static void SetStitchStripNormalsUVs_Legacy(
+	FDynamicMesh3& Mesh,
+	FDynamicMeshEditor& Editor,
+	FDynamicMeshEditResult& StitchResult,
+	const TArray<int32>& OuterLoopV,
+	double UVScaleFactor )
+{
+	int NumNewQuads = StitchResult.NewQuads.Num();
+
+	double AccumUVTranslation = 0;
+	FFrame3d FirstProjectFrame;
+	FVector3d FrameUp;
+
+	for (int k = 0; k < NumNewQuads; k++)
+	{
+		FVector3f Normal = Editor.ComputeAndSetQuadNormal(StitchResult.NewQuads[k], true);
+
+		// align axis 0 of projection frame to first edge, then for further edges,
+		// rotate around 'up' axis to keep normal aligned and frame horizontal
+		FFrame3d ProjectFrame;
+		if (k == 0)
+		{
+			FVector3d FirstEdge = Mesh.GetVertex(OuterLoopV[1]) - Mesh.GetVertex(OuterLoopV[0]);
+			Normalize(FirstEdge);
+			FirstProjectFrame = FFrame3d(FVector3d::Zero(), (FVector3d)Normal);
+			FirstProjectFrame.ConstrainedAlignAxis(0, FirstEdge, (FVector3d)Normal);
+			FrameUp = FirstProjectFrame.GetAxis(1);
+			ProjectFrame = FirstProjectFrame;
+		}
+		else
+		{
+			ProjectFrame = FirstProjectFrame;
+			ProjectFrame.ConstrainedAlignAxis(2, (FVector3d)Normal, FrameUp);
+		}
+
+		if (k > 0)
+		{
+			AccumUVTranslation += Distance(Mesh.GetVertex(OuterLoopV[k]), Mesh.GetVertex(OuterLoopV[k-1]));
+		}
+
+		// translate horizontally such that vertical spans are adjacent in UV space (so textures tile/wrap properly)
+		double TranslateU = UVScaleFactor * AccumUVTranslation;
+		Editor.SetQuadUVsFromProjection(StitchResult.NewQuads[k], ProjectFrame, (float)UVScaleFactor, FVector2f((float)TranslateU, 0));
+	}
 }
 
+
+static bool StitchRegionBorderLoopPairs_Legacy(
+	FDynamicMesh3& Mesh,
+	const TArray<FDynamicMeshEditor::FLoopPairSet>& LoopPairs,
+	const TArray<TArray<int32>>& LoopsPerEdgeNewGroupIDs,
+	const TArray<TArray<TriVertPair>> OffsetStitchSides,
+	double UVScaleFactor,
+	TArray<bool>& bLoopOK,
+	TArray<TArray<int32>>& PerLoopTrianglesOut,
+	TArray<TArray<int32>>& PerLoopGroupsOut )
+{
+	FDynamicMeshEditor Editor(&Mesh);
+
+	bool bAllSuccess = true;
+	int NumInitialLoops = LoopPairs.Num();
+
+	PerLoopTrianglesOut.SetNum(NumInitialLoops);
+	PerLoopGroupsOut.SetNum(NumInitialLoops);
+	bLoopOK.Init(false, NumInitialLoops);
+
+	for (int32 LoopIndex = 0; LoopIndex < LoopPairs.Num(); ++LoopIndex)
+	{
+		const FDynamicMeshEditor::FLoopPairSet& LoopPair = LoopPairs[LoopIndex];
+		const TArray<int32>& EdgeGroups = LoopsPerEdgeNewGroupIDs[LoopIndex];
+
+		const TArray<int32>& BaseLoopV = LoopPair.OuterVertices;
+		const TArray<int32>& OffsetLoopV = LoopPair.InnerVertices;
+
+		const TArray<TriVertPair>& OffsetLoopTriVertPairs = OffsetStitchSides[LoopIndex];
+
+		// stitch the loops
+		FDynamicMeshEditResult StitchResult;
+		bLoopOK[LoopIndex] = Editor.StitchVertexLoopToTriVidPairSequence(OffsetLoopTriVertPairs, LoopPair.OuterVertices, StitchResult);
+		if (!bLoopOK[LoopIndex])
+		{
+			bAllSuccess = false;
+			continue;
+		}
+
+		// set the groups of the new quad strip created by the stitching
+		int NumNewQuads = StitchResult.NewQuads.Num();
+		TArray<int32> NewGroupsInLoop;
+		const TArray<int32>& PerEdgeNewGroupIDs = LoopsPerEdgeNewGroupIDs[LoopIndex];		// is it always guaranteed to be the same length??
+		for (int32 k = 0; k < NumNewQuads; k++)
+		{
+			Mesh.SetTriangleGroup(StitchResult.NewQuads[k].A, PerEdgeNewGroupIDs[k]);
+			Mesh.SetTriangleGroup(StitchResult.NewQuads[k].B, PerEdgeNewGroupIDs[k]);
+			NewGroupsInLoop.AddUnique(PerEdgeNewGroupIDs[k]);
+		}
+
+		// save the stitch triangles set and associated group IDs
+		StitchResult.GetAllTriangles(PerLoopTrianglesOut[LoopIndex]);
+		PerLoopGroupsOut[LoopIndex] = MoveTemp(NewGroupsInLoop);
+
+		// for each polygon we created in stitch, set UVs and normals
+		if (Mesh.HasAttributes())
+		{
+			SetStitchStripNormalsUVs_Legacy(Mesh, Editor, StitchResult, BaseLoopV, UVScaleFactor);
+		}
+	}
+
+	return bAllSuccess;
+}
+
+
+template<typename ListType>
+FVector3d GetAngleWeightedAverageNormal(const FDynamicMesh3& Mesh, int32 VertexID, const ListType& TriangleList)
+{
+	FVector3d ExtrusionVector = FVector3d::Zero();
+
+	// Get angle-weighted normalized average vector
+	for (int32 TriangleID : Mesh.VtxTrianglesItr(VertexID))
+	{
+		if (TriangleList.Contains(TriangleID))
+		{
+			FIndex3i Triangle = Mesh.GetTriangle(TriangleID);
+			double Angle = Mesh.GetTriInternalAngleR(TriangleID, Triangle.IndexOf(VertexID));
+			ExtrusionVector += Angle * Mesh.GetTriNormal(TriangleID);
+		}
+	}
+	ExtrusionVector.Normalize();
+	return ExtrusionVector;
+}
+
+
+template<typename ListType>
+FVector3d GetAngleWeightedAdjustedNormal(const FDynamicMesh3& Mesh, int32 VertexID, const ListType& TriangleList,
+	double MaxAdjustmentScale)
+{
+	FVector3d InitialExtrusionVector = GetAngleWeightedAverageNormal<ListType>(Mesh, VertexID, TriangleList);
+
+	// Perform an angle-weighted adjustment of the vector length. For each triangle normal, the
+	// length needs to be multiplied by 1/cos(theta) to place the vertex in the plane that it
+	// would be in if the face was moved a unit along triangle normal (where theta is angle of
+	// triangle normal to the current extrusion vector).
+	double AngleSum = 0;
+	double Adjustment = 0;
+	double InvertedMaxScale = FMath::Max(FMathd::ZeroTolerance, 1.0 / MaxAdjustmentScale);
+	for (int32 TriangleID : Mesh.VtxTrianglesItr(VertexID))
+	{
+		if (TriangleList.Contains(TriangleID))
+		{
+			FIndex3i Triangle = Mesh.GetTriangle(TriangleID);
+			double Angle = Mesh.GetTriInternalAngleR(TriangleID, Triangle.IndexOf(VertexID));
+			double CosTheta = Mesh.GetTriNormal(TriangleID).Dot(InitialExtrusionVector);
+
+			if (CosTheta <= InvertedMaxScale)
+			{
+				CosTheta = InvertedMaxScale;
+			}
+			Adjustment += Angle / CosTheta;
+
+			// For the average at the end
+			AngleSum += Angle;
+		}
+	}
+	Adjustment /= AngleSum;
+
+	return InitialExtrusionVector * Adjustment;
+}
+
+
+} // end namespace OffsetMeshRegionLocals
+
+
 bool FOffsetMeshRegion::ApplyOffset(FOffsetInfo& Region)
+{
+	if (UseVersion == EVersion::Legacy)
+	{
+		return ApplyOffset_Legacy(Region);
+	}
+	else  // Version1, Current
+	{
+		return ApplyOffset_Version1(Region);
+	}
+}
+
+
+bool FOffsetMeshRegion::ApplyOffset_Version1(FOffsetInfo& Region)
+{
+	const TArray<int32>& RegionTriangles = Region.OffsetTids;
+
+	// Split any bowties in the offset region. An extrusion of a bowtie creates a nonmanifold edge which is not permitted
+	FDynamicMeshEditor TmpEditor(Mesh);
+	FDynamicMeshEditResult IgnoreResult;
+	TmpEditor.SplitBowtiesAtTriangles(RegionTriangles, IgnoreResult);
+
+	// Store offset groups
+	if (Mesh->HasTriangleGroups())
+	{
+		for (int32 TriangleID : RegionTriangles)
+		{
+			Region.OffsetGroups.AddUnique(Mesh->GetTriangleGroup(TriangleID));
+		}
+	}
+
+	FMeshRegionBoundaryLoops InitialLoops(Mesh, RegionTriangles, false);
+	bool bOK = InitialLoops.Compute();
+	if (bOK == false)
+	{
+		return false;
+	}
+
+	AllModifiedAndNewTriangles.Append(RegionTriangles);
+	TSet<int32> TriangleSet(RegionTriangles);
+
+	// Before we start changing triangles, prepare by allocating group IDs that we'll use
+	// for the stitched sides (doing it before changes to the mesh allows user-provided
+	// LoopEdgesShouldHaveSameGroup functions to operate on the original mesh).
+	TArray<TArray<int32>> LoopsEdgeGroups;
+	TArray<int32> NewGroupIDs;
+	LoopsEdgeGroups.SetNum(InitialLoops.Loops.Num());
+	for (int32 i = 0; i < InitialLoops.Loops.Num(); ++i)
+	{
+		TArray<int32>& LoopEids = InitialLoops.Loops[i].Edges;
+		TArray<int32>& CurrentEdgeGroups = LoopsEdgeGroups[i];
+		UE::Geometry::ComputeNewGroupIDsAlongEdgeLoop(*Mesh, LoopEids, CurrentEdgeGroups, NewGroupIDs,
+			LoopEdgesShouldHaveSameGroup);
+
+		// if a group is split over the start/end of the loop, shift the loops to so that index 0 lies on a group transition,
+		// this will simplify later processing
+		int GroupShift = OffsetMeshRegionLocals::FindLoopShiftFromGroupIDs(CurrentEdgeGroups);
+		if (GroupShift != 0)
+		{
+			OffsetMeshRegionLocals::LeftShiftArray(InitialLoops.Loops[i].Vertices, GroupShift);
+			OffsetMeshRegionLocals::LeftShiftArray(InitialLoops.Loops[i].Edges, GroupShift);
+			OffsetMeshRegionLocals::LeftShiftArray(CurrentEdgeGroups, GroupShift);
+		}
+	}
+
+	FDynamicMeshEditor Editor(Mesh);
+	TArray<FDynamicMeshEditor::FLoopPairSet> LoopPairs;
+
+	FDynamicMeshEditResult DuplicateResult;
+	if (Region.bIsSolid)
+	{
+		// In the solid case, we want to duplicate the region so we can cap it.
+		FMeshIndexMappings IndexMap;
+		Editor.DuplicateTriangles(RegionTriangles, IndexMap, DuplicateResult);
+
+		AllModifiedAndNewTriangles.Append(DuplicateResult.NewTriangles);
+
+		// Populate LoopPairs
+		LoopPairs.SetNum(InitialLoops.Loops.Num());
+		for (int LoopIndex = 0; LoopIndex < InitialLoops.Loops.Num(); ++LoopIndex)
+		{
+			FEdgeLoop& BaseLoop = InitialLoops.Loops[LoopIndex];
+			FDynamicMeshEditor::FLoopPairSet& LoopPair = LoopPairs[LoopIndex];
+
+			// The original RegionTriangles are the ones that are offset, so InnerVertices/Edges
+			// should be the boundaries of those.
+			LoopPair.InnerVertices = BaseLoop.Vertices;
+			LoopPair.InnerEdges = BaseLoop.Edges;
+
+			// However depending on whether we extruded down or up, we may need to reverse
+			// the loops to get them to be stitched right side out.
+			if (!bIsPositiveOffset)
+			{
+				Algo::Reverse(LoopPair.InnerVertices);
+
+				// Reversing the edges is slightly different because the last edge is between the first
+				// and last vertex, and that needs to stay in the same place when vertices are reversed.
+				int32 LastEid = LoopPair.InnerEdges.Pop();
+				Algo::Reverse(LoopPair.InnerEdges);
+				LoopPair.InnerEdges.Add(LastEid);
+
+				int32 LastEdgeGroupID = LoopsEdgeGroups[LoopIndex].Pop();
+				Algo::Reverse(LoopsEdgeGroups[LoopIndex]);
+				LoopsEdgeGroups[LoopIndex].Add(LastEdgeGroupID);
+			}
+
+			// Now assemble the paired loop
+			for (int32 Vid : LoopPair.InnerVertices)
+			{
+				LoopPair.OuterVertices.Add(IndexMap.GetNewVertex(Vid));
+			}
+			FEdgeLoop::VertexLoopToEdgeLoop(Mesh, LoopPair.OuterVertices, LoopPair.OuterEdges);
+		}
+	}
+	else
+	{
+		bOK = Editor.DisconnectTriangles(TriangleSet, InitialLoops.Loops, LoopPairs, true /*bHandleBoundaryVertices*/);
+	}
+
+	if (bOK == false)
+	{
+		return false;
+	}
+
+	// Store the vid-independent offset loop before we break bowties
+	// NO LONGER NECESSARY
+	TArray<TArray<FMeshTriOrderedEdgeID>> OffsetTriOrderedEdgeLoops;
+	OffsetTriOrderedEdgeLoops.SetNum(LoopPairs.Num());
+	for (int32 i = 0; i < LoopPairs.Num(); ++i)
+	{
+		bOK = bOK && UE::Geometry::ConvertLoopToTriOrderedEdgeLoop(
+			*Mesh, LoopPairs[i].InnerVertices, LoopPairs[i].InnerEdges, OffsetTriOrderedEdgeLoops[i]);
+	}
+	if (bOK == false)
+	{
+		return false;
+	}
+
+	FMeshVertexSelection SelectionV(Mesh);
+	SelectionV.SelectTriangleVertices(RegionTriangles);
+	TArray<int32> SelectedVids = SelectionV.AsArray();
+
+	// If we need to, assemble the vertex vectors for us to use (before we actually start moving things)
+	TArray<FVector3d> VertexExtrudeVectors;
+	if (ExtrusionVectorType == EVertexExtrusionVectorType::SelectionTriNormalsAngleWeightedAverage
+		|| ExtrusionVectorType == EVertexExtrusionVectorType::SelectionTriNormalsAngleWeightedAdjusted)
+	{
+		VertexExtrudeVectors.SetNumUninitialized(SelectedVids.Num());
+
+		// Used to test which triangles are in selection
+		for (int32 i = 0; i < SelectedVids.Num(); ++i)
+		{
+			VertexExtrudeVectors[i] = ( ExtrusionVectorType == EVertexExtrusionVectorType::SelectionTriNormalsAngleWeightedAdjusted ) ?
+				OffsetMeshRegionLocals::GetAngleWeightedAdjustedNormal(*Mesh, SelectedVids[i], TriangleSet, MaxScaleForAdjustingTriNormalsOffset)
+				 : OffsetMeshRegionLocals::GetAngleWeightedAverageNormal(*Mesh, SelectedVids[i], TriangleSet);
+		}
+	}
+	else if (ExtrusionVectorType == EVertexExtrusionVectorType::VertexNormal)
+	{
+		VertexExtrudeVectors.SetNumUninitialized(SelectedVids.Num());
+		for (int32 i = 0; i < SelectedVids.Num(); ++i)
+		{
+			VertexExtrudeVectors[i] = Mesh->HasVertexNormals() ? 
+				(FVector3d)Mesh->GetVertexNormal(SelectedVids[i]) : FMeshNormals::ComputeVertexNormal(*Mesh, SelectedVids[i]);
+		}
+	}
+
+	// Perform the actual vertex displacement.
+	for (int32 i = 0; i < SelectedVids.Num(); ++i)
+	{
+		int32 VertexID = SelectedVids[i];
+		FVector3d OldPosition = Mesh->GetVertex(VertexID);
+		FVector3d ExtrusionVector = (ExtrusionVectorType == EVertexExtrusionVectorType::Zero) ? FVector3d::Zero() : VertexExtrudeVectors[i];
+
+		FVector3d NewPosition = OffsetPositionFunc(OldPosition, ExtrusionVector, VertexID);
+		Mesh->SetVertex(VertexID, NewPosition);
+	}
+
+	// Stitch the loops
+	TArray<bool> bLoopSuccess;
+	bool bSuccess = OffsetMeshRegionLocals::StitchRegionBorderLoopPairs_Version1( *Mesh, 
+			LoopPairs, LoopsEdgeGroups, OffsetTriOrderedEdgeLoops,
+			UVScaleFactor, NumSubdivisions,
+			bLoopSuccess, Region.StitchTriangles, Region.StitchPolygonIDs);
+
+	int NumInitialLoops = LoopPairs.Num();
+	Region.BaseLoops.SetNum(NumInitialLoops);
+	Region.OffsetLoops.SetNum(NumInitialLoops);
+	for (int32 LoopIndex = 0; LoopIndex < NumInitialLoops; ++LoopIndex)
+	{
+		if (bLoopSuccess[LoopIndex])
+		{
+			Region.BaseLoops[LoopIndex].InitializeFromVertices(Mesh, LoopPairs[LoopIndex].OuterVertices);
+			Region.OffsetLoops[LoopIndex].InitializeFromVertices(Mesh, LoopPairs[LoopIndex].InnerVertices);
+		}
+	}
+
+	if (Region.bIsSolid)
+	{
+		if (bIsPositiveOffset)
+		{
+			// Flip the "bottom" of the region to face outwards
+			Editor.ReverseTriangleOrientations(DuplicateResult.NewTriangles, true);
+		}
+		else
+		{
+			Editor.ReverseTriangleOrientations(RegionTriangles, true);
+		}
+	}
+
+	return bSuccess;
+}
+
+
+
+bool FOffsetMeshRegion::ApplyOffset_Legacy(FOffsetInfo& Region)
 {
 	// Store offset groups
 	if (Mesh->HasTriangleGroups())
@@ -182,6 +775,9 @@ bool FOffsetMeshRegion::ApplyOffset(FOffsetInfo& Region)
 	}
 	else
 	{
+		// Disconnect the triangles. Note that this will recompute FMeshRegionBoundaryLoops internally,
+		// and so the output LoopPairs will generally correspond to InitialLoops for that reason,
+		// but there isn't a guarantee otherwise
 		bOK = Editor.DisconnectTriangles(Region.OffsetTids, LoopPairs, true /*bHandleBoundaryVertices*/);
 	}
 
@@ -191,8 +787,7 @@ bool FOffsetMeshRegion::ApplyOffset(FOffsetInfo& Region)
 	}
 
 	// Store the vid-independent offset loop before we break bowties
-	typedef TPair<int32, TPair<int8, int8>> TriVertPair;
-	TArray<TArray<TriVertPair>> OffsetStitchSides;
+	TArray<TArray<OffsetMeshRegionLocals::TriVertPair>> OffsetStitchSides;
 	OffsetStitchSides.SetNum(LoopPairs.Num());
 	for (int32 i = 0; i < LoopPairs.Num(); ++i)
 	{
@@ -209,10 +804,19 @@ bool FOffsetMeshRegion::ApplyOffset(FOffsetInfo& Region)
 	Editor.SplitBowtiesAtTriangles(Region.OffsetTids, Result);
 	bool bSomeLoopsBroken = Result.NewVertices.Num() > 0;
 
+	// BUG: once we split bowties, there some of the existing LoopPairs may actually no longer be boundary loops.
+	// This occurs with an "interior" bowtie, where say a figure-8 shape that was two separate loops, becomes one loop.
+	// In this case things go downhill as OffsetTriOrderedEdgeLoops can no longer be converted back to a valid
+	// boundary loop (since it's not a loop anymore). The FDynamicMeshEditor::StitchVertexLoopsMinimal function called below
+	// does not actually check that the loop it is stitching is a boundary loop, it simply stitches edge pairs with quads, 
+	// which means that the function will usually appear to succeed even though it may have created broken topology.
+
 	// If we broke bowties, the loops in the offset region have changed, and our OffsetLoops no longer
 	// match BaseLoops.
 	if (bSomeLoopsBroken)
 	{
+		// NOTE: Region.OffsetLoops is only used for output here, this does not affect the loops that
+		// will be stitched, it just needs to be computed before stitching
 		FMeshRegionBoundaryLoops UpdatedOffsetLoops(Mesh, Region.OffsetTids, false);
 		bOK = UpdatedOffsetLoops.Compute();
 		if (!bOK)
@@ -236,53 +840,12 @@ bool FOffsetMeshRegion::ApplyOffset(FOffsetInfo& Region)
 		// Used to test which triangles are in selection
 		TSet<int32> TriangleSet(Region.OffsetTids);
 
+		// Used to test which triangles are in selection
 		for (int32 i = 0; i < SelectedVids.Num(); ++i)
 		{
-			int32 Vid = SelectedVids[i];
-			FVector3d ExtrusionVector = FVector3d::Zero();
-
-			// Get angle-weighted normalized average vector
-			for (int32 Tid : Mesh->VtxTrianglesItr(Vid))
-			{
-				if (TriangleSet.Contains(Tid))
-				{
-					double Angle = Mesh->GetTriInternalAngleR(Tid, Mesh->GetTriangle(Tid).IndexOf(Vid));
-					ExtrusionVector += Angle * Mesh->GetTriNormal(Tid);
-				}
-			}
-			ExtrusionVector.Normalize();
-
-			if (ExtrusionVectorType == EVertexExtrusionVectorType::SelectionTriNormalsAngleWeightedAdjusted)
-			{
-				// Perform an angle-weighted adjustment of the vector length. For each triangle normal, the
-				// length needs to be multiplied by 1/cos(theta) to place the vertex in the plane that it
-				// would be in if the face was moved a unit along triangle normal (where theta is angle of
-				// triangle normal to the current extrusion vector).
-				double AngleSum = 0;
-				double Adjustment = 0;
-				for (int32 Tid : Mesh->VtxTrianglesItr(Vid))
-				{
-					if (TriangleSet.Contains(Tid))
-					{
-						double Angle = Mesh->GetTriInternalAngleR(Tid, Mesh->GetTriangle(Tid).IndexOf(Vid));
-						double CosTheta = Mesh->GetTriNormal(Tid).Dot(ExtrusionVector);
-
-						double InvertedMaxScaleFactor = FMath::Max(FMathd::ZeroTolerance, 1.0 / MaxScaleForAdjustingTriNormalsOffset);
-						if (CosTheta <= InvertedMaxScaleFactor)
-						{
-							CosTheta = InvertedMaxScaleFactor;
-						}
-						Adjustment += Angle / CosTheta;
-
-						// For the average at the end
-						AngleSum += Angle;
-					}
-				}
-				Adjustment /= AngleSum;
-				ExtrusionVector *= Adjustment;
-			}
-
-			VertexExtrudeVectors[i] = ExtrusionVector;
+			VertexExtrudeVectors[i] = ( ExtrusionVectorType == EVertexExtrusionVectorType::SelectionTriNormalsAngleWeightedAdjusted ) ?
+				OffsetMeshRegionLocals::GetAngleWeightedAdjustedNormal(*Mesh, SelectedVids[i], TriangleSet, MaxScaleForAdjustingTriNormalsOffset)
+				 : OffsetMeshRegionLocals::GetAngleWeightedAverageNormal(*Mesh, SelectedVids[i], TriangleSet);
 		}
 	}
 	else if (ExtrusionVectorType == EVertexExtrusionVectorType::VertexNormal)
@@ -326,7 +889,7 @@ bool FOffsetMeshRegion::ApplyOffset(FOffsetInfo& Region)
 		TArray<int32>& BaseLoopV = LoopPair.OuterVertices;
 		TArray<int32>& OffsetLoopV = LoopPair.InnerVertices;
 
-		TArray<TriVertPair>& OffsetLoopTriVertPairs = OffsetStitchSides[LoopIndex];
+		TArray<OffsetMeshRegionLocals::TriVertPair>& OffsetLoopTriVertPairs = OffsetStitchSides[LoopIndex];
 
 		// stitch the loops
 		FDynamicMeshEditResult StitchResult;
diff --git a/Engine/Plugins/Runtime/GeometryProcessing/Source/DynamicMesh/Private/Operations/PolyEditingEdgeUtil.cpp b/Engine/Plugins/Runtime/GeometryProcessing/Source/DynamicMesh/Private/Operations/PolyEditingEdgeUtil.cpp
--- a/Engine/Plugins/Runtime/GeometryProcessing/Source/DynamicMesh/Private/Operations/PolyEditingEdgeUtil.cpp
+++ b/Engine/Plugins/Runtime/GeometryProcessing/Source/DynamicMesh/Private/Operations/PolyEditingEdgeUtil.cpp
@@ -102,6 +102,54 @@ void UE::Geometry::SolveInsetVertexPositionsFromInsetLines(
 		FVector3d CurPos = Mesh.GetVertex(VertexIDs[vi]);
 		VertexPositionsOut[vi] = SolveInsetVertexPositionFromLinePair(CurPos, PrevLine, NextLine);
 	}
+}
+
+
+void UE::Geometry::ComputeNewGroupIDsAlongEdgeLoop(
+	FDynamicMesh3& Mesh,
+	const TArray<int32>& LoopEdgeIDs,
+	TArray<int32>& NewLoopEdgeGroupIDs,
+	TArray<int32>& NewGroupIDsOut,
+	TFunctionRef<bool(int32 Eid1, int32 Eid2)> EdgesShouldHaveSameGroupFunc)
+{
+	int32 NumEdgeIDs = LoopEdgeIDs.Num();
+	NewLoopEdgeGroupIDs.SetNumUninitialized(NumEdgeIDs);
+
+	if (!ensure(NumEdgeIDs > 2))
+	{
+		if (NumEdgeIDs > 0)
+		{
+			int32 OneGroupID = Mesh.AllocateTriangleGroup();
+			NewLoopEdgeGroupIDs.Init(OneGroupID, NumEdgeIDs);
+			NewGroupIDsOut.Add(OneGroupID);
+		}
+		return;
+	}
 
+	NewLoopEdgeGroupIDs[0] = Mesh.AllocateTriangleGroup();
+	NewGroupIDsOut.Add(NewLoopEdgeGroupIDs[0]);
 
+	// Propagate the group backwards first so we don't allocate an unnecessary group
+	// at the end and then have to fix it.
+	int32 LastDifferentGroupIndex = NumEdgeIDs - 1;
+	while (LastDifferentGroupIndex > 0
+		&& EdgesShouldHaveSameGroupFunc(LoopEdgeIDs[0], LoopEdgeIDs[LastDifferentGroupIndex]) )
+	{
+		NewLoopEdgeGroupIDs[LastDifferentGroupIndex] = NewLoopEdgeGroupIDs[0];
+		--LastDifferentGroupIndex;
+	}
+
+	// Now add new groups forward
+	for (int32 j = 1; j <= LastDifferentGroupIndex; ++j)
+	{
+		if ( ! EdgesShouldHaveSameGroupFunc(LoopEdgeIDs[j], LoopEdgeIDs[j-1]) )
+		{
+			NewLoopEdgeGroupIDs[j] = Mesh.AllocateTriangleGroup();
+			NewGroupIDsOut.Add(NewLoopEdgeGroupIDs[j]);
+		}
+		else
+		{
+			NewLoopEdgeGroupIDs[j] = NewLoopEdgeGroupIDs[j-1];
+		}
+	}
 }
diff --git a/Engine/Plugins/Runtime/GeometryProcessing/Source/DynamicMesh/Private/Operations/PolyEditingUVUtil.cpp b/Engine/Plugins/Runtime/GeometryProcessing/Source/DynamicMesh/Private/Operations/PolyEditingUVUtil.cpp
--- a/Engine/Plugins/Runtime/GeometryProcessing/Source/DynamicMesh/Private/Operations/PolyEditingUVUtil.cpp
+++ b/Engine/Plugins/Runtime/GeometryProcessing/Source/DynamicMesh/Private/Operations/PolyEditingUVUtil.cpp
@@ -47,4 +47,60 @@ void UE::Geometry::ComputeArbitraryTrianglePatchUVs(
 	{
 		return (float)UseUVScale * UV;
 	});
+}
+
+
+
+
+double UE::Geometry::ComputeAverageUVScaleRatioAlongVertexPath(
+	const FDynamicMesh3& Mesh,
+	const FDynamicMeshUVOverlay& UVOverlay,
+	const TArray<int32>& VertexPath,
+	double* PathLengthOut,
+	double* UVPathLengthOut )
+{
+	double MeshLength = 0;
+	double UVLength = 0;
+	int32 N = VertexPath.Num();
+	for ( int32 k = 0; k < N-1; ++k )
+	{
+		int32 EdgeID = Mesh.FindEdge(VertexPath[k], VertexPath[k+1]);
+		if (EdgeID != IndexConstants::InvalidID)
+		{
+			double EdgeLength = Distance(Mesh.GetVertex(VertexPath[k]), Mesh.GetVertex(VertexPath[k+1]));
+			FIndex2i EdgeTris = Mesh.GetEdgeT(EdgeID);
+
+			double EdgeUVLength = 0;
+			int EdgeUVLengthCount = 0;
+			for ( int32 j = 0; j < 2; ++j )
+			{
+				if ( EdgeTris[j] != IndexConstants::InvalidID && UVOverlay.IsSetTriangle(EdgeTris[j]))
+				{
+					int32 EdgeIdx = Mesh.GetTriEdges(EdgeTris[j]).IndexOf(EdgeID);
+					if ( EdgeIdx != IndexConstants::InvalidID )
+					{
+						FVector2f UVs[3];
+						UVOverlay.GetTriElements(EdgeTris[j], UVs[0], UVs[1], UVs[2]);
+						EdgeUVLength += (double)Distance( UVs[EdgeIdx], UVs[(EdgeIdx+1)%3] );
+						EdgeUVLengthCount++;
+					}
+				}
+			}
+			if ( EdgeUVLengthCount > 0 )
+			{
+				EdgeUVLength /= (double)EdgeUVLengthCount;
+				MeshLength += EdgeLength;
+				UVLength += EdgeUVLength;
+			}
+		}
+	}
+	if ( PathLengthOut )
+	{
+		*PathLengthOut = MeshLength;
+	}
+	if ( UVPathLengthOut )
+	{
+		*UVPathLengthOut = UVLength;
+	}
+	return ( MeshLength > FMathf::ZeroTolerance ) ? (UVLength / MeshLength) : 1.0;
 }
\ No newline at end of file
diff --git a/Engine/Plugins/Runtime/GeometryProcessing/Source/DynamicMesh/Private/Operations/QuadGridPatchUtil.cpp b/Engine/Plugins/Runtime/GeometryProcessing/Source/DynamicMesh/Private/Operations/QuadGridPatchUtil.cpp
new file mode 100644
--- /dev/null
+++ b/Engine/Plugins/Runtime/GeometryProcessing/Source/DynamicMesh/Private/Operations/QuadGridPatchUtil.cpp
@@ -0,0 +1,96 @@
+// Copyright Epic Games, Inc. All Rights Reserved.
+
+#include "Operations/QuadGridPatchUtil.h"
+#include "DynamicMesh/DynamicMesh3.h"
+#include "DynamicMeshEditor.h"
+#include "Parameterization/DynamicMeshUVEditor.h"
+#include "Operations/PolyEditingUVUtil.h"
+
+
+using namespace UE::Geometry;
+
+
+
+void UE::Geometry::ComputeNormalsForQuadPatch(
+	FDynamicMesh3& Mesh,
+	const FQuadGridPatch& QuadPatch )
+{
+	TArray<int32> PatchTriangles;
+	QuadPatch.GetAllTriangles(PatchTriangles);
+	FDynamicMeshEditor Editor(&Mesh);
+	Editor.SetTriangleNormals(PatchTriangles);
+}
+
+
+
+
+bool UE::Geometry::ComputeUVIslandForQuadPatch(
+	FDynamicMesh3& Mesh,
+	const FQuadGridPatch& QuadPatch,
+	double UVScaleFactor,
+	int UVOverlayIndex )
+{
+	FDynamicMeshUVEditor UVEditor(&Mesh, UVOverlayIndex, false);
+	FDynamicMeshUVOverlay* UVOverlay = UVEditor.GetOverlay();
+	if (!UVOverlay)
+	{
+		return false;
+	}
+
+	TArray<int32> AllTriangles;
+	QuadPatch.GetAllTriangles(AllTriangles);
+	UVEditor.ResetUVs(AllTriangles);
+
+	// compute scale factor that tries to keep UV dimensions roughly consistent with connected area
+	double UVLengthScale = UE::Geometry::ComputeAverageUVScaleRatioAlongVertexPath(Mesh, *UVOverlay, QuadPatch.VertexSpans[0]);
+	// should we consider other spans here??
+
+	TArray<TArray<int32>> UVElementSpans;
+	UVElementSpans.SetNum( QuadPatch.NumVertexRowsV );
+	const TArray<int32>& StartVertexSpan = QuadPatch.VertexSpans[0];
+	const TArray<int32>& EndVertexSpan = QuadPatch.VertexSpans.Last();
+	TArray<int32>& StartUVSpan = UVElementSpans[0];
+	TArray<int32>& EndUVSpan = UVElementSpans.Last();
+
+	// generate new UV elements along the two vertex spans
+	// note even in the case of a closed loop, we have a span here, so we will get an 'unwrapped' cylinder etc
+	double AccumDistU = 0;
+	int32 NumU = QuadPatch.NumVertexColsU;
+	int32 NumV = QuadPatch.NumVertexRowsV;
+	for ( int32 k = 0; k < NumU; ++k )
+	{
+		double DistV = Distance( Mesh.GetVertex(StartVertexSpan[k]), Mesh.GetVertex(EndVertexSpan[k]) );
+		//double DistV = Distance( Mesh.GetVertex(StartVertexSpan[0]), Mesh.GetVertex(EndVertexSpan[0]) );
+		float UseU = (float)(UVLengthScale * AccumDistU * UVScaleFactor);
+		float EndV = (float)(UVLengthScale * DistV * UVScaleFactor);
+		for ( int32 j = 0; j < NumV; ++j )
+		{
+			float t = (float)j / (float)(NumV-1);
+			int32 NewElemID = UVOverlay->AppendElement( FVector2f(UseU, (float)(t * EndV)) );
+			UVElementSpans[j].Add(NewElemID);
+		}
+
+		if ( k < NumU-1 )
+		{
+			AccumDistU += Distance( Mesh.GetVertex(StartVertexSpan[k]), Mesh.GetVertex(StartVertexSpan[k+1]) );
+		}
+	}
+
+	QuadPatch.ForEachQuad( [&](int32 Row, int32 QuadIndex, FIndex2i QuadTris)
+	{
+		for (int32 TriIdx = 0; TriIdx < 2; ++TriIdx)
+		{
+			int32 TriangleID = QuadTris[TriIdx];
+			FIndex3i IndicesTri = QuadPatch.GetQuadTriMappedToSpanIndices(Mesh, Row, QuadIndex, TriIdx);
+			FIndex3i UVTri;
+			for ( int32 j = 0; j < 3; ++j )
+			{
+				FIndex2i VertexSpanIndex = QuadPatch.DecodeRowColumnIndex(IndicesTri[j]);
+				UVTri[j] = UVElementSpans[VertexSpanIndex.A][VertexSpanIndex.B];
+			}
+			UVOverlay->SetTriangle(TriangleID, UVTri);
+		}
+	});
+
+	return true;
+}
diff --git a/Engine/Plugins/Runtime/GeometryProcessing/Source/DynamicMesh/Private/Selections/QuadGridPatch.cpp b/Engine/Plugins/Runtime/GeometryProcessing/Source/DynamicMesh/Private/Selections/QuadGridPatch.cpp
new file mode 100644
--- /dev/null
+++ b/Engine/Plugins/Runtime/GeometryProcessing/Source/DynamicMesh/Private/Selections/QuadGridPatch.cpp
@@ -0,0 +1,234 @@
+// Copyright Epic Games, Inc. All Rights Reserved.
+
+#include "Selections/QuadGridPatch.h"
+#include "DynamicMesh/DynamicMesh3.h"
+#include "Algo/Reverse.h"
+
+using namespace UE::Geometry;
+
+
+namespace UELocal
+{
+
+static bool FindTriEdgeIndices(
+	const FIndex3i& Triangle, 
+	int32 EdgeVert0, int32 EdgeVert1,
+	int32& EdgeIndex0, int32& EdgeIndex1, int32& OtherIndex2 )
+{
+	for ( int32 j = 0; j < 3; ++j )
+	{
+		if (Triangle[j] == EdgeVert0)
+		{
+			if ( Triangle[(j+1)%3] == EdgeVert1 )
+			{
+				EdgeIndex0 = j; EdgeIndex1 = (j+1)%3; OtherIndex2 = (j+2)%3; return true;
+			}
+			else if ( Triangle[(j+2)%3] == EdgeVert1 )
+			{
+				EdgeIndex0 = j; EdgeIndex1 = (j+2)%3; OtherIndex2 = (j+1)%3; return true;
+			}
+		}
+	}
+	return false;
+}
+
+}
+
+
+void FQuadGridPatch::InitializeFromQuadStrip(const FDynamicMesh3& Mesh, const TArray<FIndex2i>& SequentialQuadsIn, const TArray<int32>& FirstVertexSpan )
+{
+	int32 NumV = FirstVertexSpan.Num();
+	int32 NumQ = SequentialQuadsIn.Num();
+	check(NumQ == NumV-1);
+
+	NumVertexColsU = NumV;
+	NumVertexRowsV = 2;
+
+	VertexSpans.SetNum(2);
+	VertexSpans[0] = FirstVertexSpan;
+	VertexSpans[1].Reserve(NumV);
+	QuadTriangles.SetNum(1);
+	QuadTriangles[0] = SequentialQuadsIn;
+
+	for ( int32 k = 0; k < NumQ; ++k )
+	{
+		int Vertex0 = FirstVertexSpan[k];
+		int Vertex1 = FirstVertexSpan[k+1];
+
+		FIndex2i& QuadTris = QuadTriangles[0][k];
+		FIndex3i TriA = Mesh.GetTriangle(QuadTris.A);
+		FIndex3i TriB = Mesh.GetTriangle(QuadTris.B);
+		// we want TriA to be the one containing the edge (Vertex0,Vertex1), so if it's TriB, swap them, and the quad
+		if ( TriB.Contains(Vertex0) && TriB.Contains(Vertex1) )
+		{
+			Swap(TriA, TriB);
+			Swap(QuadTris.A, QuadTris.B);
+		}
+		// at this point we are not sure if OtherVertex matches with Vertex0 or Vertex1
+		int OtherVertex = IndexUtil::FindTriOtherVtx(Vertex0, Vertex1, TriA);
+		bool bIsCase1 = (TriB.Contains(Vertex0) == false);
+		int Other0 = (bIsCase1) ? OtherVertex : IndexUtil::FindTriOtherVtx(Vertex0, OtherVertex, TriB);
+		int Other1 = (bIsCase1) ? IndexUtil::FindTriOtherVtx(Vertex1, OtherVertex, TriB) : OtherVertex;
+
+		if (k == 0)
+		{
+			VertexSpans[1].Add(Other0);
+		}
+		checkSlow(FirstVertexSpan.Last() == Other0);
+		VertexSpans[1].Add(Other1);
+	}
+}
+
+
+
+void FQuadGridPatch::ReverseRows()
+{
+	Algo::Reverse(QuadTriangles);
+	Algo::Reverse(VertexSpans);
+}
+
+
+
+bool FQuadGridPatch::AppendQuadPatchRows(FQuadGridPatch&& NextPatch, bool bChecked)
+{
+	if ( NextPatch.NumVertexColsU != NumVertexColsU ) return false;
+
+	if ( bChecked )
+	{
+		const TArray<int32>& LastSpan = VertexSpans.Last();
+		const TArray<int32>& FirstSpan = NextPatch.VertexSpans[0];
+		for ( int32 k = 0; k < NumVertexColsU; ++k )
+		{
+			if (LastSpan[k] != FirstSpan[k])
+			{
+				check(false);
+				return false;
+			}
+		}
+	}
+
+	int AddedV = NextPatch.NumVertexRowsV-1;
+	for ( int32 k = 1; k < NextPatch.NumVertexRowsV; ++k )
+	{
+		VertexSpans.Add( MoveTemp(NextPatch.VertexSpans[k]) );
+	}
+	for ( int32 k = 0; k < NextPatch.NumVertexRowsV-1; ++k )
+	{
+		QuadTriangles.Add( MoveTemp(NextPatch.QuadTriangles[k]) );
+	}
+	NumVertexRowsV = VertexSpans.Num();
+	return true;
+}
+
+
+void FQuadGridPatch::GetAllTriangles(TArray<int32>& AllTrianglesOut) const
+{
+	AllTrianglesOut.Reserve(AllTrianglesOut.Num() + 2*GetNumQuads());
+	for ( const TArray<FIndex2i>& QuadStrip : QuadTriangles )
+	{
+		for (FIndex2i Quad : QuadStrip)
+		{
+			AllTrianglesOut.Add(Quad.A);
+			AllTrianglesOut.Add(Quad.B);
+		}
+	}
+}
+
+
+
+FIndex3i FQuadGridPatch::GetQuadTriMappedToSpanIndices(const FDynamicMesh3& Mesh, int QuadRow, int QuadIndex, int TriIndex) const
+{
+	check(TriIndex == 0 || TriIndex == 1);
+	check(QuadRow >= 0 && QuadRow < (NumVertexRowsV-1) );
+	check(QuadIndex >= 0 && QuadIndex < (NumVertexColsU-1) );
+	FIndex3i Tri = Mesh.GetTriangle( QuadTriangles[QuadRow][QuadIndex][TriIndex] );
+	//        C - D
+	//        |   |
+	//        A - B
+	FIndex4i QuadVertices(
+		VertexSpans[QuadRow][QuadIndex], VertexSpans[QuadRow][QuadIndex+1],
+		VertexSpans[QuadRow+1][QuadIndex], VertexSpans[QuadRow+1][QuadIndex+1] );
+	FIndex4i QuadIndices(
+		EncodeRowColumnIndex(QuadRow, QuadIndex), EncodeRowColumnIndex(QuadRow,QuadIndex+1),
+		EncodeRowColumnIndex(QuadRow+1, QuadIndex), EncodeRowColumnIndex(QuadRow+1,QuadIndex+1));
+	FIndex3i TriOrder, EncodedIndexTri = FIndex3i::Invalid();
+	if ( UELocal::FindTriEdgeIndices(Tri, QuadVertices.A, QuadVertices.D, TriOrder.A, TriOrder.B, TriOrder.C) )
+	{
+		// edge A/D exists case
+		EncodedIndexTri[TriOrder.A] = QuadIndices.A;
+		EncodedIndexTri[TriOrder.B] = QuadIndices.D;
+		EncodedIndexTri[TriOrder.C] = (Tri[TriOrder.C] == QuadVertices.C) ? QuadIndices.C : QuadIndices.B;
+	}
+	else if ( UELocal::FindTriEdgeIndices(Tri, QuadVertices.B, QuadVertices.C, TriOrder.A, TriOrder.B, TriOrder.C) )
+	{
+		// edge B/C exists case
+		EncodedIndexTri[TriOrder.A] = QuadIndices.B;
+		EncodedIndexTri[TriOrder.B] = QuadIndices.C;
+		EncodedIndexTri[TriOrder.C] = (Tri[TriOrder.C] == QuadVertices.D) ? QuadIndices.D : QuadIndices.A;
+	}
+	else
+	{
+		check(false);		// ruh roh! it should have been one of those two other cases...
+	}
+	return EncodedIndexTri;
+}
+
+
+
+void FQuadGridPatch::GetSubPatchByQuadRange(int StartRow, int EndRow, int StartCol, int EndCol, FQuadGridPatch& PatchOut ) const
+{
+	int NumQuadRows = (EndRow-StartRow)+1;
+	int NumQuadCols = (EndCol-StartCol)+1;
+
+	PatchOut.NumVertexRowsV = NumQuadRows+1;
+	PatchOut.NumVertexColsU = NumQuadCols+1;
+	PatchOut.VertexSpans.SetNum(PatchOut.NumVertexRowsV);
+	PatchOut.QuadTriangles.SetNum(PatchOut.NumVertexRowsV-1);
+	for ( int32 ri = 0; ri < NumQuadRows; ++ri )
+	{
+		int32 row = StartRow + ri;
+		for ( int32 ci = 0; ci < NumQuadCols; ++ci )
+		{
+			int col = StartCol + ci;
+			PatchOut.VertexSpans[ri].Add( VertexSpans[row][col] );
+			PatchOut.QuadTriangles[ri].Add( QuadTriangles[row][col] );
+		}
+		PatchOut.VertexSpans[ri].Add( VertexSpans[row][EndCol+1] );
+	}
+	// last row
+	for ( int32 ci = 0; ci <= NumQuadCols; ++ci )
+	{
+		int col = StartCol + ci;
+		PatchOut.VertexSpans[NumQuadRows].Add( VertexSpans[EndRow+1][col] );
+	}
+}
+
+
+void FQuadGridPatch::SplitColumnsByPredicate( 
+	TFunctionRef<bool(int32 Column0, int32 Column1)> PredicateFunc,
+	TArray<FQuadGridPatch>& SplitPatchesOut) const
+{
+	int32 NumQuadCols = NumVertexColsU-1;
+	int32 CurIndex = 0;
+	int32 SpanStartIndex = 0;
+	while ( CurIndex < NumQuadCols-1 )
+	{
+		bool bShouldSplit = PredicateFunc(CurIndex, CurIndex+1);
+		if ( bShouldSplit )
+		{
+			FQuadGridPatch SubPatch;
+			GetSubPatchByQuadRange(0, NumVertexRowsV-2, SpanStartIndex, CurIndex, SubPatch );
+			SplitPatchesOut.Add(MoveTemp(SubPatch));
+			CurIndex++;
+			SpanStartIndex = CurIndex;
+		}
+		else
+		{
+			CurIndex++;
+		}
+	}
+	// will we always have one final patch?
+	FQuadGridPatch SubPatch;
+	GetSubPatchByQuadRange(0, NumVertexRowsV-2, SpanStartIndex, CurIndex, SubPatch );
+	SplitPatchesOut.Add(MoveTemp(SubPatch));
+}
\ No newline at end of file
diff --git a/Engine/Plugins/Runtime/GeometryProcessing/Source/DynamicMesh/Public/DynamicMeshEditor.h b/Engine/Plugins/Runtime/GeometryProcessing/Source/DynamicMesh/Public/DynamicMeshEditor.h
--- a/Engine/Plugins/Runtime/GeometryProcessing/Source/DynamicMesh/Public/DynamicMeshEditor.h
+++ b/Engine/Plugins/Runtime/GeometryProcessing/Source/DynamicMesh/Public/DynamicMeshEditor.h
@@ -212,6 +212,19 @@ public:
 	 */
 	bool DisconnectTriangles(const TArray<int>& Triangles, TArray<FLoopPairSet>& LoopSetOut, bool bHandleBoundaryVertices);
 
+	/**
+	 * Variant of DisconnectTriangles that takes a precomputed set of BoundaryLoops of the triangles, and a TSet of Triangles.
+	 * These are both computed internally by the version that only takes a TArray of Triangles.
+	 * @warning If the BoundaryLoops are not correct for the provided TriangleSet, this version will likely crash.
+	 * @param Triangles set of triangles
+	 * @param BoundaryLoops precomputed set of boundary EdgeLoops of TriangleSet
+	 * @param LoopSetOut returned set of boundary loops. LoopA is original loop which remains with "outer" triangles previously connected to TriangleSet, and LoopB is new boundary loop of TriangleSet
+	 * @param bAllowBoundaryVertices if true, mesh boundary vertices are duplicated and left on the 'outer' loop, otherwise function aborts and returns false if a boundary vertex is encountered
+	 */
+	bool DisconnectTriangles(const TSet<int>& TriangleSet, const TArray<FEdgeLoop>& BoundaryLoops, TArray<FLoopPairSet>& LoopSetOut, bool bAllowBoundaryVertices);
+
+
+
 	/**
 	* Disconnects triangles (without constructing boundary loops) such that the input Triangles are not connected to any other triangles in the mesh
 	* @param Triangles set of triangles
diff --git a/Engine/Plugins/Runtime/GeometryProcessing/Source/DynamicMesh/Public/Operations/OffsetMeshRegion.h b/Engine/Plugins/Runtime/GeometryProcessing/Source/DynamicMesh/Public/Operations/OffsetMeshRegion.h
--- a/Engine/Plugins/Runtime/GeometryProcessing/Source/DynamicMesh/Public/Operations/OffsetMeshRegion.h
+++ b/Engine/Plugins/Runtime/GeometryProcessing/Source/DynamicMesh/Public/Operations/OffsetMeshRegion.h
@@ -32,15 +32,21 @@ class DYNAMICMESH_API FOffsetMeshRegion
 {
 public:
 
-
+	/**
+	 * Method used to determine the per-vertex offset direction vector
+	 */
 	enum class EVertexExtrusionVectorType
 	{
+		/** No geometric offset vector, client must define offset via custom OffsetPositionFunc */
 		Zero,
+		/** Offset vector is per-vertex normal, either provided by FDynamicMesh3 Per-Vertex Normals if defined, or computed by averaging one-ring face normals */
 		VertexNormal,
-		// Angle weighted average of the triangles in the selection that contain this vertex (unit length)
+		/** Angle weighted average of the triangles in the selection that contain this vertex(unit length) */
 		SelectionTriNormalsAngleWeightedAverage,
-		// Like SelectionTriNormalsAngleWeightedAverage, but with the vertex length adjusted to try to keep
-		// the selection triangles parallel to their original location.
+		/** 
+		 *  Like SelectionTriNormalsAngleWeightedAverage, but with the vertex length adjusted to try to keep
+		 *  the selection triangles parallel to their original location. 
+		 */
 		SelectionTriNormalsAngleWeightedAdjusted,
 	};
 
@@ -54,7 +60,7 @@ public:
 	/** The triangle region we are modifying */
 	TArray<int32> Triangles;
 
-	/** This function is called to generate the offset vertex position. */
+	/** OffsetPositionFunc function is called to generate the offset vertex position. */
 	TFunction<FVector3d(const FVector3d& Position, const FVector3d& VertexVector, int Vid)> OffsetPositionFunc = 
 		[this](const FVector3d& Position, const FVector3d& VertexVector, int Vid)
 	{
@@ -93,6 +99,30 @@ public:
 	/** When using SelectionTriNormalsAngleWeightedAdjusted to keep triangles parallel, clamp the offset scale to at most this factor */
 	double MaxScaleForAdjustingTriNormalsOffset = 4.0;
 
+	/** 
+	 * Number of subdivisions along the length of the "tubes" used to stitch the offset regions.
+	 * This parameter is only supported in EVersion::Version1 and later
+	 */
+	int32 NumSubdivisions = 0;
+
+	/**
+	 * Support for different versions of the OffsetMeshRegion geometric operation.
+	 * We default to the latest version.
+	 */
+	enum class EVersion
+	{
+		// always use most recent version
+		Current = 0,
+		// initial implementation, various limitations and issues, maintained for back-compat
+		Legacy = 1,
+		// support NumSubdivisions, UVs and Normals computed by polygroup
+		Version1 = 2,
+	};
+
+	/** Version of Offset operation being used. Defaults to Current, ie most recent version */
+	EVersion UseVersion = EVersion::Current;
+	
+
 	//
 	// Outputs
 	//
@@ -174,6 +204,9 @@ public:
 protected:
 
 	virtual bool ApplyOffset(FOffsetInfo& Region);
+
+	virtual bool ApplyOffset_Legacy(FOffsetInfo& Region);
+	virtual bool ApplyOffset_Version1(FOffsetInfo& Region);
 };
 
 } // end namespace UE::Geometry
diff --git a/Engine/Plugins/Runtime/GeometryProcessing/Source/DynamicMesh/Public/Operations/PolyEditingEdgeUtil.h b/Engine/Plugins/Runtime/GeometryProcessing/Source/DynamicMesh/Public/Operations/PolyEditingEdgeUtil.h
--- a/Engine/Plugins/Runtime/GeometryProcessing/Source/DynamicMesh/Public/Operations/PolyEditingEdgeUtil.h
+++ b/Engine/Plugins/Runtime/GeometryProcessing/Source/DynamicMesh/Public/Operations/PolyEditingEdgeUtil.h
@@ -54,5 +54,26 @@ DYNAMICMESH_API FVector3d SolveInsetVertexPositionFromLinePair(
 	const FLine3d& InsetEdgeLine2);
 
 
+/**
+ * Compute a new set of GroupIDs along an edge loop under the assumption that 
+ * the edge loop is about to be split and new geometry introduced. The decision
+ * is based on a provided EdgesShouldHaveSameGroupFunc function which is called
+ * with sequential pairs of edges. For example if this function returns true when 
+ * the GroupIDs are the same on both edges, then this function will create a new
+ * GroupID for each sequential span of non-matching GroupIDs (ie often the "expected" behavior).
+ * 
+ * @param LoopEdgeIDs the EdgeIDs of the loop
+ * @param NewLoopEdgeGroupIDsOut returned with new GroupID for each input EdgeID, same length as LoopEdgeIDs
+ * @param NewGroupIDsOut list of newly-allocated GroupIDs
+ * @param EdgesShouldHaveSameGroupFunc predicate that determines where new GroupIDs will be allocated
+ */
+DYNAMICMESH_API void ComputeNewGroupIDsAlongEdgeLoop(
+	FDynamicMesh3& Mesh,
+	const TArray<int32>& LoopEdgeIDs,
+	TArray<int32>& NewLoopEdgeGroupIDsOut,
+	TArray<int32>& NewGroupIDsOut,
+	TFunctionRef<bool(int32 Eid1, int32 Eid2)> EdgesShouldHaveSameGroupFunc);
+
+
 } // end namespace UE::Geometry
 } // end namespace UE
\ No newline at end of file
diff --git a/Engine/Plugins/Runtime/GeometryProcessing/Source/DynamicMesh/Public/Operations/PolyEditingUVUtil.h b/Engine/Plugins/Runtime/GeometryProcessing/Source/DynamicMesh/Public/Operations/PolyEditingUVUtil.h
--- a/Engine/Plugins/Runtime/GeometryProcessing/Source/DynamicMesh/Public/Operations/PolyEditingUVUtil.h
+++ b/Engine/Plugins/Runtime/GeometryProcessing/Source/DynamicMesh/Public/Operations/PolyEditingUVUtil.h
@@ -15,11 +15,38 @@ namespace UE
 namespace Geometry
 {
 
+/**
+ * Create a new single UV island for the provided TriangleSet in the given Mesh/UVOverlay,
+ * by computing an ExpMap parameterization based on an automatically-detected midpoint of the triangle patch.
+ * The UV scale is determined automatically based on the UV/3D area ratio of triangles adjacent
+ * to the TriangleSet (using a ratio of 1 if the triangles have no neighbours).
+ * 
+ * @param TriangleSet a list of triangles, must be a single connected component
+ */
 DYNAMICMESH_API void ComputeArbitraryTrianglePatchUVs(
 	FDynamicMesh3& Mesh, 
 	FDynamicMeshUVOverlay& UVOverlay,
 	const TArray<int32>& TriangleSet );
 
 
+/**
+ * Compute the total UVLength/MeshLength ratio along the given VertexPath.
+ * Other 3D lengths can then be scaled to comparable UV-space lengths by multiplying by this value.
+ * Edges must exist between the vertices along the path, and loops are not considered.
+ * Each UV edge for a 3D edge is considered, ie if a 3D edge corresponds to 2 UV edges, both are counted,
+ * and if the 3D edge has no UV edges, it is skipped.
+ * @param VertexPath list of vertices which are connected by edges
+ * @param PathLengthOut optional output param for sum of 3D edge lengths along path (arc length)
+ * @param UVPathLengthOut optional output param for sum of UV edge lengths along path
+ * @return UVPathLength / PathLength, or 1.0 if no valid UV edges along the path existed
+ */
+DYNAMICMESH_API double ComputeAverageUVScaleRatioAlongVertexPath(
+	const FDynamicMesh3& Mesh,
+	const FDynamicMeshUVOverlay& UVOverlay,
+	const TArray<int32>& VertexPath,
+	double* PathLengthOut = nullptr,
+	double* UVPathLengthOut = nullptr );
+
+
 } // end namespace UE::Geometry
 } // end namespace UE
\ No newline at end of file
diff --git a/Engine/Plugins/Runtime/GeometryProcessing/Source/DynamicMesh/Public/Operations/QuadGridPatchUtil.h b/Engine/Plugins/Runtime/GeometryProcessing/Source/DynamicMesh/Public/Operations/QuadGridPatchUtil.h
new file mode 100644
--- /dev/null
+++ b/Engine/Plugins/Runtime/GeometryProcessing/Source/DynamicMesh/Public/Operations/QuadGridPatchUtil.h
@@ -0,0 +1,34 @@
+// Copyright Epic Games, Inc. All Rights Reserved.
+
+#pragma once
+
+#include "Selections/QuadGridPatch.h"
+
+
+namespace UE
+{
+namespace Geometry
+{
+	class FDynamicMesh3;
+
+/**
+ * Create a new Normals "island" (ie surrounded by hard edges) for the QuadPatch of the Mesh,
+ */
+DYNAMICMESH_API void ComputeNormalsForQuadPatch(
+	FDynamicMesh3& Mesh,
+	const FQuadGridPatch& QuadPatch );
+
+/**
+ * Create a new UV Island for the QuadPatch triangles of the Mesh, 
+ * and then assign UVs based on the grid layout.
+ */
+DYNAMICMESH_API bool ComputeUVIslandForQuadPatch(
+	FDynamicMesh3& Mesh,
+	const FQuadGridPatch& QuadPatch,
+	double UVScaleFactor = 1.0,
+	int UVOverlayIndex = 0);
+
+
+
+} // end namespace UE::Geometry
+} // end namespace 
\ No newline at end of file
diff --git a/Engine/Plugins/Runtime/GeometryProcessing/Source/DynamicMesh/Public/Selections/QuadGridPatch.h b/Engine/Plugins/Runtime/GeometryProcessing/Source/DynamicMesh/Public/Selections/QuadGridPatch.h
new file mode 100644
--- /dev/null
+++ b/Engine/Plugins/Runtime/GeometryProcessing/Source/DynamicMesh/Public/Selections/QuadGridPatch.h
@@ -0,0 +1,134 @@
+// Copyright Epic Games, Inc. All Rights Reserved.
+
+#pragma once
+
+#include "IndexTypes.h"
+
+
+namespace UE
+{
+namespace Geometry
+{
+
+class FDynamicMesh3;
+
+/**
+ * FQuadGridPatch represents a set of quads, formed by pairs of triangles, arranged in a 2D grid pattern.
+ * U is the "across"/column direction and "V" is the "up"/row direction, and if you were thinking about it in terms
+ * of a piece of graph paper, (U=0,V=0) would be in the bottom-left.
+ * 
+ * Each row of Vertices is stored, as well as each row of Quads. In the case of a Quad grid that is a loop,
+ * the last vertex in each row is a duplicate of the first. Note, however, that there is currently no explicit 
+ * tracking of whether or not there is a Loop in U or V - this must be tracked externally. 
+ * 
+ * Generally the quad triangles will not be assumed to form any consistent pattern. This complicates some of
+ * the functions however, and perhaps this assumption could be optionally provided in the future
+ */
+class DYNAMICMESH_API FQuadGridPatch
+{
+public:
+	int NumVertexColsU;
+	int NumVertexRowsV;
+
+	/** NumVertexRowsV rows of NumVertexColsU VertexIDs, may contain repeated element if the patch forms a loop */
+	TArray<TArray<int32>> VertexSpans;
+
+	/** Quads stored as pairs of triangle indices, (NumVertexRowsV-1) rows of (NumVertexColsU-1) */
+	TArray<TArray<FIndex2i>> QuadTriangles;
+
+	/**
+	 * Initialize an Nx1 grid from a list of quads and the row of vertices on the "bottom" (ie first span)
+	 */
+	void InitializeFromQuadStrip(const FDynamicMesh3& Mesh, const TArray<FIndex2i>& SequentialQuadsIn, const TArray<int32>& FirstVertexSpan );
+
+	/**
+	 * Reverse the order of the vertex/quad rows, ie "flip" vertically
+	 */
+	void ReverseRows();
+
+	/**
+	 * Append the rows of NextPatch to the rows of the current QuadGrid.
+	 * NextPatch must have the same number of vertex-columns as this Patch.
+	 * Requires that the last vertex-row of this Patch be the same as the first
+	 * vertex-row of NextPatch. 
+	 * @param bChecked if true, function verifies requirements and returns false if a mismatch is detected
+	 * @return true if append succeeded
+	 */
+	bool AppendQuadPatchRows(FQuadGridPatch&& NextPatch, bool bChecked);
+
+
+	/** @return the total number of quads in the grid */
+	int32 GetNumQuads() const { return (NumVertexRowsV-1) * (NumVertexColsU-1); }
+
+	/** @return a list of all the triangles in the grid */
+	void GetAllTriangles(TArray<int32>& AllTrianglesOut) const;
+
+	/** 
+	 * Map a row and column index pair into a single int32. 11 Bits are available for the row, and 20 for the column 
+	 */
+	int32 EncodeRowColumnIndex(int32 Row, int32 Column ) const
+	{
+		return (Row << 20) | Column;
+	}
+	/**
+	 * Recover row and column indices from an int32 encoded via EncodeRowColumnIndex
+	 */
+	FIndex2i DecodeRowColumnIndex(int32 Encoded) const
+	{
+		return FIndex2i( (Encoded >> 20), Encoded & 0xFFFFF );
+	}
+
+	/**
+	 * In cases like assigning UVs or Normals to a QuadGridPatch, we want to assign a unique element
+	 * to each of the vertices in the grid, and then generate "element triangles" that index into those
+	 * new elements. To do this we need to know the (row,column) pair ("Vertex Span Index") for each triangle vertex. 
+	 * This function returns a 3-tuple of vertex span indices for a given quad triangle.
+	 * The function EncodeRowColumnIndex() is used to map the (row,column) to a single integer, 
+	 * and callers can use DecodeRowColumnIndex() to recover the (row,column)
+	 * 
+	 * @param QuadRow the row of quads the requested triangle is in
+	 * @param QuadIndex the "column" index of the quad containing the request triangle
+	 * @param TriIndex the index of the triangle in the quad, only 0 and 1 are valid
+	 * @return tuple of encoded (row,column) vertex span indices, in the same ordering as the identified triangle
+	 */
+	FIndex3i GetQuadTriMappedToSpanIndices(const FDynamicMesh3& Mesh, int QuadRow, int QuadIndex, int TriIndex) const;
+
+
+	/**
+	 * Call ApplyFunc for each quad, passing it's grid index and triangle pair
+	 */
+	void ForEachQuad(TFunctionRef<void(int32 QuadRow, int32 QuadIndex, FIndex2i QuadTris)> ApplyFunc) const
+	{
+		for ( int32 Row = 0; Row < NumVertexRowsV-1; ++Row )
+		{
+			for (int32 Column = 0; Column < NumVertexColsU - 1; ++Column)
+			{
+				ApplyFunc(Row, Column, QuadTriangles[Row][Column]);
+			}
+		}
+	}
+
+
+	/**
+	 * Construct a new QuadGridPatch that is a subset of the current QuadGridPatch
+	 * The values (Start|End)Quad(Row|Col) define the sub-patch, the values are inclusive,
+	 * so passing (1,2,1,2) defines a 2x2 quad patch
+	 */
+	void GetSubPatchByQuadRange(int StartQuadRow, int EndQuadRow, int StartQuadCol, int EndQuadCol, FQuadGridPatch& PatchOut ) const;
+	
+
+	/**
+	 * Split the Columns of the QuadGridPatch into new QuadGridPatches where the columns
+	 * are grouped via the PredicateFunc. 
+	 */
+	void SplitColumnsByPredicate( 
+		TFunctionRef<bool(int32 QuadColumn0, int32 QuadColumn1)> PredicateFunc,
+		TArray<FQuadGridPatch>& SplitPatchesOut) const;
+	
+};
+
+
+
+
+} // end namespace UE::Geometry
+} // end namespace UE
diff --git a/Engine/Plugins/Runtime/MeshModelingToolset/Source/ModelingOperators/Private/DeformationOps/ExtrudeOp.cpp b/Engine/Plugins/Runtime/MeshModelingToolset/Source/ModelingOperators/Private/DeformationOps/ExtrudeOp.cpp
--- a/Engine/Plugins/Runtime/MeshModelingToolset/Source/ModelingOperators/Private/DeformationOps/ExtrudeOp.cpp
+++ b/Engine/Plugins/Runtime/MeshModelingToolset/Source/ModelingOperators/Private/DeformationOps/ExtrudeOp.cpp
@@ -222,6 +222,7 @@ bool FExtrudeOp::MoveAndStitchExtrude(FProgressCancel* Progress)
 	UpdateSelectionGroups(*ResultMesh, TriangleSelection, SelectionGids);
 
 	FOffsetMeshRegion Extruder(ResultMesh.Get());
+	Extruder.UseVersion = FOffsetMeshRegion::EVersion::Legacy;		// to reconsider...
 	Extruder.Triangles = TriangleSelection;
 
 	Extruder.LoopEdgesShouldHaveSameGroup = [this](int32 Eid1, int32 Eid2) {
