GeometryProcessing: improve OffsetMeshRegion operation
Add support for inferring MaterialID, various methods for controlling creation of new polygroups in FOffsetMeshRegion
Add new PolyModelingFaceUtil.h/cpp with new function ComputeFaceSelectionFrame()
Add new PolyModelingMaterialUtil.h/cpp with new functions ComputeMaterialIDRange() and ComputeMaterialIDsForVertexPath()
Add Row/Column count accessor methods to FQuadGridPatch, and new function GetQuadOpeningAngle() that can return opening angle between two quad columns
Improve ComputeUVIslandForQuadPatch, now considers adjacent strips on either side to determine average target UV scale

diff --git a/Engine/Plugins/Runtime/GeometryProcessing/Source/DynamicMesh/Private/Operations/OffsetMeshRegion.cpp b/Engine/Plugins/Runtime/GeometryProcessing/Source/DynamicMesh/Private/Operations/OffsetMeshRegion.cpp
--- a/Engine/Plugins/Runtime/GeometryProcessing/Source/DynamicMesh/Private/Operations/OffsetMeshRegion.cpp
+++ b/Engine/Plugins/Runtime/GeometryProcessing/Source/DynamicMesh/Private/Operations/OffsetMeshRegion.cpp
@@ -16,6 +16,7 @@
 #include "Operations/PolyEditingUVUtil.h"
 #include "Operations/PolyEditingEdgeUtil.h"
 #include "Operations/QuadGridPatchUtil.h"
+#include "Operations/PolyModeling/PolyModelingMaterialUtil.h" 
 
 using namespace UE::Geometry;
 
@@ -122,6 +123,19 @@ static void LeftShiftArray(TArray<ValueType>& Values, int ShiftNum)
 }
 
 
+struct FStitchConfigOptions
+{
+	int NumSubdivisions = 0;
+	double UVScaleFactor = 1.0;
+	bool bGroupPerSubdivision = true;
+	bool bUVIslandPerGroup = true;
+	double CreaseAngleThreshold = 180.0;
+
+	int SetMaterialID = 0;
+	bool bInferMaterialID = false;
+};
+
+
 
 /**
  * Stitch pairs of border loops of Mesh defined by LoopPairs.
@@ -132,8 +146,7 @@ static bool StitchRegionBorderLoopPairs_Version1(
 	const TArray<FDynamicMeshEditor::FLoopPairSet>& LoopPairs,
 	const TArray<TArray<int32>>& LoopsPerEdgeNewGroupIDs,
 	const TArray<TArray<FMeshTriOrderedEdgeID>>& InnerTriOrderedEdgeLoops,
-	double UVScaleFactor,
-	int NumSubdivisions,
+	FStitchConfigOptions StitchConfig,
 	TArray<bool>& bLoopOK,
 	TArray<TArray<int32>>& PerLoopTrianglesOut,
 	TArray<TArray<int32>>& PerLoopGroupsOut )
@@ -156,19 +169,32 @@ static bool StitchRegionBorderLoopPairs_Version1(
 		const TArray<FMeshTriOrderedEdgeID>& InnerTriOrderedEdgeLoop = InnerTriOrderedEdgeLoops[LoopIndex];
 		int32 NV = OuterLoopV.Num();
 
+		// populate list of Material IDs attached to inner loop
+		TArray<int32> MaterialIDs;
+		FDynamicMeshMaterialAttribute* MaterialIDAttrib = (Mesh.HasAttributes() && Mesh.Attributes()->HasMaterialID()) ? 
+				Mesh.Attributes()->GetMaterialID() : nullptr;
+		if ( MaterialIDAttrib && StitchConfig.bInferMaterialID )
+		{
+			UE::Geometry::ComputeMaterialIDsForVertexPath(Mesh, InnerLoopV, true, MaterialIDs, StitchConfig.SetMaterialID);
+		}
+		else if ( MaterialIDAttrib )
+		{
+			MaterialIDs.Init(StitchConfig.SetMaterialID, NV);
+		}
+
 		// LoopStack is the set of NumSubdivisions+1 loops that will be stitched,
 		// where the "last" loop is the Outer loop. The "first" Inner loop is not 
 		// included in LoopStack, it has special handling below
 		TArray<TArray<int32>> LoopStack;
-		LoopStack.SetNum(NumSubdivisions+1);
-		LoopStack[NumSubdivisions] = OuterLoopV;
-		int LastTemp = LoopStack[NumSubdivisions][0];
+		LoopStack.SetNum(StitchConfig.NumSubdivisions+1);
+		LoopStack[StitchConfig.NumSubdivisions] = OuterLoopV;
+		int LastTemp = LoopStack[StitchConfig.NumSubdivisions][0];
 
 		// build interior loops if there are any
-		for ( int32 StackIdx = 0; StackIdx < NumSubdivisions; ++StackIdx )
+		for ( int32 StackIdx = 0; StackIdx < StitchConfig.NumSubdivisions; ++StackIdx )
 		{
 			LoopStack[StackIdx].Reserve(NV);
-			double t = (double)(StackIdx+1) / (double)(NumSubdivisions+1);
+			double t = (double)(StackIdx+1) / (double)(StitchConfig.NumSubdivisions+1);
 			for ( int32 k = 0; k < NV; ++k )
 			{
 				FVector3d A = Mesh.GetVertex(InnerLoopV[k]);
@@ -197,7 +223,7 @@ static bool StitchRegionBorderLoopPairs_Version1(
 		QuadPatch.ReverseRows();	// need LoopStack[0] to be the second span but it's currently the first one...
 
 		// incrementally stitch intermediate strips and accumulate into the Quad Patch
-		for ( int32 k = 1; k <= NumSubdivisions; ++k )
+		for ( int32 k = 1; k <= StitchConfig.NumSubdivisions; ++k )
 		{
 			FDynamicMeshEditResult StripResult;
 			bLoopOK[LoopIndex] = Editor.StitchVertexLoopsMinimal(LoopStack[k-1], LoopStack[k], StripResult);
@@ -214,17 +240,34 @@ static bool StitchRegionBorderLoopPairs_Version1(
 
 		// set group IDs on patch, propagating up the subdivisions columns
 		const TArray<int32>& PerEdgeNewGroupIDs = LoopsPerEdgeNewGroupIDs[LoopIndex];		// is it always guaranteed to be the same length??
-		QuadPatch.ForEachQuad( [&](int32 QuadRow, int32 QuadIndex, FIndex2i QuadTris)
+		QuadPatch.ForEachQuad( [&](int32 QuadRow, int32 QuadCol, FIndex2i QuadTris)
 		{
-			int32 GroupID = PerEdgeNewGroupIDs[QuadIndex];
+			int32 GroupID = PerEdgeNewGroupIDs[QuadCol];
 			Mesh.SetTriangleGroup(QuadTris.A, GroupID);
 			Mesh.SetTriangleGroup(QuadTris.B, GroupID);
+			if (MaterialIDAttrib)
+			{
+				MaterialIDAttrib->SetValue(QuadTris.A, MaterialIDs[QuadCol]);
+				MaterialIDAttrib->SetValue(QuadTris.B, MaterialIDs[QuadCol]);
+			}
 		});
 
 		// split quad patch by group ID into a set of quad patches
+		bool bSplitDueToCrease = false;
 		TArray<FQuadGridPatch> GroupStrips;
 		QuadPatch.SplitColumnsByPredicate(
-			[&](int32 ColumnIdx, int32 NextColumnIdx) { return PerEdgeNewGroupIDs[ColumnIdx] != PerEdgeNewGroupIDs[NextColumnIdx]; },
+			[&](int32 ColumnIdx, int32 NextColumnIdx) { 
+				if (PerEdgeNewGroupIDs[ColumnIdx] != PerEdgeNewGroupIDs[NextColumnIdx])
+				{
+					return true;
+				}
+				else if (QuadPatch.GetQuadOpeningAngleDeg(Mesh, NextColumnIdx, NextColumnIdx+1, 0) > StitchConfig.CreaseAngleThreshold)
+				{
+					bSplitDueToCrease = true;
+					return true;
+				}
+				return false;
+			},
 			GroupStrips);
 
 		// compute normals and UVs for each group-quad-patch
@@ -233,28 +276,54 @@ static bool StitchRegionBorderLoopPairs_Version1(
 			for (FQuadGridPatch& GroupPatch : GroupStrips)
 			{
 				UE::Geometry::ComputeNormalsForQuadPatch(Mesh, GroupPatch);
-				UE::Geometry::ComputeUVIslandForQuadPatch(Mesh, GroupPatch, UVScaleFactor);
+				if ( StitchConfig.bUVIslandPerGroup )
+				{
+					UE::Geometry::ComputeUVIslandForQuadPatch(Mesh, GroupPatch, StitchConfig.UVScaleFactor);
+				}
 			}
 		}
+		// if not per-group UV islands, do it for the entire strip)
+		if (StitchConfig.bUVIslandPerGroup == false)
+		{
+			UE::Geometry::ComputeUVIslandForQuadPatch(Mesh, QuadPatch, StitchConfig.UVScaleFactor);
+		}
 
-		// Assign a new group to each row of group-quad-patch
-		for (FQuadGridPatch& GroupPatch : GroupStrips)
+		// if we split groups due to a crease, we need to reassign groups in row 0
+		if (bSplitDueToCrease)
 		{
-			int32 NumQuadRows = GroupPatch.NumVertexRowsV-1;
-			TArray<int32> RowGroups;
-			RowGroups.SetNum(NumQuadRows);
-			for ( int32 k = 1; k < NumQuadRows; ++k )
+			for (FQuadGridPatch& GroupPatch : GroupStrips)
 			{
-				RowGroups[k] = Mesh.AllocateTriangleGroup();
+				int32 NewPatchGroup = Mesh.AllocateTriangleGroup();
+				GroupPatch.ForEachQuad([&](int32 Row, int32 Col, FIndex2i QuadTris)
+				{
+					Mesh.SetTriangleGroup(QuadTris.A, NewPatchGroup);
+					Mesh.SetTriangleGroup(QuadTris.B, NewPatchGroup);
+				});
 			}
-			GroupPatch.ForEachQuad([&](int32 Row, int32 Col, FIndex2i QuadTris)
+		}
+
+		// Assign a new group to each row of group-quad-patch
+		if ( StitchConfig.bGroupPerSubdivision )
+		{
+			for (FQuadGridPatch& GroupPatch : GroupStrips)
 			{
-				if ( Row > 0 )
+				int32 NumQuadRows = GroupPatch.NumVertexRowsV-1;
+				TArray<int32> RowGroups;
+				RowGroups.SetNum(NumQuadRows);
+				// if we split at a crease we need to reassign first column too...
+				for ( int32 k =1; k < NumQuadRows; ++k )
 				{
-					Mesh.SetTriangleGroup(QuadTris.A, RowGroups[Row]);
-					Mesh.SetTriangleGroup(QuadTris.B, RowGroups[Row]);
+					RowGroups[k] = Mesh.AllocateTriangleGroup();
 				}
-			});
+				GroupPatch.ForEachQuad([&](int32 Row, int32 Col, FIndex2i QuadTris)
+				{
+					if ( Row > 0 )
+					{
+						Mesh.SetTriangleGroup(QuadTris.A, RowGroups[Row]);
+						Mesh.SetTriangleGroup(QuadTris.B, RowGroups[Row]);
+					}
+				});
+			}
 		}
 
 		// save the stitch triangles set and associated group IDs
@@ -464,12 +533,27 @@ bool FOffsetMeshRegion::ApplyOffset_Version1(FOffsetInfo& Region)
 	FDynamicMeshEditResult IgnoreResult;
 	TmpEditor.SplitBowtiesAtTriangles(RegionTriangles, IgnoreResult);
 
-	// Store offset groups
+	TMap<int32, int32> OffsetGroupMap;
+
+	// Remap GroupIDs in offset region
 	if (Mesh->HasTriangleGroups())
 	{
 		for (int32 TriangleID : RegionTriangles)
 		{
-			Region.OffsetGroups.AddUnique(Mesh->GetTriangleGroup(TriangleID));
+			int32 CurGroupID = Mesh->GetTriangleGroup(TriangleID);
+			int32 NewGroupID = CurGroupID;
+			int32* FoundNewGroupID = OffsetGroupMap.Find(CurGroupID);
+			if (FoundNewGroupID == nullptr)
+			{
+				NewGroupID = Mesh->AllocateTriangleGroup();
+				OffsetGroupMap.Add(CurGroupID, NewGroupID);
+				Region.OffsetGroups.Add(NewGroupID);
+			}
+			else
+			{
+				NewGroupID = *FoundNewGroupID;
+			}
+			Mesh->SetTriangleGroup(TriangleID, NewGroupID);
 		}
 	}
 
@@ -621,10 +705,18 @@ bool FOffsetMeshRegion::ApplyOffset_Version1(FOffsetInfo& Region)
 	}
 
 	// Stitch the loops
+	OffsetMeshRegionLocals::FStitchConfigOptions StitchConfig;
+	StitchConfig.UVScaleFactor = this->UVScaleFactor;
+	StitchConfig.NumSubdivisions = this->NumSubdivisions;
+	StitchConfig.bGroupPerSubdivision = this->bGroupPerSubdivision;
+	StitchConfig.bUVIslandPerGroup = this->bUVIslandPerGroup;
+	StitchConfig.CreaseAngleThreshold = this->CreaseAngleThresholdDeg;
+	StitchConfig.SetMaterialID = this->SetMaterialID;
+	StitchConfig.bInferMaterialID = this->bInferMaterialID;
 	TArray<bool> bLoopSuccess;
 	bool bSuccess = OffsetMeshRegionLocals::StitchRegionBorderLoopPairs_Version1( *Mesh, 
 			LoopPairs, LoopsEdgeGroups, OffsetTriOrderedEdgeLoops,
-			UVScaleFactor, NumSubdivisions,
+			StitchConfig,
 			bLoopSuccess, Region.StitchTriangles, Region.StitchPolygonIDs);
 
 	int NumInitialLoops = LoopPairs.Num();
@@ -652,6 +744,16 @@ bool FOffsetMeshRegion::ApplyOffset_Version1(FOffsetInfo& Region)
 		}
 	}
 
+	if (bSingleGroupPerArea && Mesh->HasTriangleGroups() && Region.OffsetGroups.Num() > 1)
+	{
+		int32 NewGroupID = Region.OffsetGroups[0];
+		for (int32 TriangleID : RegionTriangles)
+		{
+			Mesh->SetTriangleGroup(TriangleID, NewGroupID);
+		}
+		Region.OffsetGroups.SetNum(1);
+	}
+
 	return bSuccess;
 }
 
diff --git a/Engine/Plugins/Runtime/GeometryProcessing/Source/DynamicMesh/Private/Operations/PolyModeling/PolyModelingFaceUtil.cpp b/Engine/Plugins/Runtime/GeometryProcessing/Source/DynamicMesh/Private/Operations/PolyModeling/PolyModelingFaceUtil.cpp
new file mode 100644
--- /dev/null
+++ b/Engine/Plugins/Runtime/GeometryProcessing/Source/DynamicMesh/Private/Operations/PolyModeling/PolyModelingFaceUtil.cpp
@@ -0,0 +1,124 @@
+// Copyright Epic Games, Inc. All Rights Reserved.
+
+#include "Operations/PolyModeling/PolyModelingFaceUtil.h"
+#include "DynamicMesh/DynamicMesh3.h"
+#include "DynamicMesh/DynamicMeshAABBTree3.h"
+#include "Selections/MeshConnectedComponents.h"
+
+
+using namespace UE::Geometry;
+
+UE::Geometry::FFrame3d UE::Geometry::ComputeFaceSelectionFrame(
+	const FDynamicMesh3& Mesh, 
+	const TArray<int32>& Triangles,
+	bool bIsDefinitelySingleComponent)
+{
+	if (Triangles.Num() == 0)
+	{
+		return FFrame3d();
+	}
+
+	// special case for 1 triangle
+	if (Triangles.Num() == 1)
+	{
+		FVector3d TriNormal, TriCentroid; double TriArea;
+		Mesh.GetTriInfo(Triangles[0], TriNormal, TriArea, TriCentroid);
+		return FFrame3d(TriCentroid, TriNormal);
+	}
+
+	// If there is more than one component we want to put the frame on the largest component
+	const TArray<int32>* UseTriangles = &Triangles;
+	FMeshConnectedComponents Components(&Mesh);
+	if (bIsDefinitelySingleComponent == false)
+	{
+		Components.FindConnectedTriangles(Triangles);
+		UseTriangles = &Components.GetComponent(Components.GetLargestIndexByCount()).Indices;		// area would be better here?
+	}
+
+	// special case for 1 triangle again
+	if (UseTriangles->Num() == 1)
+	{
+		FVector3d TriNormal, TriCentroid; double TriArea;
+		Mesh.GetTriInfo(Triangles[0], TriNormal, TriArea, TriCentroid);
+		return FFrame3d(TriCentroid, TriNormal);
+	}
+
+	// compute area-weighted average triangle normal and centroid
+	FVector3d LargestNormal = FVector3d::UnitZ();
+	double LargestArea = 0;
+	FVector3d AverageNormal = FVector3d::Zero();
+	FVector3d Centroid = FVector3d::Zero();
+	double WeightSum = 0;
+	for ( int32 tid : *UseTriangles )
+	{
+		FVector3d TriNormal, TriCentroid; double TriArea;
+		Mesh.GetTriInfo(tid, TriNormal, TriArea, TriCentroid);
+		Centroid += TriArea * TriCentroid;
+		AverageNormal += TriArea * TriNormal;
+		WeightSum += TriArea;
+		if (TriArea > LargestArea)
+		{
+			LargestArea = TriArea;
+			LargestNormal = TriNormal;
+		}
+	}
+	Centroid /= (WeightSum > FMathd::ZeroTolerance) ? WeightSum : 1.0;
+	Normalize(AverageNormal);
+	FVector3d UseNormal = ( AverageNormal.SquaredLength() < 0.98 ) ? LargestNormal : AverageNormal;
+
+	// if selection is planar (very common case) then centroid is on mesh and we can stop searching
+	FInterval1d ProjectionRange = FInterval1d::Empty();
+	for ( int32 tid : *UseTriangles )
+	{
+		FVector3d A,B,C;
+		ProjectionRange.Contain( (A-Centroid).Dot(UseNormal) );
+		ProjectionRange.Contain( (B-Centroid).Dot(UseNormal) );
+		ProjectionRange.Contain( (C-Centroid).Dot(UseNormal) );
+		if (ProjectionRange.Length() > FMathf::ZeroTolerance)
+		{
+			break;
+		}
+	}
+	if (ProjectionRange.Length() < FMathf::ZeroTolerance)
+	{
+		return FFrame3d(Centroid, UseNormal);
+	}
+
+	// Try to convert the centroid to a point on the selection surface via raycats.
+	// probably would be good to have some heuristics here, ie if < 100 triangles just
+	// do a linear iteration instead of building a BVTree...
+	// (could also only build BVTree for the subset of triangles??)
+	FDynamicMeshAABBTree3 BVTree(&Mesh, true);
+	TArray<FVector3d, TInlineAllocator<3>> HitPoints;
+	// cast rays forward and backward
+	const double Signs[2] = {1.0, -1.0};
+	for ( int32 j = 0; j < 2; ++j )
+	{
+		double RayParam; int32 HitTriangleID;
+		FRay3d Ray(Centroid, Signs[j]*UseNormal);
+		if (BVTree.FindNearestHitTriangle(Ray, RayParam, HitTriangleID))
+		{
+			if ( UseTriangles->Contains(HitTriangleID) )
+			{
+				HitPoints.Add(Ray.PointAt(RayParam));
+			}
+		}
+	}
+	// this actually seems a bit like a bad idea...if both rays miss we probably prefer the centroid
+	//HitPoints.Add(BVTree.FindNearestPoint(Centroid));
+
+	// find the hit point that is closest. Could also consider alignment w/ average normal?
+	double Nearest = TNumericLimits<double>::Max();
+	FVector3d UsePoint = Centroid;
+	for ( FVector3d Pos : HitPoints )
+	{
+		double Dist = Distance(Centroid, Pos);
+		if ( Dist < Nearest )
+		{
+			Nearest = Dist;
+			UsePoint = Pos;
+		}
+	}
+
+	return FFrame3d(UsePoint, UseNormal);
+}
\ No newline at end of file
diff --git a/Engine/Plugins/Runtime/GeometryProcessing/Source/DynamicMesh/Private/Operations/PolyModeling/PolyModelingMaterialUtil.cpp b/Engine/Plugins/Runtime/GeometryProcessing/Source/DynamicMesh/Private/Operations/PolyModeling/PolyModelingMaterialUtil.cpp
new file mode 100644
--- /dev/null
+++ b/Engine/Plugins/Runtime/GeometryProcessing/Source/DynamicMesh/Private/Operations/PolyModeling/PolyModelingMaterialUtil.cpp
@@ -0,0 +1,58 @@
+// Copyright Epic Games, Inc. All Rights Reserved.
+
+#include "Operations/PolyModeling/PolyModelingMaterialUtil.h"
+#include "DynamicMesh/DynamicMesh3.h"
+#include "DynamicMesh/DynamicMeshAttributeSet.h"
+
+using namespace UE::Geometry;
+
+
+bool UE::Geometry::ComputeMaterialIDRange(
+	const FDynamicMesh3& Mesh,
+	FInterval1i& MaterialIDRange)
+{
+	MaterialIDRange = FInterval1i(0,0);
+	const FDynamicMeshMaterialAttribute* MaterialIDAttrib = (Mesh.HasAttributes() && Mesh.Attributes()->HasMaterialID()) ? 
+			Mesh.Attributes()->GetMaterialID() : nullptr;
+	if (MaterialIDAttrib)
+	{
+		for (int32 tid : Mesh.TriangleIndicesItr())
+		{
+			MaterialIDRange.Contain(MaterialIDAttrib->GetValue(tid));
+		}
+		return true;
+	}
+	return false;
+}
+
+
+
+bool UE::Geometry::ComputeMaterialIDsForVertexPath(
+	const FDynamicMesh3& Mesh,
+	const TArray<int32>& VertexPath,
+	bool bIsLoop,
+	TArray<int32>& EdgeMaterialIDs,
+	int32 FallbackMaterialID)
+{
+	int32 NumV = VertexPath.Num();
+	EdgeMaterialIDs.Init(FallbackMaterialID, bIsLoop ? NumV : NumV-1);
+
+	const FDynamicMeshMaterialAttribute* MaterialIDAttrib = (Mesh.HasAttributes() && Mesh.Attributes()->HasMaterialID()) ? 
+			Mesh.Attributes()->GetMaterialID() : nullptr;
+	if (MaterialIDAttrib == nullptr)
+	{
+		return false;
+	}
+
+	int32 NumEdges = EdgeMaterialIDs.Num();
+	for (int32 k = 0; k < NumEdges; ++k)
+	{
+		int32 EdgeID = Mesh.FindEdge(VertexPath[k], VertexPath[(k+1)%NumV]);
+		if (EdgeID != IndexConstants::InvalidID)
+		{
+			EdgeMaterialIDs[k] = MaterialIDAttrib->GetValue( Mesh.GetEdgeT(EdgeID).A );
+		}
+	}
+
+	return true;
+}
\ No newline at end of file
diff --git a/Engine/Plugins/Runtime/GeometryProcessing/Source/DynamicMesh/Private/Operations/QuadGridPatchUtil.cpp b/Engine/Plugins/Runtime/GeometryProcessing/Source/DynamicMesh/Private/Operations/QuadGridPatchUtil.cpp
--- a/Engine/Plugins/Runtime/GeometryProcessing/Source/DynamicMesh/Private/Operations/QuadGridPatchUtil.cpp
+++ b/Engine/Plugins/Runtime/GeometryProcessing/Source/DynamicMesh/Private/Operations/QuadGridPatchUtil.cpp
@@ -41,9 +41,20 @@ bool UE::Geometry::ComputeUVIslandForQuadPatch(
 	QuadPatch.GetAllTriangles(AllTriangles);
 	UVEditor.ResetUVs(AllTriangles);
 
-	// compute scale factor that tries to keep UV dimensions roughly consistent with connected area
-	double UVLengthScale = UE::Geometry::ComputeAverageUVScaleRatioAlongVertexPath(Mesh, *UVOverlay, QuadPatch.VertexSpans[0]);
-	// should we consider other spans here??
+	// compute UV scale factor that tries to keep UV dimensions roughly consistent with connected area
+	double UVLengthScale = 0, UVLengthWeight = 0;
+	for (int32 k = 0; k < 2; ++k)
+	{
+		const TArray<int32>& VertexSpan = (k == 0) ? QuadPatch.VertexSpans[0] : QuadPatch.VertexSpans.Last();
+		double PathLength = 0;
+		double UVLengthScaleTmp = UE::Geometry::ComputeAverageUVScaleRatioAlongVertexPath(Mesh, *UVOverlay, VertexSpan, &PathLength);
+		if (PathLength > 0)
+		{
+			UVLengthScale += UVLengthScaleTmp;
+			UVLengthWeight += 1.0;
+		}
+	}
+	UVLengthScale = (UVLengthWeight == 0 || UVLengthScale == 0) ? 1.0 : (UVLengthScale / UVLengthWeight);
 
 	TArray<TArray<int32>> UVElementSpans;
 	UVElementSpans.SetNum( QuadPatch.NumVertexRowsV );
diff --git a/Engine/Plugins/Runtime/GeometryProcessing/Source/DynamicMesh/Private/Selections/QuadGridPatch.cpp b/Engine/Plugins/Runtime/GeometryProcessing/Source/DynamicMesh/Private/Selections/QuadGridPatch.cpp
--- a/Engine/Plugins/Runtime/GeometryProcessing/Source/DynamicMesh/Private/Selections/QuadGridPatch.cpp
+++ b/Engine/Plugins/Runtime/GeometryProcessing/Source/DynamicMesh/Private/Selections/QuadGridPatch.cpp
@@ -231,4 +231,20 @@ void FQuadGridPatch::SplitColumnsByPredicate(
 	FQuadGridPatch SubPatch;
 	GetSubPatchByQuadRange(0, NumVertexRowsV-2, SpanStartIndex, CurIndex, SubPatch );
 	SplitPatchesOut.Add(MoveTemp(SubPatch));
+}
+
+
+
+double FQuadGridPatch::GetQuadOpeningAngleDeg(const FDynamicMesh3& Mesh, int32 Column0, int32 Column1, int32 Row) const
+{
+	check(Column1 == (Column0+1));
+	int32 Vertex0 = VertexSpans[Row][Column0];
+	int32 Vertex1 = VertexSpans[Row+1][Column0];
+	int32 EdgeID = Mesh.FindEdge(Vertex0, Vertex1);
+	check(EdgeID != IndexConstants::InvalidID);
+	FIndex2i EdgeTris = Mesh.GetEdgeT(EdgeID);
+	check(EdgeTris.B != IndexConstants::InvalidID);
+	FVector3d Normal0 = Mesh.GetTriNormal(EdgeTris.A);
+	FVector3d Normal1 = Mesh.GetTriNormal(EdgeTris.B);
+	return AngleD(Normal0, Normal1);
 }
\ No newline at end of file
diff --git a/Engine/Plugins/Runtime/GeometryProcessing/Source/DynamicMesh/Public/Operations/OffsetMeshRegion.h b/Engine/Plugins/Runtime/GeometryProcessing/Source/DynamicMesh/Public/Operations/OffsetMeshRegion.h
--- a/Engine/Plugins/Runtime/GeometryProcessing/Source/DynamicMesh/Public/Operations/OffsetMeshRegion.h
+++ b/Engine/Plugins/Runtime/GeometryProcessing/Source/DynamicMesh/Public/Operations/OffsetMeshRegion.h
@@ -105,6 +105,48 @@ public:
 	 */
 	int32 NumSubdivisions = 0;
 
+	/**
+	 * If true, each extruded area gets a single new group, instead of remapping the input groups
+	 * This parameter is only supported in EVersion::Version1 and later
+	 */
+	bool bSingleGroupPerArea = true;
+
+	/**
+	 * If true, each subdivision level gets a new group
+	 * This parameter is only supported in EVersion::Version1 and later
+	 */
+	bool bGroupPerSubdivision = true;
+
+	/** 
+	 * if true, new UV island is assigned for each group in stitch region. 
+	 * This parameter is only supported in EVersion::Version1 and later
+	 * This parameter is only supported in EVersion::Version1 and later
+	 */
+	bool bUVIslandPerGroup = true;
+
+	/**
+	 * Split the extrude "tube" into separate groups based on the opening angle between quads.
+	 * This split is done relative to the above group options, so even if the tube would 
+	 * only have a single group, but the extrude area is a square, it will be split into
+	 * groups based on the corners.
+	 * This parameter is only supported in EVersion::Version1 and later
+	 */
+	double CreaseAngleThresholdDeg = 180.0;
+
+
+	/** 
+	 * Constant Material ID to set along the extrusion tubes 
+	 * This parameter is only supported in EVersion::Version1 and later
+	 */
+	int SetMaterialID = 0;
+
+	/** 
+	 * If true, Material IDs around the border of the extrude area are propagated "down" the extrusion tube 
+	 * This parameter is only supported in EVersion::Version1 and later
+	 */
+	bool bInferMaterialID = true;
+
+
 	/**
 	 * Support for different versions of the OffsetMeshRegion geometric operation.
 	 * We default to the latest version.
diff --git a/Engine/Plugins/Runtime/GeometryProcessing/Source/DynamicMesh/Public/Operations/PolyModeling/PolyModelingFaceUtil.h b/Engine/Plugins/Runtime/GeometryProcessing/Source/DynamicMesh/Public/Operations/PolyModeling/PolyModelingFaceUtil.h
new file mode 100644
--- /dev/null
+++ b/Engine/Plugins/Runtime/GeometryProcessing/Source/DynamicMesh/Public/Operations/PolyModeling/PolyModelingFaceUtil.h
@@ -0,0 +1,36 @@
+// Copyright Epic Games, Inc. All Rights Reserved.
+
+#pragma once
+
+#include "MathUtil.h"
+#include "FrameTypes.h"
+
+class FProgressCancel;
+
+namespace UE
+{
+namespace Geometry
+{
+
+class FDynamicMesh3;
+
+
+/**
+ * Try to compute a good "Frame" for the selected Triangles of the Mesh.
+ * Only the largest triangle-connected-component of Triangles is considered.
+ * The Frame Z is aligned with the area-weighted average normal of the connected triangles.
+ * Initially the Frame is placed at the area-weighted centroid of the connected triangles.
+ * Rays are cast along Z in both directions, if either hits, the Frame is translated to
+ * the nearer hit.
+ * 
+ * @param bIsDefinitelySingleComponent if this is known, passing as true will speed up the computation
+ */
+DYNAMICMESH_API FFrame3d ComputeFaceSelectionFrame(
+	const FDynamicMesh3& Mesh, 
+	const TArray<int32>& Triangles,
+	bool bIsDefinitelySingleComponent = false);
+
+
+
+} // end namespace UE::Geometry
+} // end namespace UE
\ No newline at end of file
diff --git a/Engine/Plugins/Runtime/GeometryProcessing/Source/DynamicMesh/Public/Operations/PolyModeling/PolyModelingMaterialUtil.h b/Engine/Plugins/Runtime/GeometryProcessing/Source/DynamicMesh/Public/Operations/PolyModeling/PolyModelingMaterialUtil.h
new file mode 100644
--- /dev/null
+++ b/Engine/Plugins/Runtime/GeometryProcessing/Source/DynamicMesh/Public/Operations/PolyModeling/PolyModelingMaterialUtil.h
@@ -0,0 +1,44 @@
+// Copyright Epic Games, Inc. All Rights Reserved.
+
+#pragma once
+
+#include "MathUtil.h"
+#include "GeometryTypes.h"
+#include "IntBoxTypes.h"
+
+namespace UE
+{
+namespace Geometry
+{
+
+class FDynamicMesh3;
+
+
+/**
+ * Compute inclusive range [0...N] of MaterialID values currently in use on Mesh (ie range always includes 0)
+ * @return true if MaterialIDs attribute exists (if false, Range is returned as [0,0])
+ */
+DYNAMICMESH_API bool ComputeMaterialIDRange(
+	const FDynamicMesh3& Mesh, 
+	FInterval1i& MaterialIDRange );
+
+
+/**
+ * Compute a MaterialID for each edge along the VertexPath based on the MaterialID of adjacent triangles. 
+ * Currently, for each Edge, EdgeTriangle.A is used.
+ * @param VertexPath list of sequential vertices that are connected by mesh edges
+ * @param bIsLoop is the path a loop? 
+ * @param EdgeMaterialIDsOut returned list of per-edge Material ID. Length is VertexPath.Num-1 for open paths and VertexPath.Num for loops
+ * @param FallbackMaterialID default MaterialID to use if no Material attribute exists or no edge is found
+ * @return true if MaterialIDs attribute exists
+ */
+DYNAMICMESH_API bool ComputeMaterialIDsForVertexPath(
+	const FDynamicMesh3& Mesh, 
+	const TArray<int32>& VertexPath,
+	bool bIsLoop,
+	TArray<int32>& EdgeMaterialIDsOut,
+	int32 FallbackMaterialID = 0);
+
+
+} // end namespace UE::Geometry
+} // end namespace UE
\ No newline at end of file
diff --git a/Engine/Plugins/Runtime/GeometryProcessing/Source/DynamicMesh/Public/Selections/QuadGridPatch.h b/Engine/Plugins/Runtime/GeometryProcessing/Source/DynamicMesh/Public/Selections/QuadGridPatch.h
--- a/Engine/Plugins/Runtime/GeometryProcessing/Source/DynamicMesh/Public/Selections/QuadGridPatch.h
+++ b/Engine/Plugins/Runtime/GeometryProcessing/Source/DynamicMesh/Public/Selections/QuadGridPatch.h
@@ -36,6 +36,13 @@ public:
 	/** Quads stored as pairs of triangle indices, (NumVertexRowsV-1) rows of (NumVertexColsU-1) */
 	TArray<TArray<FIndex2i>> QuadTriangles;
 
+
+	int NumVertexCols() const { return NumVertexColsU; }
+	int NumVertexRows() const { return NumVertexRowsV; }
+	int NumQuadCols() const { return NumVertexColsU-1; }
+	int NumQuadRows() const { return NumVertexRowsV-1; }
+
+
 	/**
 	 * Initialize an Nx1 grid from a list of quads and the row of vertices on the "bottom" (ie first span)
 	 */
@@ -97,7 +104,7 @@ public:
 	/**
 	 * Call ApplyFunc for each quad, passing it's grid index and triangle pair
 	 */
-	void ForEachQuad(TFunctionRef<void(int32 QuadRow, int32 QuadIndex, FIndex2i QuadTris)> ApplyFunc) const
+	void ForEachQuad(TFunctionRef<void(int32 QuadRow, int32 QuadCol, FIndex2i QuadTris)> ApplyFunc) const
 	{
 		for ( int32 Row = 0; Row < NumVertexRowsV-1; ++Row )
 		{
@@ -125,6 +132,16 @@ public:
 		TFunctionRef<bool(int32 QuadColumn0, int32 QuadColumn1)> PredicateFunc,
 		TArray<FQuadGridPatch>& SplitPatchesOut) const;
 	
+
+
+	/**
+	 * Compute the opening angle between the two quads in Column0 and Column1 of the given Row.
+	 * This is based on the face normals of the two triangles connected to the edge between Column0 and Column1,
+	 * so any non-planarity of the quads is ignored.
+	 * @return unsigned opening angle in degrees in range [0,180], where 0 means quads are parallel
+	 */
+	double GetQuadOpeningAngleDeg(const FDynamicMesh3& Mesh, int32 Column0, int32 Column1, int32 Row) const;
+
 };
 
 
