Added a new operator to transfer bone weights from one Dynamic Mesh to another.

diff --git a/Engine/Plugins/Runtime/GeometryProcessing/Source/DynamicMesh/Private/Operations/TransferBoneWeights.cpp b/Engine/Plugins/Runtime/GeometryProcessing/Source/DynamicMesh/Private/Operations/TransferBoneWeights.cpp
new file mode 100644
--- /dev/null
+++ b/Engine/Plugins/Runtime/GeometryProcessing/Source/DynamicMesh/Private/Operations/TransferBoneWeights.cpp
@@ -0,0 +1,245 @@
+// // Copyright Epic Games, Inc. All Rights Reserved.
+
+#include "Operations/TransferBoneWeights.h"
+#include "DynamicMesh/DynamicMesh3.h"
+#include "DynamicMesh/DynamicVertexSkinWeightsAttribute.h"
+#include "DynamicMesh/DynamicMeshAttributeSet.h"
+#include "DynamicMesh/DynamicMeshAABBTree3.h"
+#include "Async/ParallelFor.h"
+#include "Util/ProgressCancel.h"
+#include "BoneIndices.h"
+#include "BoneWeights.h"
+#include "IndexTypes.h"
+#include "TransformTypes.h"
+
+using namespace UE::AnimationCore;
+using namespace UE::Geometry;
+
+namespace TransferBoneWeightsLocals 
+{
+	/**
+	 * Given a triangle and point on a triangle (via barycentric coordinates), compute the bone weights for the point.
+	 * 
+	 * @param OutWeights Interpolated weights for a vertex with Bary barycentric coordinates
+	 * @param TriVertices The vertices of a triangle containing the point we are interpolating the weights for
+	 * @param Bary Barycentric coordinates of the point
+	 * @param Attribute Attribute containing bone weights of the mesh that TriVertices belong to
+	 * @param SourceIndexToBone Optional map from bone index to bone name for the source mesh
+	 * @param TargetBoneToIndex OPtional map from bone name to bone index for the target mesh
+	 * @param bNormalizeToOne If true, OutWeights will be normalized to sum to 1.
+	 */
+    void InterpolateBoneWeights(FBoneWeights& OutWeights,
+								const FIndex3i& TriVertices,
+								const FVector3f& Bary,
+								const FDynamicMeshVertexSkinWeightsAttribute* Attribute,
+								const TMap<FBoneIndexType, FName>* SourceIndexToBone = nullptr,
+								const TMap<FName, FBoneIndexType>* TargetBoneToIndex = nullptr,
+								bool bNormalizeToOne = true)
+	{
+		FBoneWeights Weight1, Weight2, Weight3;
+		Attribute->GetValue(TriVertices[0], Weight1);
+		Attribute->GetValue(TriVertices[1], Weight2);
+		Attribute->GetValue(TriVertices[2], Weight3);
+
+		FBoneWeightsSettings BlendSettings;
+		BlendSettings.SetNormalizeType(bNormalizeToOne ? EBoneWeightNormalizeType::Always : EBoneWeightNormalizeType::None);
+		BlendSettings.SetBlendZeroInfluence(true);
+		OutWeights = FBoneWeights::Blend(Weight1, Weight2, Weight3, Bary[0], Bary[1], Bary[2], BlendSettings);
+
+		// Check if we need to remap the indices
+		if (SourceIndexToBone && TargetBoneToIndex) 
+		{
+			FBoneWeightsSettings BoneSettings;
+			BoneSettings.SetNormalizeType(EBoneWeightNormalizeType::None);
+
+			FBoneWeights MappedWeights;
+
+			for (int WeightIdx = 0; WeightIdx < OutWeights.Num(); ++WeightIdx)
+			{
+				const FBoneWeight& BoneWeight = OutWeights[WeightIdx];
+
+				FBoneIndexType FromIdx = BoneWeight.GetBoneIndex();
+				uint16 FromWeight = BoneWeight.GetRawWeight();
+
+				if (ensure(SourceIndexToBone->Contains(FromIdx))) // the map must contain the index
+				{
+					FName BoneName = (*SourceIndexToBone)[FromIdx];
+					if (TargetBoneToIndex->Contains(BoneName))
+					{
+						FBoneIndexType ToIdx = (*TargetBoneToIndex)[BoneName];
+						FBoneWeight MappedBoneWeight(ToIdx, FromWeight);
+						MappedWeights.SetBoneWeight(MappedBoneWeight, BoneSettings);
+					}
+					else 
+					{	
+						UE_LOG(LogGeometry, Warning, TEXT("FTransferBoneWeights: Bone name %s does not exist in the target mesh."), *BoneName.ToString());
+					}
+				}
+			}
+
+
+			if (MappedWeights.Num() == 0)
+			{
+				// If no bone mappings were found, add a single entry for the root bone
+				MappedWeights.SetBoneWeight(FBoneWeight(0, 1.0f), FBoneWeightsSettings());
+			}
+			else if (OutWeights.Num() != MappedWeights.Num() && bNormalizeToOne)
+			{	
+				// In case some of the bones were not mapped we need to renormalize
+				MappedWeights.Renormalize(FBoneWeightsSettings());
+			}
+
+			OutWeights = MappedWeights;
+		}
+	};
+
+	
+	FDynamicMeshVertexSkinWeightsAttribute* GetOrCreateSkinWeightsAttribute(FDynamicMesh3& InMesh, const FName& InProfileName)
+	{
+		FDynamicMeshVertexSkinWeightsAttribute* Attribute = InMesh.Attributes()->GetSkinWeightsAttribute(InProfileName);
+		if (Attribute == nullptr)
+		{
+			Attribute = new FDynamicMeshVertexSkinWeightsAttribute(&InMesh);
+			InMesh.Attributes()->AttachSkinWeightsAttribute(InProfileName, Attribute);
+		}
+		return Attribute;
+	};
+}
+
+FTransferBoneWeights::FTransferBoneWeights(const FDynamicMesh3* InSourceMesh, 
+								   const FName& InSourceProfileName,
+								   const FDynamicMeshAABBTree3* InSourceBVH)
+:
+SourceMesh(InSourceMesh),
+SourceProfileName(InSourceProfileName),
+SourceBVH(InSourceBVH)
+{
+	// If the BVH for the source mesh was not specified then create one
+	if (SourceBVH == nullptr)
+	{
+		InternalSourceBVH = MakeUnique<FDynamicMeshAABBTree3>(SourceMesh);
+	}
+}
+
+FTransferBoneWeights::~FTransferBoneWeights() 
+{
+}
+
+bool FTransferBoneWeights::Cancelled()
+{
+	return (Progress == nullptr) ? false : Progress->Cancelled();
+}
+
+EOperationValidationResult FTransferBoneWeights::Validate()
+{	
+	if (SourceMesh == nullptr) 
+	{
+		return EOperationValidationResult::Failed_UnknownReason;
+	}
+
+	// Either BVH was passed by the caller or was created internally in the constructor
+	if (SourceBVH == nullptr && InternalSourceBVH.IsValid() == false) 
+	{
+		return EOperationValidationResult::Failed_UnknownReason;
+	}
+
+	if (SourceMesh->HasAttributes() == false) 
+	{
+		return EOperationValidationResult::Failed_UnknownReason;
+	}
+	
+	if (SourceMesh->Attributes()->GetSkinWeightsAttribute(SourceProfileName) == nullptr) 
+	{
+		return EOperationValidationResult::Failed_UnknownReason;
+	}
+
+	return EOperationValidationResult::Ok;
+}
+
+bool FTransferBoneWeights::Compute(FDynamicMesh3& InOutTargetMesh, const FTransformSRT3d& InToWorld, const FName& InTargetProfileName)
+{	
+	if (Validate() != EOperationValidationResult::Ok) 
+	{
+		checkNoEntry();
+		return false;
+	}
+	
+	FDynamicMeshVertexSkinWeightsAttribute* TargetSkinWeights = TransferBoneWeightsLocals::GetOrCreateSkinWeightsAttribute(InOutTargetMesh, InTargetProfileName);
+	checkSlow(TargetSkinWeights);
+
+	bool bFailed = false;
+	
+	ParallelFor(InOutTargetMesh.MaxVertexID(), [&](int32 VertexID)
+	{
+		if (Cancelled() || bFailed) 
+		{
+			return;
+		}
+		
+		if (InOutTargetMesh.IsVertex(VertexID)) 
+		{
+			FVector3d Point = InOutTargetMesh.GetVertex(VertexID);
+		
+			FBoneWeights Weights;
+			if (Compute(Point, InToWorld, Weights) == false)
+			{
+				bFailed = true;
+				return;
+			}
+			
+			TargetSkinWeights->SetValue(VertexID, Weights);
+		}
+
+	}, bUseParallel ? EParallelForFlags::None : EParallelForFlags::ForceSingleThread);
+
+	if (Cancelled() || bFailed) 
+	{
+		return false;
+	}
+		
+	return true;
+}
+
+bool FTransferBoneWeights::Compute(const FVector3d& InPoint, const FTransformSRT3d& InToWorld, FBoneWeights& OutWeights) 
+{
+	const FDynamicMeshVertexSkinWeightsAttribute* SourceSkinWeights = SourceMesh->Attributes()->GetSkinWeightsAttribute(SourceProfileName);
+	checkSlow(SourceSkinWeights);
+
+	IMeshSpatial::FQueryOptions Options;
+	double NearestDistSqr;
+	int32 NearTriID;
+	
+	const FVector3d WorldPoint = InToWorld.TransformPosition(InPoint);
+	if (SourceBVH != nullptr) 
+	{ 
+		NearTriID = SourceBVH->FindNearestTriangle(WorldPoint, NearestDistSqr, Options);
+	}
+	else 
+	{
+		NearTriID = InternalSourceBVH->FindNearestTriangle(WorldPoint, NearestDistSqr, Options);
+	}
+
+	if (NearTriID == IndexConstants::InvalidID) 
+	{
+		checkNoEntry();
+		return false;
+	}
+
+	const FDistPoint3Triangle3d Query = TMeshQueries<FDynamicMesh3>::TriangleDistance(*SourceMesh, NearTriID, WorldPoint);
+	const FVector3d NearestPnt = Query.ClosestTrianglePoint;
+	const FIndex3i TriVertex = SourceMesh->GetTriangle(NearTriID);
+
+	const FVector3d Bary = VectorUtil::BarycentricCoords(NearestPnt,
+														 SourceMesh->GetVertexRef(TriVertex.A),
+														 SourceMesh->GetVertexRef(TriVertex.B),
+														 SourceMesh->GetVertexRef(TriVertex.C));
+
+	TransferBoneWeightsLocals::InterpolateBoneWeights(OutWeights,
+													  TriVertex,
+													  FVector3f((float)Bary[0], (float)Bary[1], (float)Bary[2]),
+													  SourceSkinWeights,
+													  SourceIndexToBone,
+													  TargetBoneToIndex);
+
+	return true;
+}
\ No newline at end of file
diff --git a/Engine/Plugins/Runtime/GeometryProcessing/Source/DynamicMesh/Public/Operations/TransferBoneWeights.h b/Engine/Plugins/Runtime/GeometryProcessing/Source/DynamicMesh/Public/Operations/TransferBoneWeights.h
new file mode 100644
--- /dev/null
+++ b/Engine/Plugins/Runtime/GeometryProcessing/Source/DynamicMesh/Public/Operations/TransferBoneWeights.h
@@ -0,0 +1,150 @@
+// Copyright Epic Games, Inc. All Rights Reserved.
+
+#pragma once
+
+#include "GeometryTypes.h"
+#include "UObject/NameTypes.h"
+
+
+// Forward declarations
+class FProgressCancel;
+
+namespace UE::AnimationCore { class FBoneWeights; }
+
+namespace UE::Geometry
+{
+	template<typename RealType> class TTransformSRT3;
+	typedef TTransformSRT3<double> FTransformSRT3d;
+
+	class FDynamicMesh3;
+	template<typename MeshType> class TMeshAABBTree3;
+	typedef TMeshAABBTree3<FDynamicMesh3> FDynamicMeshAABBTree3;
+}
+
+
+namespace UE
+{
+namespace Geometry
+{
+/**
+ * Transfer bone weights from one mesh (source) to another (target). 
+ * 
+ * 
+ * Example usage:
+ * 
+ * FDynamicMesh SourceMesh = ...; // Mesh we transfering weights from
+ * FDynamicMesh TargetMesh = ...; // Mesh we are transfering weights into
+ * 
+ * // These can be initialized by accessing reference skeleton of the corresponding USkinnedAsset
+ * TMap<uint16, FName> SourceIndexToBone = ...; // Map the bone index to the bone name for the source mesh
+ * TMap<FName, uint16> TargetBoneToIndex = ...; // Map the bone name to the bone index for the target mesh
+ *
+ * FTransferBoneWeights TransferBoneWeights(&SourceMesh, FSkeletalMeshAttributes::DefaultSkinWeightProfileName);
+ * TransferBoneWeights.SourceIndexToBone = &SourceIndexToBone; 
+ * TransferBoneWeights.TargetBoneToIndex = &TargetBoneToIndex;
+ * 
+ * // Optionally transform the target mesh. This is useful when you want to align two meshes in world space.
+ * FTransformSRT3d InToWorld = ...; 
+ * 
+ * if (TransferBoneWeights.Validate() == EOperationValidationResult::Ok) 
+ * {
+ *      TransferBoneWeights.Compute(TargetMesh, InToWorld, FSkeletalMeshAttributes::DefaultSkinWeightProfileName);
+ * }
+ */
+
+class DYNAMICMESH_API FTransferBoneWeights
+{
+public:
+
+	enum class ETransferBoneWeightsMethod : uint8
+	{
+        // For every vertex on the target mesh, find the closest point on the surface of the source mesh.
+        // This is usually a point on a triangle where the bone weights are interpolated via barycentric coordinates.
+		ClosestPointOnSurface = 0
+	};
+
+	//
+	// Optional Inputs
+	//
+	
+    /** Set this to be able to cancel the running operation. */
+	FProgressCancel* Progress = nullptr;
+
+	/** Enable/disable multi-threading. */
+	bool bUseParallel = true;
+
+    /** 
+     * We assume that each mesh inherits its reference skeleton from the same USkeleton asset. However, their internal 
+     * indexing can be different and hence when transfering weights we need to make sure we reference the correct bones 
+     * via their names instead of indices.
+     * 
+     * TODO: We can store this information inside DynamicVertexSkinWeightsAttribute when converting Skeletal Meshes to Dynamic Meshes.
+     */
+	const TMap<uint16, FName>* SourceIndexToBone = nullptr; // Map the bone index to the bone name for the source mesh
+	const TMap<FName, uint16>* TargetBoneToIndex = nullptr; // Map the bone name to the bone index for the target mesh
+	
+	/** The transfer method to compute the bone weights. */
+	ETransferBoneWeightsMethod TransferMethod = ETransferBoneWeightsMethod::ClosestPointOnSurface;
+
+protected:
+		
+	/** Source mesh we are transfering weights from. */
+	const FDynamicMesh3* SourceMesh;
+	
+	/** The name of the source mesh skinning profile name. */
+	FName SourceProfileName;
+
+	/** 
+	 * The caller can optionally specify the source mesh BVH in case this operator is run on multiple target meshes 
+	 * while the source mesh remains the same. Otherwise BVH tree will be computed.
+	 */
+	const FDynamicMeshAABBTree3* SourceBVH = nullptr;
+
+	/** If the caller doesn't pass BVH for the source mesh then we compute one. */
+	TUniquePtr<FDynamicMeshAABBTree3> InternalSourceBVH;
+
+public:
+
+	FTransferBoneWeights(const FDynamicMesh3* InSourceMesh, 
+					     const FName& InSourceProfileName,
+					     const FDynamicMeshAABBTree3* SourceBVH = nullptr); 
+	
+	virtual ~FTransferBoneWeights();
+
+	/**
+	 * @return EOperationValidationResult::Ok if we can apply operation, or error code if we cannot.
+	 */
+	virtual EOperationValidationResult Validate();
+
+	/**
+     * Transfer the bone weights from the source mesh to the given target mesh and store the result in the skin weight  
+     * attribute with the given profile name.
+	 * 
+	 * @param InOutTargetMesh Target mesh we are transfering weights into
+	 * @param InToWorld Transform applied to the input target mesh
+     * @param InTargetProfileName Skin weight profile name we are writing into. If the profile with that name exists,  
+     *       					  then the data will be overwritten, otherwise a new attribute will be created.
+     * 
+     * @return true if the algorithm succeeds, false if it failed or was canceled by the user.
+	 */
+	virtual bool Compute(FDynamicMesh3& InOutTargetMesh, const FTransformSRT3d& InToWorld, const FName& InTargetProfileName);
+
+	/**
+     * Compute a single bone weight for a given point.
+     *
+	 * @param InPoint Point for which we are computing a bone weight
+	 * @param InToWorld Transform applied to the point
+     * @param OutWeights Bone weight computed for the input transformed point
+	 * 
+     * @return true if the algorithm succeeds, false if it failed or was canceled by the user.
+	 */
+	virtual bool Compute(const FVector3d& InPoint, const FTransformSRT3d& InToWorld, UE::AnimationCore::FBoneWeights& OutWeights);
+
+protected:
+	
+    /** @return if true, abort the computation. */
+	virtual bool Cancelled();
+};
+
+} // end namespace UE::Geometry
+} // end namespace UE
