GeometryProcessing: add FMeshConnectedComponents::GrowToConnectedVertices() and GrowToConnectedEdges(). Add FGroupTopology::FindEdgeNbrEdges(), FindCornerNbrEdges(), FindCornerNbrCorners().

diff --git a/Engine/Plugins/Runtime/GeometryProcessing/Source/DynamicMesh/Private/GroupTopology.cpp b/Engine/Plugins/Runtime/GeometryProcessing/Source/DynamicMesh/Private/GroupTopology.cpp
--- a/Engine/Plugins/Runtime/GeometryProcessing/Source/DynamicMesh/Private/GroupTopology.cpp
+++ b/Engine/Plugins/Runtime/GeometryProcessing/Source/DynamicMesh/Private/GroupTopology.cpp
@@ -325,6 +325,22 @@ void FGroupTopology::FindEdgeNbrGroups(const TArray<int>& GroupEdgeIDs, TArray<i
 	}
 }
 
+void FGroupTopology::FindEdgeNbrEdges(int GroupEdgeID, TArray<int>& EdgesOut) const
+{
+	check(GroupEdgeID >= 0 && GroupEdgeID < Edges.Num());
+	const FGroupEdge& Edge = Edges[GroupEdgeID];
+	if ( Edge.EndpointCorners.A != IndexConstants::InvalidID )
+	{
+		FindCornerNbrEdges(Edge.EndpointCorners.A, EdgesOut);
+	}
+	if ( Edge.EndpointCorners.B != IndexConstants::InvalidID )
+	{
+		FindCornerNbrEdges(Edge.EndpointCorners.B, EdgesOut);
+	}
+}
+
+
+
 bool FGroupTopology::IsBoundaryEdge(int32 GroupEdgeID) const
 {
 	return Mesh->IsBoundaryEdge(Edges[GroupEdgeID].Span.Edges[0]);
@@ -448,6 +464,50 @@ void FGroupTopology::FindCornerNbrGroups(const TArray<int>& CornerIDs, TArray<in
 }
 
 
+void FGroupTopology::FindCornerNbrEdges(int CornerID, TArray<int>& EdgesOut) const
+{
+	check(CornerID >= 0 && CornerID < Corners.Num());
+	for (int GroupID : Corners[CornerID].NeighbourGroupIDs)
+	{
+		const FGroup* Group = FindGroupByID(GroupID);
+		for (const FGroupBoundary& Boundary : Group->Boundaries)
+		{
+			for (int32 EdgeID : Boundary.GroupEdges)
+			{
+				const FGroupEdge& Edge = Edges[EdgeID];
+				if (Edge.EndpointCorners.A == CornerID || Edge.EndpointCorners.B == CornerID)
+				{
+					EdgesOut.AddUnique(EdgeID);
+				}
+			}
+		}
+
+	}
+}
+
+void FGroupTopology::FindCornerNbrCorners(int CornerID, TArray<int>& CornersOut) const
+{
+	check(CornerID >= 0 && CornerID < Corners.Num());
+	for (int GroupID : Corners[CornerID].NeighbourGroupIDs)
+	{
+		const FGroup* Group = FindGroupByID(GroupID);
+		for (const FGroupBoundary& Boundary : Group->Boundaries)
+		{
+			for (int32 EdgeID : Boundary.GroupEdges)
+			{
+				const FGroupEdge& Edge = Edges[EdgeID];
+				if (Edge.EndpointCorners.A == CornerID || Edge.EndpointCorners.B == CornerID)
+				{
+					int32 OtherCornerID = Edge.EndpointCorners.OtherElement(CornerID);
+					CornersOut.AddUnique(OtherCornerID);
+				}
+			}
+		}
+
+	}
+}
+
+
 
 
 void FGroupTopology::FindVertexNbrGroups(int VertexID, TArray<int>& GroupsOut) const
diff --git a/Engine/Plugins/Runtime/GeometryProcessing/Source/DynamicMesh/Private/Selections/MeshConnectedComponents.cpp b/Engine/Plugins/Runtime/GeometryProcessing/Source/DynamicMesh/Private/Selections/MeshConnectedComponents.cpp
--- a/Engine/Plugins/Runtime/GeometryProcessing/Source/DynamicMesh/Private/Selections/MeshConnectedComponents.cpp
+++ b/Engine/Plugins/Runtime/GeometryProcessing/Source/DynamicMesh/Private/Selections/MeshConnectedComponents.cpp
@@ -737,4 +737,72 @@ void FMeshConnectedComponents::GrowToConnectedTriangles(const FDynamicMesh3* Mes
 			}
 		}
 	}
+}
+
+
+
+void FMeshConnectedComponents::GrowToConnectedVertices(const FDynamicMesh3& Mesh,
+	const TArray<int>& InputROI, TSet<int>& ResultROI,
+	TArray<int32>* QueueBuffer,
+	TFunctionRef<bool(int32, int32)> CanGrowPredicate )
+{
+	TArray<int32> LocalQueue;
+	QueueBuffer = (QueueBuffer == nullptr) ? &LocalQueue : QueueBuffer;
+	QueueBuffer->Reset();
+	QueueBuffer->Insert(InputROI, 0);
+
+	ResultROI.Reset();
+	ResultROI.Append(InputROI);
+
+	while (QueueBuffer->Num() > 0)
+	{
+		int32 CurVertexID = QueueBuffer->Pop(false);
+		ResultROI.Add(CurVertexID);
+
+		Mesh.EnumerateVertexVertices(CurVertexID, [&](int32 NbrVertexID)
+		{
+			if (ResultROI.Contains(NbrVertexID) == false && CanGrowPredicate(CurVertexID, NbrVertexID))
+			{
+				QueueBuffer->Add(NbrVertexID);
+				ResultROI.Add(NbrVertexID);
+			}
+		});
+
+	}
+}
+
+
+
+void FMeshConnectedComponents::GrowToConnectedEdges(const FDynamicMesh3& Mesh,
+	const TArray<int>& InputROI, TSet<int>& ResultROI,
+	TArray<int32>* QueueBuffer,
+	TFunctionRef<bool(int32, int32)> CanGrowPredicate )
+{
+	TArray<int32> LocalQueue;
+	QueueBuffer = (QueueBuffer == nullptr) ? &LocalQueue : QueueBuffer;
+	QueueBuffer->Reset();
+	QueueBuffer->Insert(InputROI, 0);
+
+	ResultROI.Reset();
+	ResultROI.Append(InputROI);
+
+	while (QueueBuffer->Num() > 0)
+	{
+		int32 CurEdgeID = QueueBuffer->Pop(false);
+		ResultROI.Add(CurEdgeID);
+
+		FIndex2i EdgeV = Mesh.GetEdgeV(CurEdgeID);
+		for (int32 k = 0; k < 2; ++k)
+		{
+			Mesh.EnumerateVertexEdges(EdgeV[k], [&](int32 NbrEdgeID)
+			{
+				if (NbrEdgeID != CurEdgeID && ResultROI.Contains(NbrEdgeID) == false && CanGrowPredicate(CurEdgeID, NbrEdgeID))
+				{
+					QueueBuffer->Add(NbrEdgeID);
+					ResultROI.Add(NbrEdgeID);
+				}
+			});
+		}
+
+	}
 }
\ No newline at end of file
diff --git a/Engine/Plugins/Runtime/GeometryProcessing/Source/DynamicMesh/Public/GroupTopology.h b/Engine/Plugins/Runtime/GeometryProcessing/Source/DynamicMesh/Public/GroupTopology.h
--- a/Engine/Plugins/Runtime/GeometryProcessing/Source/DynamicMesh/Public/GroupTopology.h
+++ b/Engine/Plugins/Runtime/GeometryProcessing/Source/DynamicMesh/Public/GroupTopology.h
@@ -319,6 +319,9 @@ public:
 	/** Add the groups connected to all the GroupEdgeIDs to the GroupsOut list. This is not the either-side pair, but the set of groups on the one-ring of each connected corner. */
 	void FindEdgeNbrGroups(const TArray<int>& GroupEdgeIDs, TArray<int>& GroupsOut) const;
 
+	/** Add the edges connected to the given GroupEdgeID to the EdgesOut list. This is somewhat expensive. */
+	void FindEdgeNbrEdges(int GroupEdgeID, TArray<int>& EdgesOut) const;
+
 	/** @return true if the FGroupEdge identified by the GroupEdgeID is on the mesh boundary */
 	bool IsBoundaryEdge(int32 GroupEdgeID) const;
 	/** @return true if group edge is a "simple" edge, meaning it only has one mesh edge */
@@ -338,6 +341,11 @@ public:
 	/** Add all the groups connected to the given Corners to the GroupsOut list */
 	void FindCornerNbrGroups(const TArray<int>& CornerIDs, TArray<int>& GroupsOut) const;
 
+	/** Add all the Edges connected to the given Corner to the EdgesOut list. This is somewhat expensive as there is no connection from Corners to Edges, must iterate over adjacent faces */
+	void FindCornerNbrEdges(int CornerID, TArray<int>& EdgesOut) const;
+	/** Add all the Corners connected to the given Corner to the EdgesOut list. This is somewhat expensive as there is no connection from Corners to Corners, must iterate over adjacent faces/edges */
+	void FindCornerNbrCorners(int CornerID, TArray<int>& CornersOut) const;
+
 	/** Add all the groups connected to the given Mesh Vertex to the GroupsOut list */
 	void FindVertexNbrGroups(int VertexID, TArray<int>& GroupsOut ) const;
 	/** Add all the groups connected to the given Mesh Vertices to the GroupsOut list */
diff --git a/Engine/Plugins/Runtime/GeometryProcessing/Source/DynamicMesh/Public/Selections/MeshConnectedComponents.h b/Engine/Plugins/Runtime/GeometryProcessing/Source/DynamicMesh/Public/Selections/MeshConnectedComponents.h
--- a/Engine/Plugins/Runtime/GeometryProcessing/Source/DynamicMesh/Public/Selections/MeshConnectedComponents.h
+++ b/Engine/Plugins/Runtime/GeometryProcessing/Source/DynamicMesh/Public/Selections/MeshConnectedComponents.h
@@ -239,6 +239,40 @@ public:
 	);
 
 
+	/**
+	 * Utility function to expand a vertex selection to all vertices considered "connected".
+	 * @param Mesh Mesh to calculate on
+	 * @param InputROI input set of vertices
+	 * @param ResultROI output set of vertices connected to InputROI
+	 * @param QueueBuffer optional buffer used as internal Queue. If passed as nullptr, a TArray will be locally allocated
+	 * @param CanGrowPredicate determines whether two connected mesh vertices should be considered connected while growing
+	 */
+	static void GrowToConnectedVertices(
+		const FDynamicMesh3& Mesh,
+		const TArray<int>& InputROI,
+		TSet<int>& ResultROI,
+		TArray<int32>* QueueBuffer = nullptr,
+		TFunctionRef<bool(int32, int32)> CanGrowPredicate = [](int32, int32) { return true; }
+	);
+
+
+	/**
+	 * Utility function to expand an edge selection to all edges considered "connected".
+	 * @param Mesh Mesh to calculate on
+	 * @param InputROI input set of edges
+	 * @param ResultROI output set of edges connected to InputROI
+	 * @param QueueBuffer optional buffer used as internal Queue. If passed as nullptr, a TArray will be locally allocated
+	 * @param CanGrowPredicate determines whether two connected mesh edges should be considered connected while growing
+	 */
+	static void GrowToConnectedEdges(
+		const FDynamicMesh3& Mesh,
+		const TArray<int>& InputROI,
+		TSet<int>& ResultROI,
+		TArray<int32>* QueueBuffer = nullptr,
+		TFunctionRef<bool(int32, int32)> CanGrowPredicate = [](int32, int32) { return true; }
+	);
+
+
 };
 
 
