GeometryProcessing: add MakeSelectAllSelection(), MakeSelectAllConnectedSelection(), MakeBoundaryConnectedSelection(), CombineSelectionInPlace() utility functions. Some of this is ported from geometry script selection code, which should be replaced with calls to these util functions (after testing). Minor necessarily additions to FGeometrySelection and FGeometrySelectionEditor.

diff --git a/Engine/Plugins/Runtime/GeometryProcessing/Source/DynamicMesh/Private/Selections/GeometrySelection.cpp b/Engine/Plugins/Runtime/GeometryProcessing/Source/DynamicMesh/Private/Selections/GeometrySelection.cpp
--- a/Engine/Plugins/Runtime/GeometryProcessing/Source/DynamicMesh/Private/Selections/GeometrySelection.cpp
+++ b/Engine/Plugins/Runtime/GeometryProcessing/Source/DynamicMesh/Private/Selections/GeometrySelection.cpp
@@ -13,6 +13,14 @@ void FGeometrySelectionEditor::Initialize(FGeometrySelection* TargetSelectionIn,
 	UpdateQueryConfig(QueryConfigIn, bEnableTopologyIDFilteringIn);
 }
 
+void FGeometrySelectionEditor::Initialize(FGeometrySelection* TargetSelectionIn, bool bEnableTopologyIDFilteringIn)
+{
+	check(TargetSelectionIn != nullptr);
+	TargetSelection = TargetSelectionIn;
+	FGeometrySelectionHitQueryConfig TmpQueryConfig{TargetSelection->TopologyType, TargetSelection->ElementType, false};
+	UpdateQueryConfig(TmpQueryConfig, bEnableTopologyIDFilteringIn);
+}
+
 
 void FGeometrySelectionEditor::UpdateQueryConfig(const FGeometrySelectionHitQueryConfig& NewConfig, bool bEnableTopologyIDFilteringIn)
 {
diff --git a/Engine/Plugins/Runtime/GeometryProcessing/Source/DynamicMesh/Private/Selections/GeometrySelectionUtil.cpp b/Engine/Plugins/Runtime/GeometryProcessing/Source/DynamicMesh/Private/Selections/GeometrySelectionUtil.cpp
--- a/Engine/Plugins/Runtime/GeometryProcessing/Source/DynamicMesh/Private/Selections/GeometrySelectionUtil.cpp
+++ b/Engine/Plugins/Runtime/GeometryProcessing/Source/DynamicMesh/Private/Selections/GeometrySelectionUtil.cpp
@@ -312,7 +312,7 @@ bool UE::Geometry::EnumerateTriangleSelectionVertices(
 	{
 		for (uint64 EncodedID : MeshSelection.Selection)
 		{
-			FMeshTriEdgeID TriEdgeID((uint32)EncodedID);
+			FMeshTriEdgeID TriEdgeID(FGeoSelectionID(EncodedID).GeometryID);
 			int32 EdgeID = Mesh.IsTriangle(TriEdgeID.TriangleID) ? Mesh.GetTriEdge(TriEdgeID.TriangleID, TriEdgeID.TriEdgeIndex) : IndexConstants::InvalidID;
 			if (Mesh.IsEdge(EdgeID))
 			{
@@ -856,6 +856,647 @@ bool UE::Geometry::ConvertPolygroupSelectionToTopologySelection(
 
 
 
+bool UE::Geometry::MakeSelectAllSelection(
+	const UE::Geometry::FDynamicMesh3& Mesh,
+	const FGroupTopology* GroupTopology,
+	TFunctionRef<bool(FGeoSelectionID)> SelectionIDPredicate,
+	FGeometrySelection& AllSelection)
+{
+	if (AllSelection.TopologyType == EGeometryTopologyType::Triangle)
+	{
+		if (AllSelection.ElementType == EGeometryElementType::Vertex)
+		{
+			for (int32 vid : Mesh.VertexIndicesItr())
+			{
+				FGeoSelectionID ID = FGeoSelectionID::MeshVertex(vid);
+				if ( SelectionIDPredicate(ID) )
+				{
+					AllSelection.Selection.Add(ID.Encoded());
+				}
+			}
+		}
+		else if (AllSelection.ElementType == EGeometryElementType::Edge)
+		{
+			for (int32 eid : Mesh.EdgeIndicesItr())
+			{
+				FGeoSelectionID ID = FGeoSelectionID::MeshEdge(Mesh.GetTriEdgeIDFromEdgeID(eid));
+				if ( SelectionIDPredicate(ID) )
+				{
+					AllSelection.Selection.Add(ID.Encoded());
+				}
+			}
+		}
+		else if (AllSelection.ElementType == EGeometryElementType::Face)
+		{
+			for (int32 tid : Mesh.TriangleIndicesItr())
+			{
+				FGeoSelectionID ID = FGeoSelectionID::MeshTriangle(tid);
+				if ( SelectionIDPredicate(ID) )
+				{
+					AllSelection.Selection.Add(ID.Encoded());
+				}
+			}
+		}
+		else
+		{
+			return false;
+		}
+		return true;
+	}
+	else if ( AllSelection.TopologyType == EGeometryTopologyType::Polygroup )
+	{
+		if (!ensure(GroupTopology != nullptr))
+		{
+			return false;
+		}
+
+		if (AllSelection.ElementType == EGeometryElementType::Vertex)
+		{
+			int32 NumCorners = GroupTopology->Corners.Num();
+			for ( int32 CornerID = 0; CornerID < NumCorners; ++CornerID)
+			{
+				const FGroupTopology::FCorner& Corner = GroupTopology->Corners[CornerID];
+				FGeoSelectionID ID = FGeoSelectionID(Corner.VertexID, CornerID);
+				if ( SelectionIDPredicate(ID) )
+				{
+					AllSelection.Selection.Add(ID.Encoded());
+				}
+			}
+		}
+		else if (AllSelection.ElementType == EGeometryElementType::Edge)
+		{
+			int32 NumEdges = GroupTopology->Edges.Num();
+			for ( int32 EdgeID = 0; EdgeID < NumEdges; ++EdgeID)
+			{
+				const FGroupTopology::FGroupEdge& GroupEdge = GroupTopology->Edges[EdgeID];
+				FMeshTriEdgeID MeshEdgeID = Mesh.GetTriEdgeIDFromEdgeID(GroupEdge.Span.Edges[0]);
+				FGeoSelectionID ID = FGeoSelectionID(MeshEdgeID.Encoded(), EdgeID);
+				if ( SelectionIDPredicate(ID) )
+				{
+					AllSelection.Selection.Add(ID.Encoded());
+				}
+			}
+		}
+		else if (AllSelection.ElementType == EGeometryElementType::Face)
+		{
+			int32 NumFaces = GroupTopology->Groups.Num();
+			for ( int32 FaceID = 0; FaceID < NumFaces; ++FaceID)
+			{
+				const FGroupTopology::FGroup& GroupFace = GroupTopology->Groups[FaceID];
+				FGeoSelectionID ID = FGeoSelectionID(GroupFace.Triangles[0], GroupFace.GroupID);
+				if ( SelectionIDPredicate(ID) )
+				{
+					AllSelection.Selection.Add(ID.Encoded());
+				}
+			}
+		}
+		else
+		{
+			return false;
+		}
+		return true;
+	}
+	return false;
+}
+
+
+
+bool UE::Geometry::MakeSelectAllConnectedSelection(
+	const UE::Geometry::FDynamicMesh3& Mesh,
+	const FGroupTopology* GroupTopology,
+	const FGeometrySelection& ReferenceSelection,
+	TFunctionRef<bool(FGeoSelectionID)> SelectionIDPredicate,
+	TFunctionRef<bool(FGeoSelectionID, FGeoSelectionID)> IsConnectedPredicate,
+	FGeometrySelection& AllConnectedSelection)
+{
+	if ( ! ensure(ReferenceSelection.IsSameType(AllConnectedSelection)) ) return false;
+
+	if (AllConnectedSelection.TopologyType == EGeometryTopologyType::Triangle)
+	{
+		TArray<int32> CurIndices;
+		CurIndices.Reserve(ReferenceSelection.Num());
+
+		if (AllConnectedSelection.ElementType == EGeometryElementType::Vertex)
+		{
+			for (uint64 ElementID : ReferenceSelection.Selection)
+			{
+				CurIndices.Add( FGeoSelectionID(ElementID).GeometryID );
+			}
+			TSet<int32> ConnectedVertices;
+			FMeshConnectedComponents::GrowToConnectedVertices(Mesh, CurIndices, ConnectedVertices, nullptr,
+				[&](int32 FromVertID, int32 ToVertID) {
+					return SelectionIDPredicate(FGeoSelectionID::MeshVertex(ToVertID)) && 
+							IsConnectedPredicate( FGeoSelectionID::MeshVertex(FromVertID), FGeoSelectionID::MeshVertex(ToVertID) );
+				});
+			for (int32 vid : ConnectedVertices)
+			{
+				AllConnectedSelection.Selection.Add( FGeoSelectionID::MeshVertex(vid).Encoded() );
+			}
+		}
+		else if (AllConnectedSelection.ElementType == EGeometryElementType::Edge)
+		{
+			for (uint64 ElementID : ReferenceSelection.Selection)
+			{
+				FMeshTriEdgeID TriEdgeID( FGeoSelectionID(ElementID).GeometryID );
+				CurIndices.Add( Mesh.GetTriEdge(TriEdgeID.TriangleID, TriEdgeID.TriEdgeIndex) );
+			}
+			TSet<int32> ConnectedEdges;
+			FMeshConnectedComponents::GrowToConnectedEdges(Mesh, CurIndices, ConnectedEdges, nullptr,
+				[&](int32 FromEdgeID, int32 ToEdgeID) {
+					FMeshTriEdgeID ToTriEdgeID = Mesh.GetTriEdgeIDFromEdgeID(ToEdgeID), FromTriEdgeID = Mesh.GetTriEdgeIDFromEdgeID(FromEdgeID);
+					return SelectionIDPredicate(FGeoSelectionID::MeshEdge(ToTriEdgeID)) && 
+							IsConnectedPredicate( FGeoSelectionID::MeshEdge(FromTriEdgeID), FGeoSelectionID::MeshEdge(ToTriEdgeID) );
+				});
+			for (int32 EdgeID : ConnectedEdges)
+			{
+				AllConnectedSelection.Selection.Add( FGeoSelectionID::MeshEdge(Mesh.GetTriEdgeIDFromEdgeID(EdgeID)).Encoded() );
+			}
+		}
+		else if (AllConnectedSelection.ElementType == EGeometryElementType::Face)
+		{
+			for (uint64 ElementID : ReferenceSelection.Selection)
+			{
+				CurIndices.Add(FGeoSelectionID(ElementID).GeometryID);
+			}
+			TSet<int32> ConnectedTriangles;
+			FMeshConnectedComponents::GrowToConnectedTriangles(&Mesh, CurIndices, ConnectedTriangles, nullptr,
+				[&](int32 FromTriID, int32 ToTriID) {
+					return SelectionIDPredicate(FGeoSelectionID::MeshTriangle(ToTriID)) && 
+							IsConnectedPredicate( FGeoSelectionID::MeshTriangle(FromTriID), FGeoSelectionID::MeshTriangle(ToTriID) );
+				});
+			for (int32 tid : ConnectedTriangles)
+			{
+				AllConnectedSelection.Selection.Add( FGeoSelectionID::MeshTriangle(tid).Encoded() );
+			}
+		}
+		else
+		{
+			return false;
+		}
+		return true;
+	}
+	else if ( AllConnectedSelection.TopologyType == EGeometryTopologyType::Polygroup )
+	{
+		if (!ensure(GroupTopology != nullptr))
+		{
+			return false;
+		}
+		FGeometrySelectionEditor Editor;
+		Editor.Initialize(&AllConnectedSelection, true);
+		AllConnectedSelection = ReferenceSelection;
+		TArray<uint64> Queue;
+		for (uint64 ID : ReferenceSelection.Selection)
+		{
+			Queue.Add( ID );
+		}
+
+		if (AllConnectedSelection.ElementType == EGeometryElementType::Vertex)
+		{
+			TArray<int32> NbrCornerIDs;
+			while (Queue.Num() > 0)
+			{
+				FGeoSelectionID CurCornerSelectionID = FGeoSelectionID(Queue.Pop(false));
+				const FGroupTopology::FCorner& Corner = GroupTopology->Corners[CurCornerSelectionID.TopologyID];
+				NbrCornerIDs.Reset();
+				GroupTopology->FindCornerNbrCorners(CurCornerSelectionID.TopologyID, NbrCornerIDs);
+				for ( int32 NbrCornerID : NbrCornerIDs )
+				{
+					FGeoSelectionID NbrCornerSelectionID( GroupTopology->Corners[NbrCornerID].VertexID, NbrCornerID );
+					if ( Editor.IsSelected(NbrCornerSelectionID.Encoded()) == false 
+						&& SelectionIDPredicate(NbrCornerSelectionID) 
+						&& IsConnectedPredicate(CurCornerSelectionID, NbrCornerSelectionID) )
+					{
+						Queue.Add(NbrCornerSelectionID.Encoded());
+						Editor.Select(NbrCornerSelectionID.Encoded());
+					}
+				}
+			}
+		}
+		else if (AllConnectedSelection.ElementType == EGeometryElementType::Edge)
+		{
+			TArray<int32> NbrEdgeIDs;
+			while (Queue.Num() > 0)
+			{
+				FGeoSelectionID CurEdgeSelectionID = FGeoSelectionID(Queue.Pop(false));
+				const FGroupTopology::FGroupEdge& Edge = GroupTopology->Edges[CurEdgeSelectionID.TopologyID];
+				NbrEdgeIDs.Reset();
+				GroupTopology->FindEdgeNbrEdges(CurEdgeSelectionID.TopologyID, NbrEdgeIDs);
+				for ( int32 NbrEdgeID : NbrEdgeIDs )
+				{
+					FMeshTriEdgeID MeshEdgeID = Mesh.GetTriEdgeIDFromEdgeID(GroupTopology->Edges[NbrEdgeID].Span.Edges[0]);
+					FGeoSelectionID NbrEdgeSelectionID( MeshEdgeID.Encoded(), NbrEdgeID);
+					if ( Editor.IsSelected(NbrEdgeSelectionID.Encoded()) == false 
+						&& SelectionIDPredicate(NbrEdgeSelectionID) 
+						&& IsConnectedPredicate(CurEdgeSelectionID, NbrEdgeSelectionID) )
+					{
+						Queue.Add(NbrEdgeSelectionID.Encoded());
+						Editor.Select(NbrEdgeSelectionID.Encoded());
+					}
+				}
+			}
+		}
+		else if (AllConnectedSelection.ElementType == EGeometryElementType::Face)
+		{
+			TArray<int32> NbrGroupIDs;
+			while (Queue.Num() > 0)
+			{
+				FGeoSelectionID CurGroupSelectionID = FGeoSelectionID(Queue.Pop(false));
+				NbrGroupIDs.Reset();
+				for ( int32 NbrGroupID : GroupTopology->GetGroupNbrGroups(CurGroupSelectionID.TopologyID) )
+				{
+					const FGroupTopology::FGroup* NbrGroup = GroupTopology->FindGroupByID(NbrGroupID);
+					FGeoSelectionID NbrGroupSelectionID( NbrGroup->Triangles[0], NbrGroupID);
+					if ( Editor.IsSelected(NbrGroupSelectionID.Encoded()) == false 
+						&& SelectionIDPredicate(NbrGroupSelectionID) 
+						&& IsConnectedPredicate(CurGroupSelectionID, NbrGroupSelectionID) )
+					{
+						Queue.Add(NbrGroupSelectionID.Encoded());
+						Editor.Select(NbrGroupSelectionID.Encoded());
+					}
+				}
+			}
+
+		}
+		else
+		{
+			return false;
+		}
+		return true;
+	}
+	return false;
+
+}
+
+
+
+
+
+
+bool UE::Geometry::MakeBoundaryConnectedSelection(
+	const UE::Geometry::FDynamicMesh3& Mesh,
+	const FGroupTopology* GroupTopology,
+	const FGeometrySelection& ReferenceSelection,
+	TFunctionRef<bool(FGeoSelectionID)> SelectionIDPredicate,
+	FGeometrySelection& BoundaryConnectedSelection)
+{
+	if (BoundaryConnectedSelection.TopologyType == EGeometryTopologyType::Triangle)
+	{
+		// convert to vertex selection
+		TSet<int32> CurVertices;
+		EnumerateTriangleSelectionVertices(ReferenceSelection, Mesh, FTransform::Identity, [&](uint64 ID, const FVector3d&)
+		{
+			CurVertices.Add( FGeoSelectionID(ID).GeometryID );
+		});
+		// find 'border' vertices that have some adjacent vertices not in selection
+		TSet<int32> BorderVertices;
+		for (int32 VertexID : CurVertices)
+		{
+			bool bIsBoundary = Mesh.IsBoundaryVertex(VertexID);		// a boundary vertex is always on the selection boundary
+			Mesh.EnumerateVertexVertices(VertexID, [&](int32 NbrVertexID)
+			{
+				if (CurVertices.Contains(NbrVertexID) == false)
+				{
+					bIsBoundary = true;
+				}
+			});
+			if (bIsBoundary)
+			{
+				BorderVertices.Add(VertexID);
+			}
+		}
+
+		if (BoundaryConnectedSelection.ElementType == EGeometryElementType::Vertex)
+		{
+			TSet<int32> AdjacentVertices = BorderVertices;
+			for (int32 VertexID : BorderVertices)
+			{
+				Mesh.EnumerateVertexVertices(VertexID, [&](int32 NbrVertexID)
+				{
+					// filter out interior vertices, maybe should be a parameter
+					if ( CurVertices.Contains(NbrVertexID) == false )
+					{
+						AdjacentVertices.Add(NbrVertexID);
+					}
+				});
+			}
+			for (int32 VertexID : AdjacentVertices)
+			{
+				if (SelectionIDPredicate(FGeoSelectionID::MeshVertex(VertexID)))
+				{
+					BoundaryConnectedSelection.Selection.Add(FGeoSelectionID::MeshVertex(VertexID).Encoded());
+				}
+			}
+		}
+		else if (BoundaryConnectedSelection.ElementType == EGeometryElementType::Edge)
+		{
+			TSet<int32> AdjacentEdges;
+			for (int32 VertexID : BorderVertices)
+			{
+				for ( int32 EdgeID : Mesh.VtxEdgesItr(VertexID) )
+				{
+					AdjacentEdges.Add(EdgeID);
+				}
+			}
+			for (int32 EdgeID : AdjacentEdges)
+			{
+				FGeoSelectionID MeshEdgeID = FGeoSelectionID::MeshEdge(Mesh.GetTriEdgeIDFromEdgeID(EdgeID));
+				if (SelectionIDPredicate(MeshEdgeID))
+				{
+					BoundaryConnectedSelection.Selection.Add(MeshEdgeID.Encoded());
+				}
+			}
+		}
+		else if (BoundaryConnectedSelection.ElementType == EGeometryElementType::Face)
+		{
+			TSet<int32> AdjacentTriangles;
+			for (int32 VertexID : BorderVertices)
+			{
+				Mesh.EnumerateVertexTriangles(VertexID, [&](int32 NbrTriangleID)
+				{
+					AdjacentTriangles.Add(NbrTriangleID);
+				});
+			}
+			for (int32 TriangleID : AdjacentTriangles)
+			{
+				if (SelectionIDPredicate(FGeoSelectionID::MeshTriangle(TriangleID)))
+				{
+					BoundaryConnectedSelection.Selection.Add(FGeoSelectionID::MeshTriangle(TriangleID).Encoded());
+				}
+			}
+		}
+		else
+		{
+			return false;
+		}
+		return true;
+	}
+	else if ( BoundaryConnectedSelection.TopologyType == EGeometryTopologyType::Polygroup )
+	{
+		if (!ensure(GroupTopology != nullptr))
+		{
+			return false;
+		}
+
+		// convert to Corner selection
+		TSet<int32> CurCornerIDs;
+		if (ReferenceSelection.ElementType == EGeometryElementType::Vertex)
+		{
+			for (uint64 ID : ReferenceSelection.Selection)
+			{
+				CurCornerIDs.Add( FGeoSelectionID(ID).TopologyID );		// TODO: can we rely on TopologyID being stable here, or do we need to look up from VertexID?
+			}
+		}
+		else if (ReferenceSelection.ElementType == EGeometryElementType::Edge)
+		{
+			for (uint64 ID : ReferenceSelection.Selection)
+			{
+				int32 GroupEdgeID = FGeoSelectionID(ID).TopologyID;
+				const FGroupTopology::FGroupEdge& Edge = GroupTopology->Edges[GroupEdgeID];
+				if ( Edge.EndpointCorners.A != IndexConstants::InvalidID )
+				{
+					CurCornerIDs.Add(Edge.EndpointCorners.A);
+				}
+				if ( Edge.EndpointCorners.B != IndexConstants::InvalidID )
+				{
+					CurCornerIDs.Add(Edge.EndpointCorners.B);
+				}
+			}
+		}
+		else if (ReferenceSelection.ElementType == EGeometryElementType::Face)
+		{
+			for (uint64 ID : ReferenceSelection.Selection)
+			{
+				int32 GroupID = FGeoSelectionID(ID).TopologyID;
+				const FGroupTopology::FGroup* Group = GroupTopology->FindGroupByID(GroupID);
+				GroupTopology->ForGroupEdges(GroupID, [&](const FGroupTopology::FGroupEdge& Edge, int)
+				{
+					if ( Edge.EndpointCorners.A != IndexConstants::InvalidID )
+					{
+						CurCornerIDs.Add(Edge.EndpointCorners.A);
+					}
+					if ( Edge.EndpointCorners.B != IndexConstants::InvalidID )
+					{
+						CurCornerIDs.Add(Edge.EndpointCorners.B);
+					}
+				});
+			}
+		}
+		else
+		{
+			return false;
+		}
+
+
+		// now find 'border' corners that have some adjacent corners not in selection
+		TSet<int32> BorderCorners;
+		TArray<int32> NbrArray;
+		for (int32 CornerID : CurCornerIDs)
+		{
+			bool bIsBoundary = Mesh.IsBoundaryVertex(GroupTopology->GetCornerVertexID(CornerID));		// a boundary vertex is always on the selection boundary
+			NbrArray.Reset();
+			GroupTopology->FindCornerNbrCorners(CornerID, NbrArray);
+			for ( int32 NbrCornerID : NbrArray )
+			{
+				if (CurCornerIDs.Contains(NbrCornerID) == false)
+				{
+					bIsBoundary = true;
+				}
+			}
+			if (bIsBoundary)
+			{
+				BorderCorners.Add(CornerID);
+			}
+		}
+
+		// now that we have boundary corners, iterate over them and select connected elements
+		if (BoundaryConnectedSelection.ElementType == EGeometryElementType::Vertex)
+		{
+			TSet<int32> AdjacentCorners = BorderCorners;
+			for (int32 CornerID : BorderCorners)
+			{
+				NbrArray.Reset();
+				GroupTopology->FindCornerNbrCorners(CornerID, NbrArray);
+				for (int32 NbrCornerID : NbrArray)
+				{
+					// filter out interior corners, maybe should be a parameter
+					if ( CurCornerIDs.Contains(NbrCornerID) == false )
+					{
+						AdjacentCorners.Add(NbrCornerID);
+					}
+				}
+			}
+			for (int32 CornerID : AdjacentCorners)
+			{
+				FGeoSelectionID SelectionID( GroupTopology->GetCornerVertexID(CornerID), CornerID);
+				if (SelectionIDPredicate(SelectionID) )
+				{
+					BoundaryConnectedSelection.Selection.Add( SelectionID.Encoded() );
+				}
+			}
+		}
+		else if (BoundaryConnectedSelection.ElementType == EGeometryElementType::Edge)
+		{
+			TSet<int32> AdjacentEdges;
+			for (int32 CornerID : BorderCorners)
+			{
+				NbrArray.Reset();
+				GroupTopology->FindCornerNbrEdges(CornerID, NbrArray);
+				for (int32 NbrEdgeID : NbrArray)
+				{
+					AdjacentEdges.Add(NbrEdgeID);
+				}
+			}
+			for (int32 EdgeID : AdjacentEdges)
+			{
+				FMeshTriEdgeID MeshEdgeID = Mesh.GetTriEdgeIDFromEdgeID(GroupTopology->GetGroupEdgeEdges(EdgeID)[0]);
+				FGeoSelectionID SelectionID( MeshEdgeID.Encoded(), EdgeID);
+				if (SelectionIDPredicate(SelectionID) )
+				{
+					BoundaryConnectedSelection.Selection.Add( SelectionID.Encoded() );
+				}
+			}
+		}
+		else   // already verified we are only vertex/edge/face above
+		{
+			TSet<int32> AdjacentGroups;
+			for (int32 CornerID : BorderCorners)
+			{
+				NbrArray.Reset();
+				GroupTopology->FindCornerNbrGroups(CornerID, NbrArray);
+				for (int32 NbrGroupID : NbrArray)
+				{
+					AdjacentGroups.Add(NbrGroupID);
+				}
+			}
+			for (int32 GroupID : AdjacentGroups)
+			{
+				FGeoSelectionID SelectionID( GroupTopology->GetGroupTriangles(GroupID)[0], GroupID);
+				if (SelectionIDPredicate(SelectionID) )
+				{
+					BoundaryConnectedSelection.Selection.Add( SelectionID.Encoded() );
+				}
+			}
+		}
+
+		return true;
+	}
+	return false;
+}
+
+
+
+
+
+bool UE::Geometry::CombineSelectionInPlace(
+	FGeometrySelection& SelectionA,
+	const FGeometrySelection& SelectionB,
+	EGeometrySelectionCombineModes CombineMode)
+{
+	if (SelectionA.IsSameType(SelectionB) == false)
+	{
+		return false;
+	}
+
+	if (SelectionA.TopologyType == EGeometryTopologyType::Triangle)
+	{
+		if (CombineMode == EGeometrySelectionCombineModes::Add)
+		{
+			for (uint64 ItemB : SelectionB.Selection)
+			{
+				SelectionA.Selection.Add(ItemB);
+			}
+		}
+		else if (CombineMode == EGeometrySelectionCombineModes::Subtract)
+		{
+			if (SelectionB.IsEmpty() == false)
+			{
+				for (uint64 ItemB : SelectionB.Selection)
+				{
+					SelectionA.Selection.Remove(ItemB);
+				}
+				SelectionA.Selection.Compact();
+			}
+		}
+		else if (CombineMode == EGeometrySelectionCombineModes::Intersection)
+		{
+			TArray<uint64, TInlineAllocator<32>> ToRemove;
+			for (uint64 ItemA : SelectionA.Selection)
+			{
+				if (!SelectionB.Selection.Contains(ItemA))
+				{
+					ToRemove.Add(ItemA);
+				}
+			}
+			if (ToRemove.Num() > 0)
+			{
+				for (uint64 ItemA : ToRemove)
+				{
+					SelectionA.Selection.Remove(ItemA);
+				}
+				SelectionA.Selection.Compact();
+			}
+		}
+
+		return true;
+	}
+	else if (SelectionA.TopologyType == EGeometryTopologyType::Polygroup)
+	{
+		// for Polygroup selections, we cannot rely on TSet operations because we have set an arbitrary Triangle ID 
+		// as the 'geometry' key.
+		if (CombineMode == EGeometrySelectionCombineModes::Add)
+		{
+			for (uint64 ItemB : SelectionB.Selection)
+			{
+				uint64 FoundItemA;
+				if ( UE::Geometry::FindInSelectionByTopologyID(SelectionA, FGeoSelectionID(ItemB).TopologyID, FoundItemA) == false)
+				{
+					SelectionA.Selection.Add(ItemB);
+				}
+			}
+		}
+		else if (CombineMode == EGeometrySelectionCombineModes::Subtract)
+		{
+			if (SelectionB.IsEmpty() == false)
+			{
+				for (uint64 ItemB : SelectionB.Selection)
+				{
+					uint64 FoundItemA;
+					if (UE::Geometry::FindInSelectionByTopologyID(SelectionA, FGeoSelectionID(ItemB).TopologyID, FoundItemA))
+					{
+						SelectionA.Selection.Remove(FoundItemA);
+					}
+				}
+				SelectionA.Selection.Compact();
+			}
+		}
+		else if (CombineMode == EGeometrySelectionCombineModes::Intersection)
+		{
+			TArray<uint64, TInlineAllocator<32>> ToRemove;
+			for (uint64 ItemA : SelectionA.Selection)
+			{
+				uint64 FoundItemB;
+				if (UE::Geometry::FindInSelectionByTopologyID(SelectionA, FGeoSelectionID(ItemA).TopologyID, FoundItemB) == false)
+				{
+					ToRemove.Add(ItemA);
+				}
+			}
+			if (ToRemove.Num() > 0)
+			{
+				for (uint64 ItemA : ToRemove)
+				{
+					SelectionA.Selection.Remove(ItemA);
+				}
+				SelectionA.Selection.Compact();
+			}
+		}
+
+		return true;
+	}
+
+	return false;
+}
+
+
 
 bool UE::Geometry::GetTriangleSelectionFrame(
 	const FGeometrySelection& MeshSelection,
diff --git a/Engine/Plugins/Runtime/GeometryProcessing/Source/DynamicMesh/Public/Selections/GeometrySelection.h b/Engine/Plugins/Runtime/GeometryProcessing/Source/DynamicMesh/Public/Selections/GeometrySelection.h
--- a/Engine/Plugins/Runtime/GeometryProcessing/Source/DynamicMesh/Public/Selections/GeometrySelection.h
+++ b/Engine/Plugins/Runtime/GeometryProcessing/Source/DynamicMesh/Public/Selections/GeometrySelection.h
@@ -188,6 +188,14 @@ struct DYNAMICMESH_API FGeometrySelection
 		ElementType = FromSelection.ElementType;
 		TopologyType = FromSelection.TopologyType;
 	}
+
+	/**
+	 * @return true if the two selections have the same type
+	 */
+	bool IsSameType(const FGeometrySelection& OtherSelection) const
+	{
+		return ElementType == OtherSelection.ElementType && TopologyType == OtherSelection.TopologyType;
+	}
 };
 
 
@@ -324,6 +332,14 @@ public:
 		const FGeometrySelectionHitQueryConfig& QueryConfigIn,
 		bool bEnableTopologyIDFiltering);
 
+	/**
+	 * Initialize the Editor with the given Selection. The
+	 * TargetSelectionIn must live longer than the FGeometrySelectionEditor
+	 */
+	void Initialize(
+		FGeometrySelection* TargetSelectionIn,
+		bool bEnableTopologyIDFiltering);
+
 	/** @return the Element Type of the Target Selection */
 	EGeometryElementType GetElementType() const { return TargetSelection->ElementType; }
 	/** @return the Topology Type of the Target Selection */
@@ -353,6 +369,17 @@ public:
 	/** Access the Selection object this Editor is modifying */
 	const FGeometrySelection& GetSelection() const { return *TargetSelection; }
 
+	/** Add the items to the Target Selection */
+	bool Select(uint64 ID)
+	{
+		if (IsSelected(ID) == false)
+		{
+			TargetSelection->Selection.Add(ID);
+			return true;
+		}
+		return false;
+	}
+
 	/** Add the items in the List to the Target Selection and return change information in DeltaOut */
 	template<typename ListType>
 	bool Select(const ListType& List, FGeometrySelectionDelta& DeltaOut)
diff --git a/Engine/Plugins/Runtime/GeometryProcessing/Source/DynamicMesh/Public/Selections/GeometrySelectionUtil.h b/Engine/Plugins/Runtime/GeometryProcessing/Source/DynamicMesh/Public/Selections/GeometrySelectionUtil.h
--- a/Engine/Plugins/Runtime/GeometryProcessing/Source/DynamicMesh/Public/Selections/GeometrySelectionUtil.h
+++ b/Engine/Plugins/Runtime/GeometryProcessing/Source/DynamicMesh/Public/Selections/GeometrySelectionUtil.h
@@ -192,6 +192,73 @@ DYNAMICMESH_API bool ConvertPolygroupSelectionToTopologySelection(
 );
 
 
+/**
+ * Select all elements of the provided Mesh and GroupTopology that pass the provided SelectionIDPredicate, 
+ * and store in the output AllSelection. The type of elements selected is defined by the existing configured
+ * type of the AllSelection parameter. 
+ * @param GroupTopology precomputed group topology for Mesh, can be passed as null for EGeometryTopologyType::Triangle selections
+ * @return true if AllSelection had a known geometry/topology type pair and was populated
+ */
+DYNAMICMESH_API bool MakeSelectAllSelection(
+	const UE::Geometry::FDynamicMesh3& Mesh,
+	const FGroupTopology* GroupTopology,
+	TFunctionRef<bool(FGeoSelectionID)> SelectionIDPredicate,
+	FGeometrySelection& AllSelection);
+
+/**
+ * Expand the input ReferenceSelection to include all "connected" elements and return in AllConnectedSelection.
+ * The type of selected element is defined by ReferenceSelection.
+ * @param GroupTopology precomputed group topology for Mesh, can be passed as null for EGeometryTopologyType::Triangle selections
+ * @param SelectionIDPredicate only elements that pass this filter will be expanded "to"  (but elements of ReferenceSelection that fail the filter will still be included in output)
+ * @param IsConnectedPredicate this function determines if "A" should be considered connected to "B", ie can "expand" along that connection
+ * @return true if ReferenceSelection had a known geometry/topology type pair and AllConnectedSelection was populated
+ */
+DYNAMICMESH_API bool MakeSelectAllConnectedSelection(
+	const UE::Geometry::FDynamicMesh3& Mesh,
+	const FGroupTopology* GroupTopology,
+	const FGeometrySelection& ReferenceSelection,
+	TFunctionRef<bool(FGeoSelectionID)> SelectionIDPredicate,
+	TFunctionRef<bool(FGeoSelectionID A, FGeoSelectionID B)> IsConnectedPredicate,
+	FGeometrySelection& AllConnectedSelection);
+
+/**
+ * Create a selection of the elements adjacent to the "Border" of the given ReferenceSelection and return in BoundaryConnectedSelection.
+ * The type of selected element is defined by ReferenceSelection.
+ * Currently "adjacency" is defined as "included in the one-ring of the boundary vertices of the ReferenceSelection", ie first the 
+ * vertices on boundary edges are found, and then their one-rings are enumerated. Note that this will include "inside" and "outside" adjacent elements,
+ * and for vertices, the boundary vertices will still also be included. The main purpose of this function is to implement expand/contract selection
+ * operations, which would typically involve first finding the boundary-connected set and then using CombineSelectionInPlace to modify the original selection.
+ * @param GroupTopology precomputed group topology for Mesh, can be passed as null for EGeometryTopologyType::Triangle selections
+ * @param SelectionIDPredicate only elements that pass this filter will be expanded "to"  (but elements of ReferenceSelection that fail the filter will still be included in output)
+ * @return true if ReferenceSelection had a known geometry/topology type pair and BoundaryConnectedSelection was populated
+ */
+DYNAMICMESH_API bool MakeBoundaryConnectedSelection(
+	const UE::Geometry::FDynamicMesh3& Mesh,
+	const FGroupTopology* GroupTopology,
+	const FGeometrySelection& ReferenceSelection,
+	TFunctionRef<bool(FGeoSelectionID)> SelectionIDPredicate,
+	FGeometrySelection& BoundaryConnectedSelection);
+
+
+enum class EGeometrySelectionCombineModes : uint8
+{
+	Add,
+	Subtract,
+	Intersection
+};
+
+
+/**
+ * Combine the elements of SelectionA and SelectionB using the provided CombineMode, and store the result in SelectionA.
+ * @return true if the selectins were compatible (ie both the same type) and of supported geometry/topology type.
+ */
+DYNAMICMESH_API bool CombineSelectionInPlace(
+	FGeometrySelection& SelectionA,
+	const FGeometrySelection& SelectionB, 
+	EGeometrySelectionCombineModes CombineMode );
+
+
+
 /**
  * Compute a 3D Frame suitable for use as a 3D transform gizmo position/orientation
  * for the given MeshSelection
diff --git a/Engine/Source/Runtime/GeometryCore/Public/DynamicMesh/InfoTypes.h b/Engine/Source/Runtime/GeometryCore/Public/DynamicMesh/InfoTypes.h
--- a/Engine/Source/Runtime/GeometryCore/Public/DynamicMesh/InfoTypes.h
+++ b/Engine/Source/Runtime/GeometryCore/Public/DynamicMesh/InfoTypes.h
@@ -79,7 +79,7 @@ struct FMeshTriEdgeID
 	/**
 	 * Decode an encoded FMeshTriEdgeID from a packed uint32 created by the Encoded() function
 	 */
-	FMeshTriEdgeID(uint32 EncodedEdgeKey)
+	explicit FMeshTriEdgeID(uint32 EncodedEdgeKey)
 	{
 		TriangleID = EncodedEdgeKey & 0x8FFFFFFF;
 		TriEdgeIndex = (EncodedEdgeKey & 0xC0000000) >> 30;
