GeometryProcessing: fix infinite loop that can happen if mesh a connected component has no border

diff --git a/Engine/Plugins/Runtime/GeometryProcessing/Source/DynamicMesh/Private/FaceGroupUtil.cpp b/Engine/Plugins/Runtime/GeometryProcessing/Source/DynamicMesh/Private/FaceGroupUtil.cpp
--- a/Engine/Plugins/Runtime/GeometryProcessing/Source/DynamicMesh/Private/FaceGroupUtil.cpp
+++ b/Engine/Plugins/Runtime/GeometryProcessing/Source/DynamicMesh/Private/FaceGroupUtil.cpp
@@ -268,6 +268,11 @@ void FGroupVisualizationCache::UpdateGroupInfo_ConnectedComponents(
 						BorderTris.Add(tid);
 					}
 				}
+				if (BorderTris.IsEmpty())
+				{
+					CenterTriID = *TrisInGroup.begin();
+					break;
+				}
 				for (int32 tid : BorderTris)
 				{
 					TrisInGroup.Remove(tid);
