ModelingTools: Fix typo in QuadGridPatch that causes a checkSlow to always be hit.

diff --git a/Engine/Plugins/Runtime/GeometryProcessing/Source/DynamicMesh/Private/Selections/QuadGridPatch.cpp b/Engine/Plugins/Runtime/GeometryProcessing/Source/DynamicMesh/Private/Selections/QuadGridPatch.cpp
--- a/Engine/Plugins/Runtime/GeometryProcessing/Source/DynamicMesh/Private/Selections/QuadGridPatch.cpp
+++ b/Engine/Plugins/Runtime/GeometryProcessing/Source/DynamicMesh/Private/Selections/QuadGridPatch.cpp
@@ -74,7 +74,7 @@ void FQuadGridPatch::InitializeFromQuadStrip(const FDynamicMesh3& Mesh, const TA
 		{
 			VertexSpans[1].Add(Other0);
 		}
-		checkSlow(FirstVertexSpan.Last() == Other0);
+		checkSlow(VertexSpans[1].Last() == Other0);
 		VertexSpans[1].Add(Other1);
 	}
 }
