Add a general-use version of the 'simplify along edges' logic from FMeshBoolean, and use it to simplify plane cuts and edge insertions
 + Update the plane cut logic to allow it to use the simplify function

diff --git a/Engine/Plugins/Runtime/GeometryProcessing/Source/DynamicMesh/Private/Operations/GroupEdgeInserter.cpp b/Engine/Plugins/Runtime/GeometryProcessing/Source/DynamicMesh/Private/Operations/GroupEdgeInserter.cpp
--- a/Engine/Plugins/Runtime/GeometryProcessing/Source/DynamicMesh/Private/Operations/GroupEdgeInserter.cpp
+++ b/Engine/Plugins/Runtime/GeometryProcessing/Source/DynamicMesh/Private/Operations/GroupEdgeInserter.cpp
@@ -14,6 +14,7 @@
 #include "Operations/MeshPlaneCut.h"
 #include "Operations/EmbedSurfacePath.h"
 #include "Operations/SimpleHoleFiller.h"
+#include "Operations/LocalPlanarSimplify.h"
 #include "Selections/MeshConnectedComponents.h"
 #include "Util/ProgressCancel.h"
 #include "Util/IndexUtil.h"
@@ -70,12 +71,12 @@ bool ConnectMultipleUsingPlaneCut(FDynamicMesh3& Mesh,
 	const FGroupTopology::FGroupBoundary& GroupBoundary,
 	const TArray<FGroupEdgeInserter::FGroupEdgeSplitPoint>& StartPoints,
 	const TArray <FGroupEdgeInserter::FGroupEdgeSplitPoint>& EndPoints,
-	double VertexTolerance, int32& NumGroupsCreated, 
+	double VertexTolerance, bool bSimplifyAlongPath, int32& NumGroupsCreated,
 	FGroupEdgeInserter::FOptionalOutputParams& OptionalOut, FProgressCancel* Progress);
 bool EmbedPlaneCutPath(FDynamicMesh3& Mesh, const FGroupTopology& Topology, int32 GroupID,
 	const FGroupEdgeInserter::FGroupEdgeSplitPoint& StartPoint,
 	const FGroupEdgeInserter::FGroupEdgeSplitPoint& EndPoint,
-	double VertexTolerance, TSet<int32>& PathEidsOut, 
+	double VertexTolerance, bool bSimplifyAlongPath, TSet<int32>& PathEidsOut, 
 	TSet<int32>* ChangedTrisOut, FProgressCancel* Progress);
 bool CreateNewGroups(FDynamicMesh3& Mesh, TSet<int32>& PathEids, int32 OriginalGroupID, int32& NumGroupsCreated, 
 	FGroupEdgeInserter::FOptionalOutputParams& OptionalOut, FProgressCancel* Progress);
@@ -556,7 +557,7 @@ bool ConnectEndpoints(
 	else if (Params.Mode == FGroupEdgeInserter::EInsertionMode::PlaneCut)
 	{
 		return ConnectMultipleUsingPlaneCut(*Params.Mesh, *Params.Topology, GroupID,
-			GroupBoundary, StartPoints, EndPoints, Params.VertexTolerance,
+			GroupBoundary, StartPoints, EndPoints, Params.VertexTolerance, Params.bSimplifyAlongPath,
 			NumGroupsCreated, OptionalOut, Progress);
 	}
 	else
@@ -921,7 +922,7 @@ bool ConnectMultipleUsingPlaneCut(FDynamicMesh3& Mesh,
 	const FGroupTopology::FGroupBoundary& GroupBoundary,
 	const TArray<FGroupEdgeInserter::FGroupEdgeSplitPoint>& StartPoints,
 	const TArray <FGroupEdgeInserter::FGroupEdgeSplitPoint>& EndPoints,
-	double VertexTolerance, int32& NumGroupsCreated, 
+	double VertexTolerance, bool bSimplifyAlongPath, int32& NumGroupsCreated,
 	FGroupEdgeInserter::FOptionalOutputParams& OptionalOut, 
 	FProgressCancel* Progress)
 {
@@ -936,7 +937,7 @@ bool ConnectMultipleUsingPlaneCut(FDynamicMesh3& Mesh,
 	for (int32 i = 0; i < NumEdgesToInsert; ++i)
 	{
 		bool bSuccess = EmbedPlaneCutPath(Mesh, Topology, GroupID, StartPoints[i], EndPoints[i],
-			VertexTolerance, PathsEids, OptionalOut.ChangedTidsOut, Progress);
+			VertexTolerance, bSimplifyAlongPath, PathsEids, OptionalOut.ChangedTidsOut, Progress);
 
 		if (!bSuccess || (Progress && Progress->Cancelled()))
 		{
@@ -960,7 +961,7 @@ bool ConnectMultipleUsingPlaneCut(FDynamicMesh3& Mesh,
 bool EmbedPlaneCutPath(FDynamicMesh3& Mesh, const FGroupTopology& Topology,
 	int32 GroupID, const FGroupEdgeInserter::FGroupEdgeSplitPoint& StartPoint,
 	const FGroupEdgeInserter::FGroupEdgeSplitPoint& EndPoint,
-	double VertexTolerance, TSet<int32>& PathEidsOut, 
+	double VertexTolerance, bool bSimplifyAlongPath, TSet<int32>& PathEidsOut, 
 	TSet<int32>* ChangedTrisOut, FProgressCancel* Progress)
 {
 	if (Progress && Progress->Cancelled())
@@ -1039,6 +1040,15 @@ bool EmbedPlaneCutPath(FDynamicMesh3& Mesh, const FGroupTopology& Topology,
 			PathEidsOut.Add(Eid);
 		}
 	}
+	
+	if (bSimplifyAlongPath)
+	{
+		FLocalPlanarSimplify Simplify;
+		// Note: Vertex normals are not reliable here, so preserving them will just prevent simplifications that are fine ...
+		Simplify.bPreserveVertexNormals = false;
+		// Note: Any triangles that this simplify can remove should already be in ChangedTrisOut, so we shouldn't need to update that array here
+		Simplify.SimplifyAlongEdges(Mesh, PathEidsOut);
+	}
 
 	return true;
 }
@@ -1056,7 +1066,7 @@ bool CreateNewGroups(FDynamicMesh3& Mesh, TSet<int32>& PathEids, int32 OriginalG
 		return false;
 	}
 
-	// Create the new groups. We do so
+	// Create the new groups.
 	TSet<int32> SeedTriangleSet;
 	for (int32 Eid : PathEids)
 	{
@@ -1144,7 +1154,7 @@ bool FGroupEdgeInserter::InsertGroupEdge(FGroupEdgeInsertionParams& Params, FGro
 		TSet<int32>* NewEids = OptionalOut.NewEidsOut ? OptionalOut.NewEidsOut : &TempNewEids;
 
 		bool bSuccess = EmbedPlaneCutPath(*Params.Mesh, *Params.Topology, Params.GroupID, Params.StartPoint,
-			Params.EndPoint, Params.VertexTolerance, *NewEids, OptionalOut.ChangedTidsOut, Progress);
+			Params.EndPoint, Params.VertexTolerance, Params.bSimplifyAlongPath, *NewEids, OptionalOut.ChangedTidsOut, Progress);
 		if (!bSuccess || (Progress && Progress->Cancelled()))
 		{
 			return false;
diff --git a/Engine/Plugins/Runtime/GeometryProcessing/Source/DynamicMesh/Private/Operations/LocalPlanarSimplify.cpp b/Engine/Plugins/Runtime/GeometryProcessing/Source/DynamicMesh/Private/Operations/LocalPlanarSimplify.cpp
new file mode 100644
--- /dev/null
+++ b/Engine/Plugins/Runtime/GeometryProcessing/Source/DynamicMesh/Private/Operations/LocalPlanarSimplify.cpp
@@ -0,0 +1,400 @@
+// Copyright Epic Games, Inc. All Rights Reserved.
+
+
+#include "Operations/LocalPlanarSimplify.h"
+
+#include "DynamicMesh/DynamicMeshAttributeSet.h"
+#include "DynamicMesh/DynamicMesh3.h"
+
+using namespace UE::Geometry;
+
+
+bool FLocalPlanarSimplify::IsFlat(const FDynamicMesh3& Mesh, int VID, double DotTolerance, FVector3d& OutFirstNormal)
+{
+	bool bHasFirst = false;
+	bool bIsFlat = true;
+	Mesh.EnumerateVertexTriangles(VID, [&Mesh, DotTolerance, &OutFirstNormal, &bHasFirst, &bIsFlat](int32 TID)
+		{
+			if (!bIsFlat)
+			{
+				return;
+			}
+			FVector3d Normal = Mesh.GetTriNormal(TID);
+			if (!bHasFirst)
+			{
+				OutFirstNormal = Normal;
+				bHasFirst = true;
+			}
+			else
+			{
+				bIsFlat = bIsFlat && Normal.Dot(OutFirstNormal) >= DotTolerance;
+			}
+		});
+	return bIsFlat;
+}
+
+/**
+ * Test if a given edge collapse would cause a triangle flip or other unacceptable decrease in mesh quality
+ */
+bool FLocalPlanarSimplify::CollapseWouldHurtTriangleQuality(
+	const FDynamicMesh3& Mesh, const FVector3d& ExpectNormal,
+	int32 RemoveV, const FVector3d& RemoveVPos, int32 KeepV, const FVector3d& KeepVPos,
+	double TryToImproveTriQualityThreshold
+	)
+{
+	double WorstQualityNewTriangle = FMathd::MaxReal;
+
+	bool bIsHurt = false;
+	Mesh.EnumerateVertexTriangles(RemoveV,
+		[&Mesh, &bIsHurt, &KeepVPos, RemoveV, KeepV, &ExpectNormal, 
+		TryToImproveTriQualityThreshold, &WorstQualityNewTriangle](int32 TID)
+		{
+			if (bIsHurt)
+			{
+				return;
+			}
+			FIndex3i Tri = Mesh.GetTriangle(TID);
+			FVector3d Verts[3];
+			for (int Idx = 0; Idx < 3; Idx++)
+			{
+				int VID = Tri[Idx];
+				if (VID == KeepV)
+				{
+					// this tri has both RemoveV and KeepV, so it'll be removed and we don't need to consider it
+					return;
+				}
+				else if (VID == RemoveV)
+				{
+					// anything at RemoveV is reconnected to KeepV's position
+					Verts[Idx] = KeepVPos;
+				}
+				else
+				{
+					// it's not on the collapsed edge so it stays still
+					Verts[Idx] = Mesh.GetVertex(Tri[Idx]);
+				}
+			}
+			FVector3d Edge1(Verts[1] - Verts[0]);
+			FVector3d Edge2(Verts[2] - Verts[0]);
+			FVector3d VCross(Edge2.Cross(Edge1));
+
+			// TODO: does this tolerance make a difference?  if not, set to zero and remove the Normalize(VCross)
+			double EdgeFlipTolerance = 1.e-5;
+			double Area2 = Normalize(VCross);
+			if (TryToImproveTriQualityThreshold > 0)
+			{
+				FVector3d Edge3(Verts[2] - Verts[1]);
+				double MaxLenSq = FMathd::Max3(Edge1.SquaredLength(), Edge2.SquaredLength(), Edge3.SquaredLength());
+				double Quality = Area2 / (MaxLenSq + FMathd::ZeroTolerance);
+				WorstQualityNewTriangle = FMathd::Min(Quality, WorstQualityNewTriangle);
+			}
+
+			bIsHurt = VCross.Dot(ExpectNormal) <= EdgeFlipTolerance;
+		}
+	);
+
+	// note tri quality was computed as 2*Area / MaxEdgeLenSquared
+	//  so need to multiply it by 2/sqrt(3) to get actual aspect ratio
+	if (!bIsHurt && WorstQualityNewTriangle * 2 * FMathd::InvSqrt3 < TryToImproveTriQualityThreshold)
+	{
+		// we found a bad tri; tentatively switch to rejecting this edge collapse
+		bIsHurt = true;
+		// but if there was an even worse tri in the original neighborhood, accept the collapse after all
+		Mesh.EnumerateVertexTriangles(RemoveV, [&Mesh, &bIsHurt, WorstQualityNewTriangle](int TID)
+			{
+				if (!bIsHurt) // early out if we already found an originally-worse tri
+				{
+					return;
+				}
+				FVector3d A, B, C;
+				Mesh.GetTriVertices(TID, A, B, C);
+				FVector3d E1 = B - A, E2 = C - A, E3 = C - B;
+				double Area2 = E1.Cross(E2).Length();
+				double MaxLenSq = FMathd::Max3(E1.SquaredLength(), E2.SquaredLength(), E3.SquaredLength());
+				double Quality = Area2 / (MaxLenSq + FMathd::ZeroTolerance);
+				if (Quality < WorstQualityNewTriangle)
+				{
+					bIsHurt = false;
+				}
+			}
+		);
+	}
+	return bIsHurt;
+}
+
+
+/**
+ * Test if a given edge collapse would change the mesh shape or UVs unacceptably
+ */
+bool FLocalPlanarSimplify::CollapseWouldChangeShapeOrUVs(
+	const FDynamicMesh3& Mesh, const TSet<int>& CutBoundaryEdgeSet, double DotTolerance, int SourceEID,
+	int32 RemoveV, const FVector3d& RemoveVPos, int32 KeepV, const FVector3d& KeepVPos, const FVector3d& EdgeDir,
+	bool bPreserveTriangleGroups, bool bPreserveUVsForMesh, bool bPreserveVertexUVs, bool bPreserveOverlayUVs,
+	float UVEqualThresholdSq, bool bPreserveVertexNormals, float NormalEqualCosThreshold)
+{
+	// Search the edges connected to the vertex to find one in the boundary set that points in the opposite direction
+	// If we don't find that edge, or if there are other boundary/seam edges attached, we can't remove this vertex
+	// We also can't remove the vertex if doing so would distort the UVs
+	bool bHasBadEdge = false;
+
+	int OpposedEdge = -1;
+	int SourceGroupID = Mesh.GetTriangleGroup(Mesh.GetEdgeT(SourceEID).A);
+	Mesh.EnumerateVertexEdges(RemoveV,
+		[&](int32 VertEID)
+		{
+			if (bHasBadEdge || VertEID == SourceEID)
+			{
+				return;
+			}
+
+			FDynamicMesh3::FEdge Edge = Mesh.GetEdge(VertEID);
+			if (bPreserveTriangleGroups && Mesh.HasTriangleGroups())
+			{
+				if (SourceGroupID != Mesh.GetTriangleGroup(Edge.Tri.A) ||
+					(Edge.Tri.B != FDynamicMesh3::InvalidID && SourceGroupID != Mesh.GetTriangleGroup(Edge.Tri.B)))
+				{
+					// RemoveV is on a group boundary, so the edge collapse would change the shape of the groups
+					bHasBadEdge = true;
+					return;
+				}
+			}
+
+			// it's a known boundary edge; check if it's the opposite-facing one we need
+			if (CutBoundaryEdgeSet.Contains(VertEID))
+			{
+				if (OpposedEdge != -1)
+				{
+					bHasBadEdge = true;
+					return;
+				}
+				FIndex2i OtherEdgeV = Edge.Vert;
+				int OtherV = IndexUtil::FindEdgeOtherVertex(OtherEdgeV, RemoveV);
+				FVector3d OtherVPos = Mesh.GetVertex(OtherV);
+				FVector3d OtherEdgeDir = OtherVPos - RemoveVPos;
+				if (Normalize(OtherEdgeDir) == 0)
+				{
+					// collapsing degenerate edges above should prevent this
+					bHasBadEdge = true;
+					return; // break instead of continue to skip the whole edge
+				}
+				if (OtherEdgeDir.Dot(EdgeDir) <= -DotTolerance)
+				{
+					OpposedEdge = VertEID;
+				}
+				else
+				{
+					bHasBadEdge = true;
+					return;
+				}
+
+				// test that UVs are not too distorted through the collapse
+				if (!bPreserveUVsForMesh)
+				{
+					return;
+				}
+				float LerpT = float( (RemoveVPos - OtherVPos).Dot(OtherEdgeDir) / (KeepVPos - OtherVPos).Dot(OtherEdgeDir) );
+				if (bPreserveVertexUVs && Mesh.HasVertexUVs())
+				{
+					FVector2f OtherUV = Mesh.GetVertexUV(OtherV);
+					FVector2f RemoveUV = Mesh.GetVertexUV(RemoveV);
+					FVector2f KeepUV = Mesh.GetVertexUV(KeepV);
+					if ( DistanceSquared( Lerp(OtherUV, KeepUV, LerpT), RemoveUV) > UVEqualThresholdSq)
+					{
+						bHasBadEdge = true;
+						return;
+					}
+				}
+				if (bPreserveVertexNormals && Mesh.HasVertexNormals())
+				{
+					FVector3f OtherN = Mesh.GetVertexNormal(OtherV);
+					FVector3f RemoveN = Mesh.GetVertexNormal(RemoveV);
+					FVector3f KeepN = Mesh.GetVertexNormal(KeepV);
+					if (Normalized(Lerp(OtherN, KeepN, LerpT)).Dot(Normalized(RemoveN)) < NormalEqualCosThreshold)
+					{
+						bHasBadEdge = true;
+						return;
+					}
+				}
+				if (bPreserveOverlayUVs && Mesh.HasAttributes())
+				{
+					int NumLayers = Mesh.Attributes()->NumUVLayers();
+					FIndex2i SourceEdgeTris = Mesh.GetEdgeT(SourceEID);
+					FIndex2i OppEdgeTris = Edge.Tri;
+
+					// special handling of seam edge when the edges aren't boundary edges
+					// -- if they're seams, we'd need to check both sides of the seams for a UV match
+					//    but this is complicated and should be quite rare, so we just don't collapse these
+					if (SourceEdgeTris.B != -1 || OppEdgeTris.B != -1)
+					{
+						if (Mesh.Attributes()->IsSeamEdge(SourceEID) ||
+							Mesh.Attributes()->IsSeamEdge(VertEID))
+						{
+							bHasBadEdge = true;
+							return;
+						}
+					}
+
+					FIndex3i SourceBaseTri = Mesh.GetTriangle(SourceEdgeTris.A);
+					FIndex3i OppBaseTri = Mesh.GetTriangle(OppEdgeTris.A);
+					int KeepSourceIdx = IndexUtil::FindTriIndex(KeepV, SourceBaseTri);
+					int RemoveSourceIdx = IndexUtil::FindTriIndex(RemoveV, SourceBaseTri);
+					int OtherOppIdx = IndexUtil::FindTriIndex(OtherV, OppBaseTri);
+					if (!ensure(KeepSourceIdx != -1 && RemoveSourceIdx != -1 && OtherOppIdx != -1))
+					{
+						bHasBadEdge = true;
+						return;
+					}
+
+					// get the UVs per overlay off the triangle(s) attached the two edges
+					for (int UVLayerIdx = 0; UVLayerIdx < NumLayers; UVLayerIdx++)
+					{
+						const FDynamicMeshUVOverlay* UVs = Mesh.Attributes()->GetUVLayer(UVLayerIdx);
+						if (UVs->ElementCount() < 3)
+						{
+							// overlay is not actually in use; skip it
+							continue;
+						}
+						FIndex3i SourceT = UVs->GetTriangle(SourceEdgeTris.A);
+						FIndex3i OppT = UVs->GetTriangle(OppEdgeTris.A);
+						int KeepE = SourceT[KeepSourceIdx];
+						int RemoveE = SourceT[RemoveSourceIdx];
+						int OtherE = OppT[OtherOppIdx];
+						if (KeepE == -1 || RemoveE == -1 || OtherE == -1)
+						{
+							// overlay is not set on relevant triangles; skip it
+							continue;
+						}
+						FVector2f OtherUV = UVs->GetElement(OtherE);
+						FVector2f RemoveUV = UVs->GetElement(RemoveE);
+						FVector2f KeepUV = UVs->GetElement(KeepE);
+						if ( DistanceSquared( Lerp(OtherUV, KeepUV, LerpT), RemoveUV) > UVEqualThresholdSq)
+						{
+							bHasBadEdge = true;
+							return;
+						}
+					}
+				}
+			}
+			else // it wasn't in the boundary edge set; check if it's one that would prevent us from safely removing the vertex
+			{
+				if (Mesh.IsBoundaryEdge(VertEID) || (Mesh.HasAttributes() && Mesh.Attributes()->IsSeamEdge(VertEID)))
+				{
+					bHasBadEdge = true;
+				}
+			}
+		});
+
+	return bHasBadEdge;
+}
+
+
+void FLocalPlanarSimplify::SimplifyAlongEdges(FDynamicMesh3& Mesh, TSet<int32>& InOutEdges, TUniqueFunction<void(const DynamicMeshInfo::FEdgeCollapseInfo&)> ProcessCollapse) const
+{
+	double DotTolerance = FMathd::Cos(SimplificationAngleTolerance * FMathd::DegToRad);
+
+	// Save the input list of edges to iterate over, so we can edit the set w/ new edges while safely iterating over the old ones
+	TArray<int32> CutBoundaryEdgesArray = InOutEdges.Array();
+
+	int NumCollapses = 0, CollapseIters = 0;
+	int MaxCollapseIters = 1; // TODO: is there a case where we need more iterations?  Perhaps if we add some triangle quality criteria?
+	while (CollapseIters < MaxCollapseIters)
+	{
+		int LastNumCollapses = NumCollapses;
+		for (int EID : CutBoundaryEdgesArray)
+		{
+			// this can happen if a collapse removes another cut boundary edge
+			// (which can happen e.g. if you have a degenerate (colinear) tri flat on the cut boundary)
+			if (!Mesh.IsEdge(EID))
+			{
+				continue;
+			}
+
+			FDynamicMesh3::FEdge Edge = Mesh.GetEdge(EID);
+
+			// track whether the neighborhood of the vertex is flat
+			bool Flat[2]{ false, false };
+			// normals for each flat vertex
+			FVector3d FlatNormals[2]{ FVector3d::Zero(), FVector3d::Zero() };
+			int NumFlat = 0;
+			for (int VIdx = 0; VIdx < 2; VIdx++)
+			{
+				Flat[VIdx] = IsFlat(Mesh, Edge.Vert[VIdx], DotTolerance, FlatNormals[VIdx]);
+
+				if (Flat[VIdx])
+				{
+					NumFlat++;
+				}
+			}
+
+			if (NumFlat == 0)
+			{
+				continue;
+			}
+
+			// see if we can collapse to remove either vertex
+			for (int RemoveVIdx = 0; RemoveVIdx < 2; RemoveVIdx++)
+			{
+				if (!Flat[RemoveVIdx])
+				{
+					continue;
+				}
+				int KeepVIdx = 1 - RemoveVIdx;
+				FVector3d RemoveVPos = Mesh.GetVertex(Edge.Vert[RemoveVIdx]);
+				FVector3d KeepVPos = Mesh.GetVertex(Edge.Vert[KeepVIdx]);
+				FVector3d EdgeDir = KeepVPos - RemoveVPos;
+				if (Normalize(EdgeDir) == 0) // 0 is returned as a special case when the edge was too short to normalize
+				{
+					// This case is often avoided by collapsing degenerate edges in a separate pre-pass, so we just skip over it here for now
+					// TODO: Consider adding degenerate edge-collapse logic here
+					break; // break instead of continue to skip the whole edge
+				}
+
+				bool bHasBadEdge = false; // will be set if either mesh can't collapse the edge
+				int RemoveV = Edge.Vert[RemoveVIdx];
+				int KeepV = Edge.Vert[KeepVIdx];
+				int SourceEID = EID;
+
+				bHasBadEdge = bHasBadEdge || CollapseWouldHurtTriangleQuality(
+					Mesh, FlatNormals[RemoveVIdx], RemoveV, RemoveVPos, KeepV, KeepVPos, TryToImproveTriQualityThreshold);
+
+				bHasBadEdge = bHasBadEdge || CollapseWouldChangeShapeOrUVs(
+					Mesh, InOutEdges, DotTolerance,
+					SourceEID, RemoveV, RemoveVPos, KeepV, KeepVPos, EdgeDir, bPreserveTriangleGroups,
+					true, bPreserveVertexUVs, bPreserveOverlayUVs, UVDistortTolerance * UVDistortTolerance,
+					bPreserveVertexNormals, FMathf::Cos(NormalDistortTolerance * FMathf::DegToRad));
+
+				if (bHasBadEdge)
+				{
+					continue;
+				}
+
+				FDynamicMesh3::FEdgeCollapseInfo CollapseInfo;
+				EMeshResult CollapseResult = Mesh.CollapseEdge(KeepV, RemoveV, 0, CollapseInfo);
+				if (CollapseResult == EMeshResult::Ok)
+				{
+					if (ProcessCollapse)
+					{
+						ProcessCollapse(CollapseInfo);
+					}
+					NumCollapses++;
+					InOutEdges.Remove(CollapseInfo.CollapsedEdge);
+					InOutEdges.Remove(CollapseInfo.RemovedEdges[0]);
+					if (CollapseInfo.RemovedEdges[1] != -1)
+					{
+						InOutEdges.Remove(CollapseInfo.RemovedEdges[1]);
+					}
+				}
+				break; // if we got through to trying to collapse the edge, don't try to collapse from the other vertex.
+			}
+		}
+
+		CutBoundaryEdgesArray = InOutEdges.Array();
+
+		if (NumCollapses == LastNumCollapses)
+		{
+			break;
+		}
+
+		CollapseIters++;
+	}
+}
diff --git a/Engine/Plugins/Runtime/GeometryProcessing/Source/DynamicMesh/Private/Operations/MeshBoolean.cpp b/Engine/Plugins/Runtime/GeometryProcessing/Source/DynamicMesh/Private/Operations/MeshBoolean.cpp
--- a/Engine/Plugins/Runtime/GeometryProcessing/Source/DynamicMesh/Private/Operations/MeshBoolean.cpp
+++ b/Engine/Plugins/Runtime/GeometryProcessing/Source/DynamicMesh/Private/Operations/MeshBoolean.cpp
@@ -4,6 +4,7 @@
 
 #include "Operations/MeshBoolean.h"
 #include "Operations/MeshMeshCut.h"
+#include "Operations/LocalPlanarSimplify.h"
 #include "Selections/MeshConnectedComponents.h"
 
 #include "Async/ParallelFor.h"
@@ -762,285 +763,6 @@ bool FMeshBoolean::Compute()
 
 
 
-bool FMeshBoolean::IsFlat(const FDynamicMesh3& Mesh, int VID, double DotTolerance, FVector3d& OutFirstNormal)
-{
-	bool bHasFirst = false;
-	bool bIsFlat = true;
-	Mesh.EnumerateVertexTriangles(VID, [&Mesh, DotTolerance, &OutFirstNormal, &bHasFirst, &bIsFlat](int32 TID)
-		{
-			if (!bIsFlat)
-			{
-				return;
-			}
-			FVector3d Normal = Mesh.GetTriNormal(TID);
-			if (!bHasFirst)
-			{
-				OutFirstNormal = Normal;
-				bHasFirst = true;
-			}
-			else
-			{
-				bIsFlat = bIsFlat && Normal.Dot(OutFirstNormal) >= DotTolerance;
-			}
-		});
-	return bIsFlat;
-}
-
-/**
- * Test if a given edge collapse would cause a triangle flip or other unacceptable decrease in mesh quality
- */
-bool FMeshBoolean::CollapseWouldHurtTriangleQuality(
-	const FDynamicMesh3& Mesh, const FVector3d& ExpectNormal,
-	int32 RemoveV, const FVector3d& RemoveVPos, int32 KeepV, const FVector3d& KeepVPos,
-	double TryToImproveTriQualityThreshold
-	)
-{
-	double WorstQualityNewTriangle = FMathd::MaxReal;
-
-	bool bIsHurt = false;
-	Mesh.EnumerateVertexTriangles(RemoveV,
-		[&Mesh, &bIsHurt, &KeepVPos, RemoveV, KeepV, &ExpectNormal, 
-		TryToImproveTriQualityThreshold, &WorstQualityNewTriangle](int32 TID)
-		{
-			if (bIsHurt)
-			{
-				return;
-			}
-			FIndex3i Tri = Mesh.GetTriangle(TID);
-			FVector3d Verts[3];
-			for (int Idx = 0; Idx < 3; Idx++)
-			{
-				int VID = Tri[Idx];
-				if (VID == KeepV)
-				{
-					// this tri has both RemoveV and KeepV, so it'll be removed and we don't need to consider it
-					return;
-				}
-				else if (VID == RemoveV)
-				{
-					// anything at RemoveV is reconnected to KeepV's position
-					Verts[Idx] = KeepVPos;
-				}
-				else
-				{
-					// it's not on the collapsed edge so it stays still
-					Verts[Idx] = Mesh.GetVertex(Tri[Idx]);
-				}
-			}
-			FVector3d Edge1(Verts[1] - Verts[0]);
-			FVector3d Edge2(Verts[2] - Verts[0]);
-			FVector3d VCross(Edge2.Cross(Edge1));
-
-			// TODO: does this tolerance make a difference?  if not, set to zero and remove the Normalize(VCross)
-			double EdgeFlipTolerance = 1.e-5;
-			double Area2 = Normalize(VCross);
-			if (TryToImproveTriQualityThreshold > 0)
-			{
-				FVector3d Edge3(Verts[2] - Verts[1]);
-				double MaxLenSq = FMathd::Max3(Edge1.SquaredLength(), Edge2.SquaredLength(), Edge3.SquaredLength());
-				double Quality = Area2 / (MaxLenSq + FMathd::ZeroTolerance);
-				WorstQualityNewTriangle = FMathd::Min(Quality, WorstQualityNewTriangle);
-			}
-
-			bIsHurt = VCross.Dot(ExpectNormal) <= EdgeFlipTolerance;
-		}
-	);
-
-	// note tri quality was computed as 2*Area / MaxEdgeLenSquared
-	//  so need to multiply it by 2/sqrt(3) to get actual aspect ratio
-	if (!bIsHurt && WorstQualityNewTriangle * 2 * FMathd::InvSqrt3 < TryToImproveTriQualityThreshold)
-	{
-		// we found a bad tri; tentatively switch to rejecting this edge collapse
-		bIsHurt = true;
-		// but if there was an even worse tri in the original neighborhood, accept the collapse after all
-		Mesh.EnumerateVertexTriangles(RemoveV, [&Mesh, &bIsHurt, WorstQualityNewTriangle](int TID)
-			{
-				if (!bIsHurt) // early out if we already found an originally-worse tri
-				{
-					return;
-				}
-				FVector3d A, B, C;
-				Mesh.GetTriVertices(TID, A, B, C);
-				FVector3d E1 = B - A, E2 = C - A, E3 = C - B;
-				double Area2 = E1.Cross(E2).Length();
-				double MaxLenSq = FMathd::Max3(E1.SquaredLength(), E2.SquaredLength(), E3.SquaredLength());
-				double Quality = Area2 / (MaxLenSq + FMathd::ZeroTolerance);
-				if (Quality < WorstQualityNewTriangle)
-				{
-					bIsHurt = false;
-				}
-			}
-		);
-	}
-	return bIsHurt;
-}
-
-
-/**
- * Test if a given edge collapse would change the mesh shape or UVs unacceptably
- */
-bool FMeshBoolean::CollapseWouldChangeShapeOrUVs(
-	const FDynamicMesh3& Mesh, const TSet<int>& CutBoundaryEdgeSet, double DotTolerance, int SourceEID,
-	int32 RemoveV, const FVector3d& RemoveVPos, int32 KeepV, const FVector3d& KeepVPos, const FVector3d& EdgeDir,
-	bool bPreserveTriangleGroups, bool bPreserveUVsForMesh, bool bPreserveVertexUVs, bool bPreserveOverlayUVs,
-	float UVEqualThresholdSq, bool bPreserveVertexNormals, float NormalEqualCosThreshold)
-{
-	// Search the edges connected to the vertex to find one in the boundary set that points in the opposite direction
-	// If we don't find that edge, or if there are other boundary/seam edges attached, we can't remove this vertex
-	// We also can't remove the vertex if doing so would distort the UVs
-	bool bHasBadEdge = false;
-
-	int OpposedEdge = -1;
-	int SourceGroupID = Mesh.GetTriangleGroup(Mesh.GetEdgeT(SourceEID).A);
-	Mesh.EnumerateVertexEdges(RemoveV,
-		[&](int32 VertEID)
-		{
-			if (bHasBadEdge || VertEID == SourceEID)
-			{
-				return;
-			}
-
-			FDynamicMesh3::FEdge Edge = Mesh.GetEdge(VertEID);
-			if (bPreserveTriangleGroups && Mesh.HasTriangleGroups())
-			{
-				if (SourceGroupID != Mesh.GetTriangleGroup(Edge.Tri.A) ||
-					(Edge.Tri.B != FDynamicMesh3::InvalidID && SourceGroupID != Mesh.GetTriangleGroup(Edge.Tri.B)))
-				{
-					// RemoveV is on a group boundary, so the edge collapse would change the shape of the groups
-					bHasBadEdge = true;
-					return;
-				}
-			}
-
-			// it's a known boundary edge; check if it's the opposite-facing one we need
-			if (CutBoundaryEdgeSet.Contains(VertEID))
-			{
-				if (OpposedEdge != -1)
-				{
-					bHasBadEdge = true;
-					return;
-				}
-				FIndex2i OtherEdgeV = Edge.Vert;
-				int OtherV = IndexUtil::FindEdgeOtherVertex(OtherEdgeV, RemoveV);
-				FVector3d OtherVPos = Mesh.GetVertex(OtherV);
-				FVector3d OtherEdgeDir = OtherVPos - RemoveVPos;
-				if (Normalize(OtherEdgeDir) == 0)
-				{
-					// collapsing degenerate edges above should prevent this
-					bHasBadEdge = true;
-					return; // break instead of continue to skip the whole edge
-				}
-				if (OtherEdgeDir.Dot(EdgeDir) <= -DotTolerance)
-				{
-					OpposedEdge = VertEID;
-				}
-				else
-				{
-					bHasBadEdge = true;
-					return;
-				}
-
-				// test that UVs are not too distorted through the collapse
-				if (!bPreserveUVsForMesh)
-				{
-					return;
-				}
-				float LerpT = float( (RemoveVPos - OtherVPos).Dot(OtherEdgeDir) / (KeepVPos - OtherVPos).Dot(OtherEdgeDir) );
-				if (bPreserveVertexUVs && Mesh.HasVertexUVs())
-				{
-					FVector2f OtherUV = Mesh.GetVertexUV(OtherV);
-					FVector2f RemoveUV = Mesh.GetVertexUV(RemoveV);
-					FVector2f KeepUV = Mesh.GetVertexUV(KeepV);
-					if ( DistanceSquared( Lerp(OtherUV, KeepUV, LerpT), RemoveUV) > UVEqualThresholdSq)
-					{
-						bHasBadEdge = true;
-						return;
-					}
-				}
-				if (bPreserveVertexNormals && Mesh.HasVertexNormals())
-				{
-					FVector3f OtherN = Mesh.GetVertexNormal(OtherV);
-					FVector3f RemoveN = Mesh.GetVertexNormal(RemoveV);
-					FVector3f KeepN = Mesh.GetVertexNormal(KeepV);
-					if (Normalized(Lerp(OtherN, KeepN, LerpT)).Dot(Normalized(RemoveN)) < NormalEqualCosThreshold)
-					{
-						bHasBadEdge = true;
-						return;
-					}
-				}
-				if (bPreserveOverlayUVs && Mesh.HasAttributes())
-				{
-					int NumLayers = Mesh.Attributes()->NumUVLayers();
-					FIndex2i SourceEdgeTris = Mesh.GetEdgeT(SourceEID);
-					FIndex2i OppEdgeTris = Edge.Tri;
-
-					// special handling of seam edge when the edges aren't boundary edges
-					// -- if they're seams, we'd need to check both sides of the seams for a UV match
-					//    but this is complicated and should be quite rare, so we just don't collapse these
-					if (SourceEdgeTris.B != -1 || OppEdgeTris.B != -1)
-					{
-						if (Mesh.Attributes()->IsSeamEdge(SourceEID) ||
-							Mesh.Attributes()->IsSeamEdge(VertEID))
-						{
-							bHasBadEdge = true;
-							return;
-						}
-					}
-
-					FIndex3i SourceBaseTri = Mesh.GetTriangle(SourceEdgeTris.A);
-					FIndex3i OppBaseTri = Mesh.GetTriangle(OppEdgeTris.A);
-					int KeepSourceIdx = IndexUtil::FindTriIndex(KeepV, SourceBaseTri);
-					int RemoveSourceIdx = IndexUtil::FindTriIndex(RemoveV, SourceBaseTri);
-					int OtherOppIdx = IndexUtil::FindTriIndex(OtherV, OppBaseTri);
-					if (!ensure(KeepSourceIdx != -1 && RemoveSourceIdx != -1 && OtherOppIdx != -1))
-					{
-						bHasBadEdge = true;
-						return;
-					}
-
-					// get the UVs per overlay off the triangle(s) attached the two edges
-					for (int UVLayerIdx = 0; UVLayerIdx < NumLayers; UVLayerIdx++)
-					{
-						const FDynamicMeshUVOverlay* UVs = Mesh.Attributes()->GetUVLayer(UVLayerIdx);
-						if (UVs->ElementCount() < 3)
-						{
-							// overlay is not actually in use; skip it
-							continue;
-						}
-						FIndex3i SourceT = UVs->GetTriangle(SourceEdgeTris.A);
-						FIndex3i OppT = UVs->GetTriangle(OppEdgeTris.A);
-						int KeepE = SourceT[KeepSourceIdx];
-						int RemoveE = SourceT[RemoveSourceIdx];
-						int OtherE = OppT[OtherOppIdx];
-						if (KeepE == -1 || RemoveE == -1 || OtherE == -1)
-						{
-							// overlay is not set on relevant triangles; skip it
-							continue;
-						}
-						FVector2f OtherUV = UVs->GetElement(OtherE);
-						FVector2f RemoveUV = UVs->GetElement(RemoveE);
-						FVector2f KeepUV = UVs->GetElement(KeepE);
-						if ( DistanceSquared( Lerp(OtherUV, KeepUV, LerpT), RemoveUV) > UVEqualThresholdSq)
-						{
-							bHasBadEdge = true;
-							return;
-						}
-					}
-				}
-			}
-			else // it wasn't in the boundary edge set; check if it's one that would prevent us from safely removing the vertex
-			{
-				if (Mesh.IsBoundaryEdge(VertEID) || (Mesh.HasAttributes() && Mesh.Attributes()->IsSeamEdge(VertEID)))
-				{
-					bHasBadEdge = true;
-				}
-			}
-		});
-
-	return bHasBadEdge;
-}
-
-
 void FMeshBoolean::SimplifyAlongNewEdges(int NumMeshesToProcess, FDynamicMesh3* CutMesh[2], TArray<int> CutBoundaryEdges[2], TMap<int, int>& AllVIDMatches)
 {
 	double DotTolerance = FMathd::Cos(SimplificationAngleTolerance * FMathd::DegToRad);
@@ -1113,9 +835,9 @@ void FMeshBoolean::SimplifyAlongNewEdges(int NumMeshesToProcess, FDynamicMesh3*
 			int NumFlat = 0;
 			for (int VIdx = 0; VIdx < 2; VIdx++)
 			{
-				if (IsFlat(*CutMesh[0], Edge.Vert[VIdx], DotTolerance, FlatNormals[VIdx][0]))
+				if (FLocalPlanarSimplify::IsFlat(*CutMesh[0], Edge.Vert[VIdx], DotTolerance, FlatNormals[VIdx][0]))
 				{
-					Flat[VIdx] = (Matches[VIdx] == -1) || IsFlat(*CutMesh[1], Matches[VIdx], DotTolerance, FlatNormals[VIdx][1]);
+					Flat[VIdx] = (Matches[VIdx] == -1) || FLocalPlanarSimplify::IsFlat(*CutMesh[1], Matches[VIdx], DotTolerance, FlatNormals[VIdx][1]);
 				}
 
 				if (Flat[VIdx])
@@ -1155,10 +877,10 @@ void FMeshBoolean::SimplifyAlongNewEdges(int NumMeshesToProcess, FDynamicMesh3*
 					int KeepV = MeshIdx == 0 ? Edge.Vert[KeepVIdx] : Matches[KeepVIdx];
 					int SourceEID = MeshIdx == 0 ? EID : OtherEID;
 
-					bHasBadEdge = bHasBadEdge || CollapseWouldHurtTriangleQuality(*CutMesh[MeshIdx],
+					bHasBadEdge = bHasBadEdge || FLocalPlanarSimplify::CollapseWouldHurtTriangleQuality(*CutMesh[MeshIdx],
 						FlatNormals[RemoveVIdx][MeshIdx], RemoveV, RemoveVPos, KeepV, KeepVPos, TryToImproveTriQualityThreshold);
 					
-					bHasBadEdge = bHasBadEdge || CollapseWouldChangeShapeOrUVs(
+					bHasBadEdge = bHasBadEdge || FLocalPlanarSimplify::CollapseWouldChangeShapeOrUVs(
 						*CutMesh[MeshIdx], CutBoundaryEdgeSets[MeshIdx], DotTolerance,
 						SourceEID, RemoveV, RemoveVPos, KeepV, KeepVPos, EdgeDir, bPreserveTriangleGroups,
 						PreserveUVsOnlyForMesh == -1 || MeshIdx == PreserveUVsOnlyForMesh,
diff --git a/Engine/Plugins/Runtime/GeometryProcessing/Source/DynamicMesh/Private/Operations/MeshPlaneCut.cpp b/Engine/Plugins/Runtime/GeometryProcessing/Source/DynamicMesh/Private/Operations/MeshPlaneCut.cpp
--- a/Engine/Plugins/Runtime/GeometryProcessing/Source/DynamicMesh/Private/Operations/MeshPlaneCut.cpp
+++ b/Engine/Plugins/Runtime/GeometryProcessing/Source/DynamicMesh/Private/Operations/MeshPlaneCut.cpp
@@ -17,33 +17,11 @@
 
 using namespace UE::Geometry;
 
-void FMeshPlaneCut::SplitCrossingEdges(TArray<double>& Signs, TSet<int>& ZeroEdges, TSet<int>& OnCutEdges, bool bDeleteTrisOnPlane)
-{
-	TSet<int> OnSplitEdges;
-	SplitCrossingEdges(Signs, ZeroEdges, OnCutEdges, OnSplitEdges, bDeleteTrisOnPlane);
-}
-
-void FMeshPlaneCut::SplitCrossingEdges(TArray<double>& Signs, TSet<int>& ZeroEdges, TSet<int>& OnCutEdges, TSet<int>& OnSplitEdges, bool bDeleteTrisOnPlane)
+void FMeshPlaneCut::ComputeVertexSignedDistances(TArray<double>& Signs, double InvalidDist)
 {
-	Signs.Reset(); ZeroEdges.Reset(); OnCutEdges.Reset(); OnSplitEdges.Reset();
-	OnCutVertices.Reset();
-
-	double InvalidDist = -FMathd::MaxReal;
-
-	// TODO: handle selections
-	//MeshEdgeSelection CutEdgeSet = null;
-	//MeshVertexSelection CutVertexSet = null;
-	//if (CutFaceSet != null) {
-	//	CutEdgeSet = new MeshEdgeSelection(Mesh, CutFaceSet);
-	//	CutVertexSet = new MeshVertexSelection(Mesh, CutEdgeSet);
-	//}
-
-	// compute Signs
 	int MaxVID = Mesh->MaxVertexID();
-	
 	Signs.SetNum(MaxVID);
-
-	bool bNoParallel = false;
+	constexpr bool bNoParallel = false;
 	ParallelFor(MaxVID, [&](int32 VID)
 	{
 		if (Mesh->IsVertex(VID))
@@ -55,6 +33,23 @@ void FMeshPlaneCut::SplitCrossingEdges(TArray<double>& Signs, TSet<int>& ZeroEdg
 			Signs[VID] = InvalidDist;
 		}
 	}, bNoParallel);
+}
+
+void FMeshPlaneCut::SplitCrossingEdges(bool bDeleteTrisOnPlane, TArray<double>& Signs, TSet<int>& AlreadyOnPlaneEdges, TSet<int32>& CutPlaneEdges, TSet<int>* OnSplitEdges, TSet<int>* OnPlaneVertices, TSet<int32>* TriangleSelection)
+{
+	AlreadyOnPlaneEdges.Reset();
+	CutPlaneEdges.Reset();
+	if (OnSplitEdges)
+	{
+		OnSplitEdges->Reset();
+	}
+	if (OnPlaneVertices)
+	{
+		OnPlaneVertices->Reset();
+	}
+
+	// Compute signed distances at all vertices. Set invalid dists to zero, because any vertex we add will be on the plane.
+	ComputeVertexSignedDistances(Signs, 0.0);
 
 	if (bDeleteTrisOnPlane)
 	{
@@ -67,26 +62,15 @@ void FMeshPlaneCut::SplitCrossingEdges(TArray<double>& Signs, TSet<int>& ZeroEdg
 
 			FIndex3i Tri = Mesh->GetTriangle(TID);
 			FIndex3i TriEdges = Mesh->GetTriEdges(TID);
-			if (   FMathd::Abs(Signs[Tri.A]) < PlaneTolerance
-				&& FMathd::Abs(Signs[Tri.B]) < PlaneTolerance
-				&& FMathd::Abs(Signs[Tri.C]) < PlaneTolerance)
+			if (FMathd::Abs(Signs[Tri.A]) < PlaneTolerance &&
+				FMathd::Abs(Signs[Tri.B]) < PlaneTolerance &&
+				FMathd::Abs(Signs[Tri.C]) < PlaneTolerance)
 			{
-				EMeshResult Res = Mesh->RemoveTriangle(TID, true, false);
-				ensure(Res == EMeshResult::Ok);
-				for (int EIdx = 0; EIdx < 3; EIdx++)
+				Mesh->RemoveTriangle(TID, true, false);
+				if (TriangleSelection)
 				{
-					int EID = TriEdges[EIdx];
-					if (Mesh->IsEdge(EID))
-					{
-						// any edge that still exists after removal is a possible cut edge
-						OnCutEdges.Add(EID);
-					}
-					else
-					{
-						// in case the now-gone edge was added to cut edges earlier, make sure it's removed
-						// (note: if it's not present this function will just do nothing, which is fine)
-						OnCutEdges.Remove(EID);
-					}
+					// Remove triangle from the selection as well (does nothing if it was not in the selection already)
+					TriangleSelection->Remove(TID);
 				}
 			}
 		}
@@ -95,22 +79,19 @@ void FMeshPlaneCut::SplitCrossingEdges(TArray<double>& Signs, TSet<int>& ZeroEdg
 	// have to skip processing of new edges. If edge id
 	// is > max at start, is new. Otherwise if in NewEdges list, also new.
 	int MaxEID = Mesh->MaxEdgeID();
-	TSet<int> NewEdges;
-
-	FDynamicMesh3::edge_iterator EdgeItr = Mesh->EdgeIndicesItr();
-	// TODO: selection logic
-	//IEnumerable<int> edgeItr = Interval1i.Range(MaxEID);
-	//if (CutEdgeSet != null)
-	//	edgeItr = CutEdgeSet;
-
-	// cut existing edges with plane, using edge split
-	for (int EID : EdgeItr)
+	TSet<int> NewEdgesBeforeMaxID;
+	auto AddNewEdge = [&NewEdgesBeforeMaxID, MaxEID](int32 NewEID)
 	{
-		if (!Mesh->IsEdge(EID))
+		if (NewEID < MaxEID)
 		{
-			continue;
+			NewEdgesBeforeMaxID.Add(NewEID);
 		}
-		if (EID >= MaxEID || NewEdges.Contains(EID))
+	};
+
+	// cut existing edges with plane, using edge split
+	for (int32 EID = 0; EID < MaxEID; ++EID)
+	{
+		if (!Mesh->IsEdge(EID) || NewEdgesBeforeMaxID.Contains(EID))
 		{
 			continue;
 		}
@@ -119,71 +100,126 @@ void FMeshPlaneCut::SplitCrossingEdges(TArray<double>& Signs, TSet<int>& ZeroEdg
 			continue;
 		}
 
-		FIndex2i ev = Mesh->GetEdgeV(EID);
-		double f0 = Signs[ev.A];
-		double f1 = Signs[ev.B];
+		FIndex2i EdgeV = Mesh->GetEdgeV(EID);
+		const double DistA = Signs[EdgeV.A];
+		const double DistB = Signs[EdgeV.B];
 
 		// If both Signs are 0, this edge is on-contour
 		// If one sign is 0, that vertex is on-contour
-		int n0 = (FMathd::Abs(f0) < PlaneTolerance) ? 1 : 0;
-		int n1 = (FMathd::Abs(f1) < PlaneTolerance) ? 1 : 0;
-		if (n0 + n1 > 0)
+		int AOnPlane = (FMathd::Abs(DistA) < PlaneTolerance) ? 1 : 0;
+		int BOnPlane = (FMathd::Abs(DistB) < PlaneTolerance) ? 1 : 0;
+		if (AOnPlane || BOnPlane)
 		{
-			if (n0 + n1 == 2)
+			if (AOnPlane && BOnPlane)
 			{
-				ZeroEdges.Add(EID);
-				OnCutVertices.Add(ev.A);
-				OnCutVertices.Add(ev.B);
+				AlreadyOnPlaneEdges.Add(EID);
+				if (OnPlaneVertices)
+				{
+					OnPlaneVertices->Add(EdgeV.A);
+					OnPlaneVertices->Add(EdgeV.B);
+				}
 			}
-			else
+			else if (OnPlaneVertices)
 			{
-				OnCutVertices.Add((n0 == 1) ? ev[0] : ev[1]);
-				//ZeroVertices.Add((n0 == 1) ? ev[0] : ev[1]);
+				// Note: if (!BOnPlane), then this will pick EdgeV.A, otherwise it will pick EdgeV.B
+				OnPlaneVertices->Add(EdgeV[BOnPlane]);
 			}
 			continue;
 		}
 
 		// no crossing
-		if (f0 * f1 > 0)
+		if (DistA * DistB > 0)
 		{
 			continue;
 		}
 
-		FDynamicMesh3::FEdgeSplitInfo splitInfo;
-		double t = f0 / (f0 - f1);
-		EMeshResult result = Mesh->SplitEdge(EID, splitInfo, t);
-		if (!ensureMsgf(result == EMeshResult::Ok, TEXT("FMeshPlaneCut::Cut: failed to SplitEdge")))
+		FDynamicMesh3::FEdgeSplitInfo SplitInfo;
+		double Param = DistA / (DistA - DistB);
+		EMeshResult SplitResult = Mesh->SplitEdge(EID, SplitInfo, Param);
+		if (!ensureMsgf(SplitResult == EMeshResult::Ok, TEXT("FMeshPlaneCut::Cut: failed to SplitEdge")))
 		{
 			continue; // edge split really shouldn't fail; skip the edge if it somehow does
 		}
 
-		OnSplitEdges.Add(EID);
-		NewEdges.Add(splitInfo.NewEdges.A); OnSplitEdges.Add(splitInfo.NewEdges.A);
-		NewEdges.Add(splitInfo.NewEdges.B); OnCutEdges.Add(splitInfo.NewEdges.B);
-		if (splitInfo.NewEdges.C != FDynamicMesh3::InvalidID)
+		if (TriangleSelection && TriangleSelection->Contains(SplitInfo.OriginalTriangles.A))
+		{
+			TriangleSelection->Add(SplitInfo.NewTriangles.A);
+		}
+
+		AddNewEdge(SplitInfo.NewEdges.A);
+		AddNewEdge(SplitInfo.NewEdges.B);
+
+		// If requested, track the edges we've split
+		if (OnSplitEdges)
+		{
+			OnSplitEdges->Add(EID);
+			OnSplitEdges->Add(SplitInfo.NewEdges.A);
+		}
+
+		// We need to check whether the other vertices are on plane to decide if the connected edges are on the plane or not
+		int32 OtherVIDA = SplitInfo.OtherVertices.A;
+		// Other vertex is on plane if it's newly created  or within plane tolerance
+		if (OtherVIDA >= Signs.Num() || FMath::Abs(Signs[OtherVIDA]) < PlaneTolerance)
 		{
-			NewEdges.Add(splitInfo.NewEdges.C); OnCutEdges.Add(splitInfo.NewEdges.C);
+			CutPlaneEdges.Add(SplitInfo.NewEdges.B);
+		}
+		
+		if (SplitInfo.NewEdges.C != FDynamicMesh3::InvalidID)
+		{
+			if (TriangleSelection && TriangleSelection->Contains(SplitInfo.OriginalTriangles.B))
+			{
+				TriangleSelection->Add(SplitInfo.NewTriangles.B);
+			}
+			AddNewEdge(SplitInfo.NewEdges.C);
+			int32 OtherVIDB = SplitInfo.OtherVertices.B;
+			if (OtherVIDB >= Signs.Num() || FMath::Abs(Signs[OtherVIDB]) < PlaneTolerance)
+			{
+				CutPlaneEdges.Add(SplitInfo.NewEdges.C);
+			}
 		}
 
-		OnCutVertices.Add(splitInfo.NewVertex);
+		if (OnPlaneVertices)
+		{
+			OnPlaneVertices->Add(SplitInfo.NewVertex);
+		}
 	}
 }
 
+void FMeshPlaneCut::SplitCrossingEdges(TArray<double>& Signs, TSet<int>& ZeroEdges, TSet<int>& OnCutEdges, TSet<int>& OnSplitEdges, bool bDeleteTrisOnPlane)
+{
+	PRAGMA_DISABLE_DEPRECATION_WARNINGS
+	SplitCrossingEdges(bDeleteTrisOnPlane, Signs, ZeroEdges, OnCutEdges, &OnSplitEdges, &OnCutVertices, nullptr);
+	PRAGMA_ENABLE_DEPRECATION_WARNINGS
+}
+
+void FMeshPlaneCut::SplitCrossingEdges(TArray<double>& Signs, TSet<int>& ZeroEdges, TSet<int>& OnCutEdges, bool bDeleteTrisOnPlane)
+{
+	PRAGMA_DISABLE_DEPRECATION_WARNINGS
+	SplitCrossingEdges(bDeleteTrisOnPlane, Signs, ZeroEdges, OnCutEdges, nullptr, &OnCutVertices, nullptr);
+	PRAGMA_ENABLE_DEPRECATION_WARNINGS
+}
+
+
 bool FMeshPlaneCut::CutWithoutDelete(bool bSplitVerticesAtPlane, float OffsetVertices, TDynamicMeshScalarTriangleAttribute<int>* TriLabels, int NewLabelStartID, bool bAddBoundariesFirstHalf, bool bAddBoundariesSecondHalf)
 {
 	TArray<double> Signs;
-	TSet<int> ZeroEdges, OnCutEdges;
-	SplitCrossingEdges(Signs, ZeroEdges, OnCutEdges, bSplitVerticesAtPlane /* only delete on-plane tris if we're also splitting vertices apart */);
+	TSet<int> AlreadyOnPlaneEdges, OnCutEdges;
+	SplitCrossingEdges(bSplitVerticesAtPlane, Signs, AlreadyOnPlaneEdges, OnCutEdges, nullptr, nullptr);
 
 	if (!bSplitVerticesAtPlane)
 	{
 		ensure(OffsetVertices == 0.0); // it would be weird to not split vertices and still request any offset of the 'other side' vertices; please don't do that
 	}
 
-	// collapse degenerate edges if we got em
+	// collapse degenerate edges
 	if (bCollapseDegenerateEdgesOnCut)
 	{
-		CollapseDegenerateEdges(OnCutEdges, ZeroEdges);
+		CollapseDegenerateEdges(OnCutEdges, false);
+	}
+
+	if (bSimplifyAlongNewEdges)
+	{
+		SimplifySettings.SimplifyAlongEdges(*Mesh, OnCutEdges);
 	}
 
 	if (!bSplitVerticesAtPlane)
@@ -191,8 +227,12 @@ bool FMeshPlaneCut::CutWithoutDelete(bool bSplitVerticesAtPlane, float OffsetVer
 		return true;
 	}
 
-	ensure(TriLabels); // need labels to split verts currently
+	if (!ensure(TriLabels))
+	{
+		return false; // need labels to split verts currently
+	}
 
+	// Update label IDs for triangles on the positive side of the plane, and shift vertices by the offset vector
 	TMap<int, int> OldLabelToNew;
 	int AvailableID = NewLabelStartID;
 	FVector3d VertexOffsetVec = (double)OffsetVertices * PlaneNormal;
@@ -225,7 +265,7 @@ bool FMeshPlaneCut::CutWithoutDelete(bool bSplitVerticesAtPlane, float OffsetVer
 	// split the mesh apart and add open boundary info
 	TMap<int, int> SplitVertices;
 	TSet<int> BoundaryVertices;
-	const TSet<int>* Sets[2] { &OnCutEdges, &ZeroEdges };
+	const TSet<int>* Sets[2] { &AlreadyOnPlaneEdges, &OnCutEdges };
 	for (int SetIdx = 0; SetIdx < 2; SetIdx++)
 	{
 		const TSet<int>& Set = *(Sets[SetIdx]);
@@ -261,8 +301,9 @@ bool FMeshPlaneCut::CutWithoutDelete(bool bSplitVerticesAtPlane, float OffsetVer
 				NonSplitTriCount++;
 			}
 		}
-		if (NonSplitTriCount > 0) // connected to both old and new labels -- needs split
+		if (NonSplitTriCount > 0) // connected to old labels
 		{
+			// if also connected to new labels, needs split
 			if (Triangles.Num() > 0 && EMeshResult::Ok == Mesh->SplitVertex(VID, Triangles, SplitInfo))
 			{
 				SplitVertices.Add(VID, SplitInfo.NewVertex);
@@ -326,7 +367,7 @@ bool FMeshPlaneCut::CutWithoutDelete(bool bSplitVerticesAtPlane, float OffsetVer
 			}
 		}
 	
-		TSet<int> UnusedZeroEdgesSet;
+		TSet<int> UnusedZeroEdgesSet; // This set is unused except to pass to ExtractBoundaryLoops, which expects two sets of edges
 		for (TPair<int, TSet<int>>& LabelIDEdges : LabelToCutEdges)
 		{
 			int LabelID = LabelIDEdges.Key;
@@ -346,7 +387,6 @@ bool FMeshPlaneCut::CutWithoutDelete(bool bSplitVerticesAtPlane, float OffsetVer
 				Boundary.NormalSign = -1;
 			}
 
-			check(UnusedZeroEdgesSet.Num() == 0); // for simplicity, we only put stuff in the CutEdges set
 			bool ExtractOk = ExtractBoundaryLoops(Edges, UnusedZeroEdgesSet, Boundary);
 			AllExtractionsOk = ExtractOk && AllExtractionsOk;
 		}
@@ -359,16 +399,12 @@ bool FMeshPlaneCut::Cut()
 {
 	TArray<double> Signs;
 	TSet<int> ZeroEdges, OnCutEdges;
-	SplitCrossingEdges(Signs, ZeroEdges, OnCutEdges);
+	SplitCrossingEdges(true, Signs, ZeroEdges, OnCutEdges, nullptr, nullptr);
 	
-	// @todo handle selection logic
-	//IEnumerable<int> vertexSet = Interval1i.Range(MaxVID);
-	//if (CutVertexSet != null)
-	//	vertexSet = CutVertexSet;
 	// remove one-rings of all positive-side vertices. 
-	for (int VID : Mesh->VertexIndicesItr())
+	for (int VID = 0; VID < Signs.Num(); ++VID)
 	{
-		if (VID < Signs.Num() && Signs[VID] > PlaneTolerance)
+		if (Signs[VID] > PlaneTolerance)
 		{
 			constexpr bool bPreserveManifold = false;
 			Mesh->RemoveVertex(VID, bPreserveManifold);
@@ -378,7 +414,12 @@ bool FMeshPlaneCut::Cut()
 	// collapse degenerate edges if we got em
 	if (bCollapseDegenerateEdgesOnCut)
 	{
-		CollapseDegenerateEdges(OnCutEdges, ZeroEdges);
+		CollapseDegenerateEdges(OnCutEdges, false);
+	}
+
+	if (bSimplifyAlongNewEdges)
+	{
+		SimplifySettings.SimplifyAlongEdges(*Mesh, OnCutEdges);
 	}
 
 	FMeshPlaneCut::FOpenBoundary& Boundary = OpenBoundaries.Emplace_GetRef();
@@ -386,41 +427,58 @@ bool FMeshPlaneCut::Cut()
 
 }
 
-bool FMeshPlaneCut::SplitEdgesOnly(bool bAssignNewGroups)
+bool FMeshPlaneCut::SplitEdgesOnlyHelper(bool bAssignNewGroups, TSet<int32>* TriangleSelection, bool bAddDeprecatedResultSeedTriangles)
 {
 	// split edges with current plane
 	TArray<double> Signs;
 	TSet<int32> ZeroEdges, OnCutEdges, OnSplitEdges;
-	SplitCrossingEdges(Signs, ZeroEdges, OnCutEdges, OnSplitEdges, false);
+	SplitCrossingEdges(false, Signs, ZeroEdges, OnCutEdges, &OnSplitEdges, nullptr, TriangleSelection);
 
 	if (bAssignNewGroups == false)
 	{
 		return true;
 	}
 
-	// find relevant edges/triangles/groups on cut
-	TSet<int32> CutEdges;				// edges lying on the cut
-	TArray<int32> CutEdgeTriangles;		// triangles connected to those edges
-	TSet<int32> CutEdgeGroups;			// group IDs of those triangles (ie groups touching cut)
-	TSet<int32>* EdgeLists[2] = { &ZeroEdges, &OnCutEdges };
-	for (TSet<int32>* EdgeList : EdgeLists)
+	if (bSimplifyAlongNewEdges)
 	{
-		for (int32 eid : *EdgeList)
+		if (ensureMsgf(!bAddDeprecatedResultSeedTriangles, TEXT("Deprecated SplitEdgesOnly without TriangleSelection parameter does not support simplification option")))
 		{
-			FIndex2i EdgeVerts = Mesh->GetEdgeV(eid);
-			if (OnCutVertices.Contains(EdgeVerts.A) && OnCutVertices.Contains(EdgeVerts.B) && CutEdges.Contains(eid) == false )
+			// TODO: Consider making the simplification also preserve the TriangleSelection boundary
+			CollapseDegenerateEdges(OnCutEdges, false, TriangleSelection);
+			if (TriangleSelection)
 			{
-				CutEdges.Add(eid);
-				FIndex2i EdgeTris = Mesh->GetEdgeT(eid);
-				for (int32 j = 0; j < 2; ++j)
+				SimplifySettings.SimplifyAlongEdges(*Mesh, OnCutEdges, [&TriangleSelection](const DynamicMeshInfo::FEdgeCollapseInfo& CollapseInfo)
 				{
-					if (EdgeTris[j] != FDynamicMesh3::InvalidID)
+					TriangleSelection->Remove(CollapseInfo.RemovedTris.A);
+					if (CollapseInfo.RemovedTris.B != FDynamicMesh3::InvalidID)
 					{
-						CutEdgeTriangles.Add(EdgeTris[j]);
-						int32 Group = Mesh->GetTriangleGroup(EdgeTris[j]);
-						CutEdgeGroups.Add(Group);
+						TriangleSelection->Remove(CollapseInfo.RemovedTris.B);
 					}
-				}
+				});
+			}
+			else
+			{
+				SimplifySettings.SimplifyAlongEdges(*Mesh, OnCutEdges);
+			}
+		}
+	}
+
+	TSet<int32> OnPlaneEdges = ZeroEdges;
+	OnPlaneEdges.Append(OnCutEdges);
+
+	// Use the edges along the cut (both pre-existing and newly added) to assign group IDs across both sides of the triangulation
+	TArray<int32> CutEdgeTriangles;		// triangles connected to those edges
+	TSet<int32> CutEdgeGroups;			// group IDs of those triangles (ie groups touching cut)
+	for (int32 EID : OnPlaneEdges)
+	{
+		FIndex2i EdgeTris = Mesh->GetEdgeT(EID);
+		for (int32 j = 0; j < 2; ++j)
+		{
+			if (EdgeTris[j] != FDynamicMesh3::InvalidID)
+			{
+				CutEdgeTriangles.Add(EdgeTris[j]);
+				int32 Group = Mesh->GetTriangleGroup(EdgeTris[j]);
+				CutEdgeGroups.Add(Group);
 			}
 		}
 	}
@@ -433,7 +491,7 @@ bool FMeshPlaneCut::SplitEdgesOnly(bool bAssignNewGroups)
 		if (Group0 == Group1)
 		{
 			int32 SharedEdge = Mesh->FindEdgeFromTriPair(t0, t1);
-			if (CutEdges.Contains(SharedEdge) == false)
+			if (OnPlaneEdges.Contains(SharedEdge) == false)
 			{
 				return true;
 			}
@@ -441,7 +499,7 @@ bool FMeshPlaneCut::SplitEdgesOnly(bool bAssignNewGroups)
 		return false;
 	});
 
-	// Assign a new group id for each component
+	// Assign a new group ID for each component
 	// Do we want to keep existing groups? possibly cleaner to assign new ones because one input group may
 	// be split into multiple child groups on each side of the cut. 
 	// But perhaps should track group-mapping?
@@ -460,6 +518,13 @@ bool FMeshPlaneCut::SplitEdgesOnly(bool bAssignNewGroups)
 		Result.Triangles = MoveTemp(Component.Indices);
 	}
 
+	if (!bAddDeprecatedResultSeedTriangles)
+	{
+		return true;
+	}
+
+	// Note: the below logic will be removed when ResultSeedTriangles is removed
+PRAGMA_DISABLE_DEPRECATION_WARNINGS
 	// Compute the set of triangle IDs in the cut mesh that represent
 	// the original/seed triangle selection along the cut. This assumes
 	// that the edge filter function contains edges that originate
@@ -491,7 +556,7 @@ bool FMeshPlaneCut::SplitEdgesOnly(bool bAssignNewGroups)
 			int eid = TriEdges[j];
 			FIndex2i ev = Mesh->GetEdgeV(eid);
 			bool bIsSeedEdge = (EdgeFilterFunc && EdgeFilterFunc(eid));
-			if (bIsSeedEdge || CutEdges.Contains(eid) || OnSplitEdges.Contains(eid))
+			if (bIsSeedEdge || OnCutEdges.Contains(eid) || OnSplitEdges.Contains(eid))
 			{
 				continue;
 			}
@@ -500,6 +565,7 @@ bool FMeshPlaneCut::SplitEdgesOnly(bool bAssignNewGroups)
 			ResultSeedTriangles.Add(EdgeTris.B);
 		}
 	}
+PRAGMA_ENABLE_DEPRECATION_WARNINGS
 	return true;
 }
 
@@ -532,6 +598,7 @@ bool FMeshPlaneCut::ExtractBoundaryLoops(const TSet<int>& OnCutEdges, const TSet
 	return !Boundary.CutLoopsFailed;
 }
 
+PRAGMA_DISABLE_DEPRECATION_WARNINGS
 void FMeshPlaneCut::CollapseDegenerateEdges(const TSet<int>& OnCutEdges, const TSet<int>& ZeroEdges)
 {
 	const TSet<int>* Sets[2] { &OnCutEdges, &ZeroEdges };
@@ -578,7 +645,73 @@ void FMeshPlaneCut::CollapseDegenerateEdges(const TSet<int>& OnCutEdges, const T
 		}
 	} while (Collapsed != 0);
 }
+PRAGMA_ENABLE_DEPRECATION_WARNINGS
+
+void FMeshPlaneCut::CollapseDegenerateEdges(TSet<int32>& Edges, bool bRemoveAllDegenerateFromInputSet, TSet<int>* TriangleSelection)
+{
+	double Tol2 = DegenerateEdgeTol * DegenerateEdgeTol;
+	FVector3d A, B;
+	TArray<int32> ToRemove;
+	int Collapsed = 0;
+	do
+	{
+		Collapsed = 0;
+		for (auto EdgesIter = Edges.CreateIterator(); EdgesIter; ++EdgesIter)
+		{
+			int32 EID = *EdgesIter;
+			if (!Mesh->IsEdge(EID))
+			{
+				continue;
+			}
+			Mesh->GetEdgeV(EID, A, B);
+			if (DistanceSquared(A, B) > Tol2)
+			{
+				continue;
+			}
+
+			FIndex2i EV = Mesh->GetEdgeV(EID);
+			// if the vertex we'd remove is on a seam, try removing the other one instead
+			if (Mesh->HasAttributes() && Mesh->Attributes()->IsSeamVertex(EV.B, false))
+			{
+				Swap(EV.A, EV.B);
+				// if they were both on seams, then collapse should not happen?  (& would break OnCollapseEdge assumptions in overlay)
+				if (Mesh->HasAttributes() && Mesh->Attributes()->IsSeamVertex(EV.B, false))
+				{
+					continue;
+				}
+			}
+			FDynamicMesh3::FEdgeCollapseInfo CollapseInfo;
+			EMeshResult Result = Mesh->CollapseEdge(EV.A, EV.B, CollapseInfo);
+			if (Result == EMeshResult::Ok)
+			{
+				Collapsed++;
+				EdgesIter.RemoveCurrent();
+				if (TriangleSelection)
+				{
+					TriangleSelection->Remove(CollapseInfo.RemovedEdges.A);
+					if (CollapseInfo.RemovedEdges.B != FDynamicMesh3::InvalidID)
+					{
+						TriangleSelection->Remove(CollapseInfo.RemovedEdges.B);
+					}
+				}
+				// Add the other edge for later removal, so we don't potentially invalidate our iterator
+				// (but only if we really need to remove invalid edges here; some algorithms may be ok with leaving invalid EIDs in the set)
+				if (bRemoveAllDegenerateFromInputSet && CollapseInfo.RemovedEdges.B != FDynamicMesh3::InvalidID)
+				{
+					ToRemove.Add(CollapseInfo.RemovedEdges.B);
+				}
+			}
+		}
+	} while (Collapsed != 0);
 
+	if (bRemoveAllDegenerateFromInputSet)
+	{
+		for (int32 EID : ToRemove)
+		{
+			Edges.Remove(EID);
+		}
+	}
+}
 
 bool FMeshPlaneCut::SimpleHoleFill(int ConstantGroupID)
 {
diff --git a/Engine/Plugins/Runtime/GeometryProcessing/Source/DynamicMesh/Private/Operations/MeshSelfUnion.cpp b/Engine/Plugins/Runtime/GeometryProcessing/Source/DynamicMesh/Private/Operations/MeshSelfUnion.cpp
--- a/Engine/Plugins/Runtime/GeometryProcessing/Source/DynamicMesh/Private/Operations/MeshSelfUnion.cpp
+++ b/Engine/Plugins/Runtime/GeometryProcessing/Source/DynamicMesh/Private/Operations/MeshSelfUnion.cpp
@@ -4,6 +4,7 @@
 
 #include "Operations/MeshSelfUnion.h"
 #include "Operations/MeshMeshCut.h"
+#include "Operations/LocalPlanarSimplify.h"
 
 #include "Selections/MeshConnectedComponents.h"
 
@@ -550,8 +551,8 @@ void FMeshSelfUnion::SimplifyAlongNewEdges(TArray<int>& CutBoundaryEdges, TMap<i
 			int NumFlat = 0;
 			for (int VIdx = 0; VIdx < 2; VIdx++)
 			{
-				Flat[VIdx] = FMeshBoolean::IsFlat(*Mesh, Edge.Vert[VIdx], DotTolerance, FlatNormals[VIdx][0]) 
-						  && FMeshBoolean::IsFlat(*Mesh, Matches[VIdx], DotTolerance, FlatNormals[VIdx][1]);
+				Flat[VIdx] = FLocalPlanarSimplify::IsFlat(*Mesh, Edge.Vert[VIdx], DotTolerance, FlatNormals[VIdx][0]) 
+						  && FLocalPlanarSimplify::IsFlat(*Mesh, Matches[VIdx], DotTolerance, FlatNormals[VIdx][1]);
 
 				if (Flat[VIdx])
 				{
@@ -591,10 +592,10 @@ void FMeshSelfUnion::SimplifyAlongNewEdges(TArray<int>& CutBoundaryEdges, TMap<i
 					int KeepV = WhichEdge == 0 ? Edge.Vert[KeepVIdx] : Matches[KeepVIdx];
 					int SourceEID = WhichEdge == 0 ? EID : MatchEID;
 
-					bHasBadEdge = bHasBadEdge || FMeshBoolean::CollapseWouldHurtTriangleQuality(
+					bHasBadEdge = bHasBadEdge || FLocalPlanarSimplify::CollapseWouldHurtTriangleQuality(
 						*Mesh, FlatNormals[RemoveVIdx][WhichEdge], RemoveV, RemoveVPos, KeepV, KeepVPos, TryToImproveTriQualityThreshold);
 
-					bHasBadEdge = bHasBadEdge || FMeshBoolean::CollapseWouldChangeShapeOrUVs(
+					bHasBadEdge = bHasBadEdge || FLocalPlanarSimplify::CollapseWouldChangeShapeOrUVs(
 						*Mesh, CutBoundaryEdgeSet, DotTolerance,
 						SourceEID, RemoveV, RemoveVPos, KeepV, KeepVPos, EdgeDir, bPreserveTriangleGroups,
 						true, bPreserveVertexUVs, bPreserveOverlayUVs, UVDistortTolerance * UVDistortTolerance,
diff --git a/Engine/Plugins/Runtime/GeometryProcessing/Source/DynamicMesh/Public/Operations/GroupEdgeInserter.h b/Engine/Plugins/Runtime/GeometryProcessing/Source/DynamicMesh/Public/Operations/GroupEdgeInserter.h
--- a/Engine/Plugins/Runtime/GeometryProcessing/Source/DynamicMesh/Public/Operations/GroupEdgeInserter.h
+++ b/Engine/Plugins/Runtime/GeometryProcessing/Source/DynamicMesh/Public/Operations/GroupEdgeInserter.h
@@ -58,6 +58,13 @@ public:
 		 */
 		double VertexTolerance = KINDA_SMALL_NUMBER * 10;
 
+		/**
+		 * When inserting edges, we can locally simplify to improve mesh quality along the edge path.
+		 * Only does so when it should not change mesh shape, UVs, triangle groups, etc,
+		 * and only applies when Insertion Mode is "Plane Cut"
+		 */
+		bool bSimplifyAlongPath = true;
+
 		/** 
 		 * Determines how the edge is inserted: by using a cutting plane to cut existing triangles
 		 * along the path, or deleting triangles and retriangulating. Retriangulation will keep
@@ -138,6 +145,13 @@ public:
 		 */
 		double VertexTolerance = KINDA_SMALL_NUMBER * 10;
 
+		/**
+		 * When inserting edges, we can locally simplify to improve mesh quality along the edge path.
+		 * Only does so when it should not change mesh shape, UVs, triangle groups, etc,
+		 * and only applies when Insertion Mode is "Plane Cut"
+		 */
+		bool bSimplifyAlongPath = true;
+
 		EInsertionMode Mode = EInsertionMode::Retriangulate;
 	};
 
diff --git a/Engine/Plugins/Runtime/GeometryProcessing/Source/DynamicMesh/Public/Operations/LocalPlanarSimplify.h b/Engine/Plugins/Runtime/GeometryProcessing/Source/DynamicMesh/Public/Operations/LocalPlanarSimplify.h
new file mode 100644
--- /dev/null
+++ b/Engine/Plugins/Runtime/GeometryProcessing/Source/DynamicMesh/Public/Operations/LocalPlanarSimplify.h
@@ -0,0 +1,100 @@
+// Copyright Epic Games, Inc. All Rights Reserved.
+
+// Port of geometry3Sharp MeshPlaneCut
+
+#pragma once
+
+#include "MathUtil.h"
+#include "VectorTypes.h"
+#include "GeometryTypes.h"
+#include "DynamicMesh/InfoTypes.h"
+
+
+namespace UE
+{
+namespace Geometry
+{
+
+class FDynamicMesh3;
+
+
+//
+// Class to support local planar simplification, to reduce excess edges created by plane cuts, edge loop insertions, and similar operations
+// (Note a custom version of this logic also exists on the FMeshBoolean class)
+//
+class DYNAMICMESH_API FLocalPlanarSimplify
+{
+public:
+	// Simplification settings
+
+	/** Degrees of deviation from coplanar that we will still simplify */
+	double SimplificationAngleTolerance = .1;
+	/**
+	 * If triangle quality (aspect ratio) is worse than this threshold, only simplify in ways that improve quality.  If <= 0, triangle quality is ignored.
+	 *  Note: For aspect ratio we use definition: 4*TriArea / (sqrt(3)*MaxEdgeLen^2), ref: https://people.eecs.berkeley.edu/~jrs/papers/elemj.pdf p.53
+	 *  Equilateral triangles have value 1; Smaller values -> lower quality
+	 */
+	double TryToImproveTriQualityThreshold = .25;
+	/** Prevent simplification from distorting triangle groups */
+	bool bPreserveTriangleGroups = true;
+	/** Prevent simplification from distorting vertex UVs */
+	bool bPreserveVertexUVs = true;
+	/** Prevent simplification from distorting overlay UVs */
+	bool bPreserveOverlayUVs = true;
+	/** When preserving UVs, sets maximum allowed change in UV coordinates from collapsing an edge, measured at the removed vertex */
+	float UVDistortTolerance = FMathf::ZeroTolerance;
+	/** Prevent simplification from distorting vertex normals */
+	bool bPreserveVertexNormals = true;
+	/** When preserving normals, sets maximum allowed change in normals from collapsing an edge, measured at the removed vertex in degrees */
+	float NormalDistortTolerance = .01f;
+
+
+	/**
+	 * Simplify Mesh along the given edges, and update the edge set accordingly
+	 * @param Mesh				Mesh to simplify
+	 * @param InOutEdges		Edges to simplify along -- will be updated to remove any edges that have been collapsed by simplification
+	 * @param ProcessCollapse	Optional function to be called whenever the simplification collapses an edge
+	 */
+	void SimplifyAlongEdges(FDynamicMesh3& Mesh, TSet<int32>& InOutEdges, TUniqueFunction<void(const DynamicMeshInfo::FEdgeCollapseInfo&)> ProcessCollapse = nullptr) const;
+
+
+
+
+	//
+	// The below helper functions are used by SimplifyAlongEdges, but also the similar customized versions of the same algorithm in FMeshBoolean and FMeshSelfUnion as well
+	//
+
+	/**
+	 * Test if the triangles connected to a vertex are all coplanar
+	 * @param Mesh The mesh to query
+	 * @param VID The vertex to query
+	 * @param DotTolerance If the dot product of two normals are >= this tolerance, the normals are considered equivalent
+	 * @param The normal of the first triangle attached to the vertex.
+	 * @return Whether all the triangles were coplanar
+	 */
+	static bool IsFlat(const FDynamicMesh3& Mesh, int VID, double DotTolerance, FVector3d& OutFirstNormal);
+
+	/**
+	 * Test if a given edge collapse would cause a triangle flip or other unacceptable decrease in mesh quality
+	 * Specialized for collapsing at flat triangles
+	 */
+	static bool CollapseWouldHurtTriangleQuality(
+		const FDynamicMesh3& Mesh, const FVector3d& ExpectNormal,
+		int32 RemoveV, const FVector3d& RemoveVPos, int32 KeepV, const FVector3d& KeepVPos,
+		double TryToImproveTriQualityThreshold
+	);
+
+	/**
+	 * Test if a given edge collapse would change the mesh shape, mesh triangle group shape, or UVs unacceptably
+	 */
+	static bool CollapseWouldChangeShapeOrUVs(
+		const FDynamicMesh3& Mesh, const TSet<int>& CutBoundaryEdgeSet, double DotTolerance, int SourceEID,
+		int32 RemoveV, const FVector3d& RemoveVPos, int32 KeepV, const FVector3d& KeepVPos,
+		const FVector3d& EdgeDir, bool bPreserveTriangleGroups, bool bPreserveUVsForMesh,
+		bool bPreserveVertexUVs, bool bPreserveOverlayUVs, float UVToleranceSq,
+		bool bPreserveVertexNormals, float NormalEqualCosThreshold);
+
+};
+
+} // end namespace UE::Geometry
+} // end namespace UE
\ No newline at end of file
diff --git a/Engine/Plugins/Runtime/GeometryProcessing/Source/DynamicMesh/Public/Operations/MeshBoolean.h b/Engine/Plugins/Runtime/GeometryProcessing/Source/DynamicMesh/Public/Operations/MeshBoolean.h
--- a/Engine/Plugins/Runtime/GeometryProcessing/Source/DynamicMesh/Public/Operations/MeshBoolean.h
+++ b/Engine/Plugins/Runtime/GeometryProcessing/Source/DynamicMesh/Public/Operations/MeshBoolean.h
@@ -195,40 +195,6 @@ private:
 
 	void SimplifyAlongNewEdges(int NumMeshesToProcess, FDynamicMesh3* CutMesh[2], TArray<int> CutBoundaryEdges[2], TMap<int, int>& AllVIDMatches);
 
-public:
-	
-	// specialized helper functions useful for implementing mesh boolean-like features (shared with the extremely similar FMeshSelfUnion algorithm)
-
-	/**
-	 * Test if the triangles connected to a vertex are all coplanar
-	 * @param Mesh The mesh to query
-	 * @param VID The vertex to query
-	 * @param DotTolerance If the dot product of two normals are >= this tolerance, the normals are considered equivalent
-	 * @param The normal of the first triangle attached to the vertex.
-	 * @return Whether all the triangles were coplanar
-	 */
-	static bool IsFlat(const FDynamicMesh3& Mesh, int VID, double DotTolerance, FVector3d& OutFirstNormal);
-
-	/**
-	 * Test if a given edge collapse would cause a triangle flip or other unacceptable decrease in mesh quality
-	 * Specialized for collapsing at flat triangles
-	 * TODO: Currently this only detects triangle flips; need to extend it to also detect other mesh quality issues
-	 */
-	static bool CollapseWouldHurtTriangleQuality(
-		const FDynamicMesh3& Mesh, const FVector3d& ExpectNormal,
-		int32 RemoveV, const FVector3d& RemoveVPos, int32 KeepV, const FVector3d& KeepVPos,
-		double TryToImproveTriQualityThreshold
-	);
-
-	/**
-	 * Test if a given edge collapse would change the mesh shape, mesh triangle group shape, or UVs unacceptably
-	 */
-	static bool CollapseWouldChangeShapeOrUVs(
-		const FDynamicMesh3& Mesh, const TSet<int>& CutBoundaryEdgeSet, double DotTolerance, int SourceEID,
-		int32 RemoveV, const FVector3d& RemoveVPos, int32 KeepV, const FVector3d& KeepVPos,
-		const FVector3d& EdgeDir, bool bPreserveTriangleGroups, bool bPreserveUVsForMesh,
-		bool bPreserveVertexUVs, bool bPreserveOverlayUVs, float UVToleranceSq,
-		bool bPreserveVertexNormals, float NormalEqualCosThreshold);
 };
 
 
diff --git a/Engine/Plugins/Runtime/GeometryProcessing/Source/DynamicMesh/Public/Operations/MeshPlaneCut.h b/Engine/Plugins/Runtime/GeometryProcessing/Source/DynamicMesh/Public/Operations/MeshPlaneCut.h
--- a/Engine/Plugins/Runtime/GeometryProcessing/Source/DynamicMesh/Public/Operations/MeshPlaneCut.h
+++ b/Engine/Plugins/Runtime/GeometryProcessing/Source/DynamicMesh/Public/Operations/MeshPlaneCut.h
@@ -10,6 +10,7 @@
 #include "MeshBoundaryLoops.h"
 #include "Curve/GeneralPolygon2.h"
 
+#include "Operations/LocalPlanarSimplify.h"
 
 
 namespace UE
@@ -62,6 +63,13 @@ public:
 	double PlaneTolerance = FMathf::ZeroTolerance * 10.0;
 
 
+	/** Control whether we attempt to auto-simplify the small planar triangles that the plane cut operation tends to generate */
+	bool bSimplifyAlongNewEdges = false;
+
+	/** Settings to apply if bSimplifyAlongNewEdges == true */
+	FLocalPlanarSimplify SimplifySettings;
+
+
 	// TODO support optionally restricting plane cut to a mesh selection
 	//MeshFaceSelection CutFaceSet;
 
@@ -91,8 +99,13 @@ public:
 	};
 	/** List of output cut regions (eg that have separate GroupIDs). Currently only calculated by SplitEdgesOnly() path */
 	TArray<FCutResultRegion> ResultRegions;
+
+
+	PRAGMA_DISABLE_DEPRECATION_WARNINGS
 	/** List of output cut triangles representing the seed triangles along the cut. Currently only calculated by SplitEdgesOnly() path */
+	UE_DEPRECATED(5.3, "To preserve a triangle selection when using SplitEdgesOnly(), instead pass the selection as a pointer to that function.")
 	TArray<int32> ResultSeedTriangles;
+	PRAGMA_ENABLE_DEPRECATION_WARNINGS
 
 public:
 
@@ -102,7 +115,11 @@ public:
 	FMeshPlaneCut(FDynamicMesh3* Mesh, FVector3d Origin, FVector3d Normal) : Mesh(Mesh), PlaneOrigin(Origin), PlaneNormal(Normal)
 	{
 	}
+
+	// Note: deprecation warnings disabled on destructor due to deprecated member variables
+	PRAGMA_DISABLE_DEPRECATION_WARNINGS
 	virtual ~FMeshPlaneCut() {}
+	PRAGMA_ENABLE_DEPRECATION_WARNINGS
 	
 	/**
 	 * @return EOperationValidationResult::Ok if we can apply operation, or error code if we cannot
@@ -137,9 +154,19 @@ public:
 	/**
 	 * Compute the plane cut by splitting mesh edges that cross the cut plane, and then optionally update groups
 	 * @param bAssignNewGroups if true, update group IDs such that each group-connected-component on either side of the cut plane is assigned a new unique group ID
+	 * @param OptionalTriangleSelection if non-null, a set of triangle IDs that must be updated along with any triangle splits or deletions in the edge split process
 	 * @return true if operation succeeds
 	 */
-	virtual bool SplitEdgesOnly(bool bAssignNewGroups);
+	virtual bool SplitEdgesOnly(bool bAssignNewGroups, TSet<int32>* OptionalTriangleSelection)
+	{
+		return SplitEdgesOnlyHelper(bAssignNewGroups, OptionalTriangleSelection, false);
+	}
+
+	UE_DEPRECATED(5.3, "Instead use the two-parameter version of SplitEdgesOnly, which has direction selection tracking and does not populate the deprecated ResultSeedTriangles")
+	virtual bool SplitEdgesOnly(bool bAssignNewGroups)
+	{
+		return SplitEdgesOnlyHelper(bAssignNewGroups, nullptr, true);
+	}
 
 	/**
 	 *  Fill cut loops with FSimpleHoleFiller
@@ -158,13 +185,33 @@ public:
 
 protected:
 
+	UE_DEPRECATED(5.3, "Use the single-set CollapseDegenerateEdges instead")
 	void CollapseDegenerateEdges(const TSet<int>& OnCutEdges, const TSet<int>& ZeroEdges);
+	
+	/// Collapse degenerate edges
+	/// @param Edges									Edges to consider for collapse; will updated by removing edges at they are collapsed
+	/// @param bRemoveAllDegenerateFromInputSet			Whether we should also check whether the neighbor edges removed by collapse were also in the set
+	/// 												(typically this is not needed, and has additional cost)
+	/// @param TriangleSelection						Optional set tracking an active selection. Any triangles removed by collapse will also be removed from the set.
+	void CollapseDegenerateEdges(TSet<int>& Edges, bool bRemoveAllDegenerateFromInputSet, TSet<int>* TriangleSelection = nullptr);
+	
 	void SplitCrossingEdges(TArray<double>& Signs, TSet<int>& ZeroEdges, TSet<int>& OnCutEdges, bool bDeleteTrisOnPlane = true);
 	void SplitCrossingEdges(TArray<double>& Signs, TSet<int>& ZeroEdges, TSet<int>& OnCutEdges, TSet<int>& OnSplitEdges, bool bDeleteTrisOnPlane = true);
+	void SplitCrossingEdges(bool bDeleteTrisOnPlane, TArray<double>& Signs, TSet<int>& AlreadyOnPlaneEdges, TSet<int32>& CutPlaneEdges, 
+		TSet<int>* SplitEdges = nullptr, TSet<int>* OnPlaneVertices = nullptr, TSet<int>* TriangleSelection = nullptr);
+	
 	bool ExtractBoundaryLoops(const TSet<int>& OnCutEdges, const TSet<int>& ZeroEdges, FMeshPlaneCut::FOpenBoundary& Boundary);
 
 	// set of vertices lying on plane after calling SplitCrossingEdges
+	UE_DEPRECATED(5.3, "If needed, explicitly request the vertices on the cut plane the the OnPlaneVertices argument of SplitCrossingEdges")
 	TSet<int32> OnCutVertices;
+
+private:
+
+	// Helper to compute signed distances from the cutting plane for all vertices of the mesh.  Value at Invalid Vertex IDs will be set to InvalidDist.
+	void ComputeVertexSignedDistances(TArray<double>& Signs, double InvalidDist);
+
+	bool SplitEdgesOnlyHelper(bool bAssignNewGroups, TSet<int32>* OptionalTriangleSelection, bool bAddDeprecatedResultSeedTriangles);
 };
 
 
diff --git a/Engine/Plugins/Runtime/MeshModelingToolset/Source/MeshModelingTools/Private/ToolActivities/PolyEditCutFacesActivity.cpp b/Engine/Plugins/Runtime/MeshModelingToolset/Source/MeshModelingTools/Private/ToolActivities/PolyEditCutFacesActivity.cpp
--- a/Engine/Plugins/Runtime/MeshModelingToolset/Source/MeshModelingTools/Private/ToolActivities/PolyEditCutFacesActivity.cpp
+++ b/Engine/Plugins/Runtime/MeshModelingToolset/Source/MeshModelingTools/Private/ToolActivities/PolyEditCutFacesActivity.cpp
@@ -166,10 +166,17 @@ void UPolyEditCutFacesActivity::ApplyCutFaces()
 	// apply the cut to edges of selected triangles
 	FGroupTopologySelection OutputSelection;
 	FMeshPlaneCut Cut(ActivityContext->CurrentMesh.Get(), PlaneOrigin, PlaneNormal);
+	Cut.bSimplifyAlongNewEdges = !ActivityContext->bTriangleMode;
+	Cut.SimplifySettings.bPreserveVertexNormals = false;
 	FMeshEdgeSelection Edges(ActivityContext->CurrentMesh.Get());
 	Edges.SelectTriangleEdges(ActiveTriangleSelection);
 	Cut.EdgeFilterFunc = [&](int EdgeID) { return Edges.IsSelected(EdgeID); };
-	if (Cut.SplitEdgesOnly(true))
+	TSet<int32> TriangleSelectionSet;
+	if (ActivityContext->bTriangleMode)
+	{
+		TriangleSelectionSet.Append(ActiveTriangleSelection);
+	}
+	if (Cut.SplitEdgesOnly(true, ActivityContext->bTriangleMode ? &TriangleSelectionSet : nullptr))
 	{
 		if (!ActivityContext->bTriangleMode)
 		{
@@ -180,11 +187,7 @@ void UPolyEditCutFacesActivity::ApplyCutFaces()
 		}
 		else
 		{
-			// Retain the selection along the cut. ResultSeedTriangles does not
-			// contain selected tris that are not cut, so re-add the original selected
-			// tris.
-			OutputSelection.SelectedGroupIDs.Append(Cut.ResultSeedTriangles);
-			OutputSelection.SelectedGroupIDs.Append(ActiveTriangleSelection);
+			OutputSelection.SelectedGroupIDs.Append(TriangleSelectionSet);
 		}
 	}
 
