BakeVtx: fix per-channel baking of curvature/ambient occlusion

diff --git a/Engine/Plugins/Runtime/GeometryProcessing/Source/DynamicMesh/Private/Sampling/MeshConstantMapEvaluator.cpp b/Engine/Plugins/Runtime/GeometryProcessing/Source/DynamicMesh/Private/Sampling/MeshConstantMapEvaluator.cpp
--- a/Engine/Plugins/Runtime/GeometryProcessing/Source/DynamicMesh/Private/Sampling/MeshConstantMapEvaluator.cpp
+++ b/Engine/Plugins/Runtime/GeometryProcessing/Source/DynamicMesh/Private/Sampling/MeshConstantMapEvaluator.cpp
@@ -16,6 +16,7 @@ void FMeshConstantMapEvaluator::Setup(const FMeshBaseBaker& Baker, FEvaluationCo
 	Context.Evaluate = &EvaluateSample;
 	Context.EvaluateDefault = &EvaluateDefault;
 	Context.EvaluateColor = &EvaluateColor;
+	Context.EvaluateChannel = &EvaluateChannel;
 	Context.EvalData = this;
 	Context.AccumulateMode = EAccumulateMode::Overwrite;
 	Context.DataLayout = DataLayout();
@@ -45,4 +46,9 @@ void FMeshConstantMapEvaluator::EvaluateColor(const int DataIdx, float*& In, FVe
 	In += 1;
 }
 
+void FMeshConstantMapEvaluator::EvaluateChannel(const int DataIdx, float*& In, float& Out, void* EvalData)
+{
+	Out = *In;
+	In += 1;
+}
 
diff --git a/Engine/Plugins/Runtime/GeometryProcessing/Source/DynamicMesh/Private/Sampling/MeshCurvatureMapEvaluator.cpp b/Engine/Plugins/Runtime/GeometryProcessing/Source/DynamicMesh/Private/Sampling/MeshCurvatureMapEvaluator.cpp
--- a/Engine/Plugins/Runtime/GeometryProcessing/Source/DynamicMesh/Private/Sampling/MeshCurvatureMapEvaluator.cpp
+++ b/Engine/Plugins/Runtime/GeometryProcessing/Source/DynamicMesh/Private/Sampling/MeshCurvatureMapEvaluator.cpp
@@ -25,6 +25,7 @@ void FMeshCurvatureMapEvaluator::Setup(const FMeshBaseBaker& Baker, FEvaluationC
 	Context.Evaluate = &EvaluateSample;
 	Context.EvaluateDefault = &EvaluateDefault;
 	Context.EvaluateColor = &EvaluateColor;
+	Context.EvaluateChannel = &EvaluateChannel;
 	Context.EvalData = this;
 	Context.AccumulateMode = EAccumulateMode::Add;
 	Context.DataLayout = DataLayout();
@@ -104,6 +105,22 @@ void FMeshCurvatureMapEvaluator::EvaluateColor(const int DataIdx, float*& In, FV
 	In += 1;
 }
 
+
+void FMeshCurvatureMapEvaluator::EvaluateChannel(const int DataIdx, float*& In, float& Out, void* EvalData)
+{
+	const FMeshCurvatureMapEvaluator* Eval = static_cast<FMeshCurvatureMapEvaluator*>(EvalData);
+	const float Curvature = In[0];
+	const float Sign = FMathf::Sign(Curvature);
+	const float T = (float)Eval->ClampRange.GetT(FMathf::Abs(Curvature));
+	
+	constexpr float ZeroValue = 0.5f;
+	constexpr float NegativeValue = 0.0f;
+	constexpr float PositiveValue = 1.0f;
+	Out = (Sign < 0) ? FMath::Lerp(ZeroValue, NegativeValue, T) : FMath::Lerp(ZeroValue, PositiveValue, T);
+
+	In += 1;
+}
+
 double FMeshCurvatureMapEvaluator::SampleFunction(const FCorrespondenceSample& SampleData) const
 {
 	const void* DetailMesh = SampleData.DetailMesh;
diff --git a/Engine/Plugins/Runtime/GeometryProcessing/Source/DynamicMesh/Private/Sampling/MeshNormalMapEvaluator.cpp b/Engine/Plugins/Runtime/GeometryProcessing/Source/DynamicMesh/Private/Sampling/MeshNormalMapEvaluator.cpp
--- a/Engine/Plugins/Runtime/GeometryProcessing/Source/DynamicMesh/Private/Sampling/MeshNormalMapEvaluator.cpp
+++ b/Engine/Plugins/Runtime/GeometryProcessing/Source/DynamicMesh/Private/Sampling/MeshNormalMapEvaluator.cpp
@@ -34,6 +34,7 @@ void FMeshNormalMapEvaluator::Setup(const FMeshBaseBaker& Baker, FEvaluationCont
 	Context.Evaluate = bHasDetailNormalTextures ? &EvaluateSample<true> : &EvaluateSample<false>;
 	Context.EvaluateDefault = &EvaluateDefault;
 	Context.EvaluateColor = &EvaluateColor;
+	Context.EvaluateChannel = &EvaluateChannel;
 	Context.EvalData = this;
 	Context.AccumulateMode = EAccumulateMode::Add;
 	Context.DataLayout = DataLayout();
@@ -68,6 +69,14 @@ void FMeshNormalMapEvaluator::EvaluateColor(const int DataIdx, float*& In, FVect
 	In += 3;
 }
 
+void FMeshNormalMapEvaluator::EvaluateChannel(const int DataIdx, float*& In, float& Out, void* EvalData)
+{
+	// Per-channel should not be an option for NormalMap evaluator
+	ensure(false);
+	Out = 0.0f;
+	In += 1;
+}
+
 template <bool bUseDetailNormalMap>
 FVector3f FMeshNormalMapEvaluator::SampleFunction(const FCorrespondenceSample& SampleData) const
 {
diff --git a/Engine/Plugins/Runtime/GeometryProcessing/Source/DynamicMesh/Private/Sampling/MeshOcclusionMapEvaluator.cpp b/Engine/Plugins/Runtime/GeometryProcessing/Source/DynamicMesh/Private/Sampling/MeshOcclusionMapEvaluator.cpp
--- a/Engine/Plugins/Runtime/GeometryProcessing/Source/DynamicMesh/Private/Sampling/MeshOcclusionMapEvaluator.cpp
+++ b/Engine/Plugins/Runtime/GeometryProcessing/Source/DynamicMesh/Private/Sampling/MeshOcclusionMapEvaluator.cpp
@@ -32,12 +32,14 @@ void FMeshOcclusionMapEvaluator::Setup(const FMeshBaseBaker& Baker, FEvaluationC
 			Context.Evaluate = &EvaluateSample<EMeshOcclusionMapType::All, ESpace::Tangent>;
 			Context.EvaluateDefault = &EvaluateDefault<EMeshOcclusionMapType::All, ESpace::Tangent>;
 			Context.EvaluateColor = &EvaluateColor<EMeshOcclusionMapType::All, ESpace::Tangent>;
+			Context.EvaluateChannel = &EvaluateChannel<EMeshOcclusionMapType::All, ESpace::Tangent>;
 		}
 		else // NormalSpace == ESpace::Object
 		{
 			Context.Evaluate = &EvaluateSample<EMeshOcclusionMapType::All, ESpace::Object>;
 			Context.EvaluateDefault = &EvaluateDefault<EMeshOcclusionMapType::All, ESpace::Object>;
 			Context.EvaluateColor = &EvaluateColor<EMeshOcclusionMapType::All, ESpace::Object>;
+			Context.EvaluateChannel = &EvaluateChannel<EMeshOcclusionMapType::All, ESpace::Object>;
 		}
 		Context.EvalData = this;
 		Context.AccumulateMode = EAccumulateMode::Add;
@@ -47,6 +49,7 @@ void FMeshOcclusionMapEvaluator::Setup(const FMeshBaseBaker& Baker, FEvaluationC
 		Context.Evaluate = &EvaluateSample<EMeshOcclusionMapType::AmbientOcclusion, ESpace::Tangent>;
 		Context.EvaluateDefault = &EvaluateDefault<EMeshOcclusionMapType::AmbientOcclusion, ESpace::Tangent>;
 		Context.EvaluateColor = &EvaluateColor<EMeshOcclusionMapType::AmbientOcclusion, ESpace::Tangent>;
+		Context.EvaluateChannel = &EvaluateChannel<EMeshOcclusionMapType::AmbientOcclusion, ESpace::Tangent>;
 		Context.EvalData = this;
 		Context.AccumulateMode = EAccumulateMode::Add;
 	}
@@ -57,12 +60,14 @@ void FMeshOcclusionMapEvaluator::Setup(const FMeshBaseBaker& Baker, FEvaluationC
 			Context.Evaluate = &EvaluateSample<EMeshOcclusionMapType::BentNormal, ESpace::Tangent>;
 			Context.EvaluateDefault = &EvaluateDefault<EMeshOcclusionMapType::BentNormal, ESpace::Tangent>;
 			Context.EvaluateColor = &EvaluateColor<EMeshOcclusionMapType::BentNormal, ESpace::Tangent>;
+			Context.EvaluateChannel = &EvaluateChannel<EMeshOcclusionMapType::BentNormal, ESpace::Tangent>;
 		}
 		else // NormalSpace == ESpace::Object
 		{
 			Context.Evaluate = &EvaluateSample<EMeshOcclusionMapType::BentNormal, ESpace::Object>;
 			Context.EvaluateDefault = &EvaluateDefault<EMeshOcclusionMapType::BentNormal, ESpace::Object>;
 			Context.EvaluateColor = &EvaluateColor<EMeshOcclusionMapType::BentNormal, ESpace::Object>;
+			Context.EvaluateChannel = &EvaluateChannel<EMeshOcclusionMapType::BentNormal, ESpace::Object>;
 		}
 		Context.EvalData = this;
 		Context.AccumulateMode = EAccumulateMode::Add;
@@ -212,6 +217,37 @@ void FMeshOcclusionMapEvaluator::EvaluateColor(const int DataIdx, float*& In, FV
 	}
 }
 
+template <EMeshOcclusionMapType ComputeType, FMeshOcclusionMapEvaluator::ESpace ComputeSpace>
+void FMeshOcclusionMapEvaluator::EvaluateChannel(const int DataIdx, float*& In, float& Out, void* EvalData)
+{
+	auto EvalOcclusion = [&In, &Out]()
+	{
+		Out = In[0];
+		In += 1;
+	};
+
+	auto EvalNormal = [&In, &Out]()
+	{
+		ensure(false);	// Should not be able to do per-channel baking for normals
+		Out = 0.0f;
+		In += 1;
+	};
+
+	if constexpr (WANT_AMBIENT_OCCLUSION(ComputeType) && WANT_BENT_NORMAL(ComputeType))
+	{
+		DataIdx == 0 ? EvalOcclusion() : EvalNormal();
+	}
+	else if constexpr (WANT_AMBIENT_OCCLUSION(ComputeType))
+	{
+		EvalOcclusion();
+	}
+	else if constexpr (WANT_BENT_NORMAL(ComputeType))
+	{
+		EvalNormal();
+	}
+}
+
+
 template <EMeshOcclusionMapType ComputeType, FMeshOcclusionMapEvaluator::ESpace ComputeSpace>
 FMeshOcclusionMapEvaluator::FOcclusionTuple FMeshOcclusionMapEvaluator::SampleFunction(const FCorrespondenceSample& SampleData, const FVector3d& DefaultNormal)
 {
diff --git a/Engine/Plugins/Runtime/GeometryProcessing/Source/DynamicMesh/Private/Sampling/MeshPropertyMapEvaluator.cpp b/Engine/Plugins/Runtime/GeometryProcessing/Source/DynamicMesh/Private/Sampling/MeshPropertyMapEvaluator.cpp
--- a/Engine/Plugins/Runtime/GeometryProcessing/Source/DynamicMesh/Private/Sampling/MeshPropertyMapEvaluator.cpp
+++ b/Engine/Plugins/Runtime/GeometryProcessing/Source/DynamicMesh/Private/Sampling/MeshPropertyMapEvaluator.cpp
@@ -32,6 +32,7 @@ void FMeshPropertyMapEvaluator::Setup(const FMeshBaseBaker& Baker, FEvaluationCo
 	Context.Evaluate = bHasDetailNormalTextures ? &EvaluateSample<true> : &EvaluateSample<false>;
 	Context.EvaluateDefault = &EvaluateDefault;
 	Context.EvaluateColor = &EvaluateColor;
+	Context.EvaluateChannel = &EvaluateChannel;
 	Context.EvalData = this;
 	Context.AccumulateMode = EAccumulateMode::Add;
 	Context.DataLayout = DataLayout();
@@ -97,6 +98,13 @@ void FMeshPropertyMapEvaluator::EvaluateColor(const int DataIdx, float*& In, FVe
 	In += 3;
 }
 
+void FMeshPropertyMapEvaluator::EvaluateChannel(const int DataIdx, float*& In, float& Out, void* EvalData)
+{
+	ensure(false);		// Should not be able to select per-channel evaluation for multi-dimensional properties
+	Out = In[0];
+	In += 1;
+}
+
 template <bool bUseDetailNormalMap>
 FVector3f FMeshPropertyMapEvaluator::SampleFunction(const FCorrespondenceSample& SampleData) const
 {
diff --git a/Engine/Plugins/Runtime/GeometryProcessing/Source/DynamicMesh/Private/Sampling/MeshResampleImageEvaluator.cpp b/Engine/Plugins/Runtime/GeometryProcessing/Source/DynamicMesh/Private/Sampling/MeshResampleImageEvaluator.cpp
--- a/Engine/Plugins/Runtime/GeometryProcessing/Source/DynamicMesh/Private/Sampling/MeshResampleImageEvaluator.cpp
+++ b/Engine/Plugins/Runtime/GeometryProcessing/Source/DynamicMesh/Private/Sampling/MeshResampleImageEvaluator.cpp
@@ -46,6 +46,12 @@ void FMeshResampleImageEvaluator::EvaluateColor(const int DataIdx, float*& In, F
 	In += 4;
 }
 
+void FMeshResampleImageEvaluator::EvaluateChannel(const int DataIdx, float*& In, float& Out, void* EvalData)
+{
+	Out = In[0];
+	In += 1;
+}
+
 FVector4f FMeshResampleImageEvaluator::ImageSampleFunction(const FCorrespondenceSample& SampleData) const
 {
 	const void* DetailMesh = SampleData.DetailMesh;
@@ -77,6 +83,7 @@ void FMeshMultiResampleImageEvaluator::Setup(const FMeshBaseBaker& Baker, FEvalu
 	Context.Evaluate = &EvaluateSampleMulti;
 	Context.EvaluateDefault = &EvaluateDefault;
 	Context.EvaluateColor = &EvaluateColor;
+	Context.EvaluateChannel = &EvaluateChannel;
 	Context.EvalData = this;
 	Context.AccumulateMode = EAccumulateMode::Add;
 	Context.DataLayout = DataLayout();
diff --git a/Engine/Plugins/Runtime/GeometryProcessing/Source/DynamicMesh/Private/Sampling/MeshVertexBaker.cpp b/Engine/Plugins/Runtime/GeometryProcessing/Source/DynamicMesh/Private/Sampling/MeshVertexBaker.cpp
--- a/Engine/Plugins/Runtime/GeometryProcessing/Source/DynamicMesh/Private/Sampling/MeshVertexBaker.cpp
+++ b/Engine/Plugins/Runtime/GeometryProcessing/Source/DynamicMesh/Private/Sampling/MeshVertexBaker.cpp
@@ -227,15 +227,20 @@ void FMeshVertexBaker::BakeImpl(void* Data)
 				Baker->BakeContexts[BakerIdx].Evaluate(BufferPtr, Sample, Baker->BakeContexts[BakerIdx].EvalData);
 			}
 
-			// For color bakes, ask our evaluators to convert the float data to color.
-			if constexpr(ComputeMode == EBakeMode::RGBA)
+			// TODO: Use a separate buffer rather than R/W from the same pixel.
+			BufferPtr = &Pixel[0];
+			for (int32 BakerIdx = 0; BakerIdx < NumBakers; ++BakerIdx)
 			{
-				// TODO: Use a separate buffer rather than R/W from the same pixel.
-				BufferPtr = &Pixel[0];
-				for (int32 BakerIdx = 0; BakerIdx < NumBakers; ++BakerIdx)
+				if constexpr (ComputeMode == EBakeMode::RGBA)
 				{
 					Baker->BakeContexts[BakerIdx].EvaluateColor(0, BufferPtr, Pixel, Baker->BakeContexts[BakerIdx].EvalData);
 				}
+				else if constexpr (ComputeMode == EBakeMode::PerChannel)
+				{
+					// TODO: Are we guaranteed to have the same NumBakers as number of Pixel channels?
+					ensure(NumBakers == 4);
+					Baker->BakeContexts[BakerIdx].EvaluateChannel(0, BufferPtr, Pixel[BakerIdx], Baker->BakeContexts[BakerIdx].EvalData);
+				}
 			}
 		}
 	}, !Baker->bParallel ? EParallelForFlags::ForceSingleThread : EParallelForFlags::None);
diff --git a/Engine/Plugins/Runtime/GeometryProcessing/Source/DynamicMesh/Public/Sampling/MeshConstantMapEvaluator.h b/Engine/Plugins/Runtime/GeometryProcessing/Source/DynamicMesh/Public/Sampling/MeshConstantMapEvaluator.h
--- a/Engine/Plugins/Runtime/GeometryProcessing/Source/DynamicMesh/Public/Sampling/MeshConstantMapEvaluator.h
+++ b/Engine/Plugins/Runtime/GeometryProcessing/Source/DynamicMesh/Public/Sampling/MeshConstantMapEvaluator.h
@@ -33,6 +33,7 @@ public:
 	static void EvaluateDefault(float*& Out, void* EvalData);
 
 	static void EvaluateColor(const int DataIdx, float*& In, FVector4f& Out, void* EvalData);
+	static void EvaluateChannel(const int DataIdx, float*& In, float& Out, void* EvalData);
 
 public:
 	float Value = 0.0f;
diff --git a/Engine/Plugins/Runtime/GeometryProcessing/Source/DynamicMesh/Public/Sampling/MeshCurvatureMapEvaluator.h b/Engine/Plugins/Runtime/GeometryProcessing/Source/DynamicMesh/Public/Sampling/MeshCurvatureMapEvaluator.h
--- a/Engine/Plugins/Runtime/GeometryProcessing/Source/DynamicMesh/Public/Sampling/MeshCurvatureMapEvaluator.h
+++ b/Engine/Plugins/Runtime/GeometryProcessing/Source/DynamicMesh/Public/Sampling/MeshCurvatureMapEvaluator.h
@@ -68,6 +68,8 @@ public:
 
 	static void EvaluateColor(const int DataIdx, float*& In, FVector4f& Out, void* EvalData);
 
+	static void EvaluateChannel(const int DataIdx, float*& In, float& Out, void* EvalData);
+
 	/** Populate Curvatures member if valid data has not been provided */
 	void CacheDetailCurvatures();
 
diff --git a/Engine/Plugins/Runtime/GeometryProcessing/Source/DynamicMesh/Public/Sampling/MeshMapEvaluator.h b/Engine/Plugins/Runtime/GeometryProcessing/Source/DynamicMesh/Public/Sampling/MeshMapEvaluator.h
--- a/Engine/Plugins/Runtime/GeometryProcessing/Source/DynamicMesh/Public/Sampling/MeshMapEvaluator.h
+++ b/Engine/Plugins/Runtime/GeometryProcessing/Source/DynamicMesh/Public/Sampling/MeshMapEvaluator.h
@@ -119,6 +119,27 @@ public:
 		using EvaluateColorFn = void(*)(const int DataIdx, float*& In, FVector4f& Out, void* EvalData);
 		EvaluateColorFn EvaluateColor = nullptr;
 
+		/**
+		 * Function pointer to evaluate the single channel representation of
+		 * an evaluated sample result.
+		 *
+		 * This function will be invoked once per pixel channel after the
+		 * evaluated samples have been accumulated. This function
+		 * should:
+		 *
+		 * - assume that the size of the buffer is correct.
+		 * - advance the [In] buffer pointer by the DataLayout stride.
+		 *
+		 * @param DataIdx the index into the DataLayout being processed.
+		 * @param In the buffer containing the accumulated sample result.
+		 * @param Out the output float channel to populate
+		 * @param EvalData custom data pointer provided by the evaluation context.
+		 */
+		using EvaluateChannelFn = void(*)(const int DataIdx, float*& In, float& Out, void* EvalData);
+		EvaluateChannelFn EvaluateChannel = nullptr;
+
+
+
 		/** Define custom data to be passed to the evaluation function. */
 		void* EvalData = nullptr;
 
diff --git a/Engine/Plugins/Runtime/GeometryProcessing/Source/DynamicMesh/Public/Sampling/MeshNormalMapEvaluator.h b/Engine/Plugins/Runtime/GeometryProcessing/Source/DynamicMesh/Public/Sampling/MeshNormalMapEvaluator.h
--- a/Engine/Plugins/Runtime/GeometryProcessing/Source/DynamicMesh/Public/Sampling/MeshNormalMapEvaluator.h
+++ b/Engine/Plugins/Runtime/GeometryProcessing/Source/DynamicMesh/Public/Sampling/MeshNormalMapEvaluator.h
@@ -39,6 +39,7 @@ public:
 	static void EvaluateDefault(float*& Out, void* EvalData);
 
 	static void EvaluateColor(const int DataIdx, float*& In, FVector4f& Out, void* EvalData);
+	static void EvaluateChannel(const int DataIdx, float*& In, float& Out, void* EvalData);
 
 protected:
 	// Cached data
diff --git a/Engine/Plugins/Runtime/GeometryProcessing/Source/DynamicMesh/Public/Sampling/MeshOcclusionMapEvaluator.h b/Engine/Plugins/Runtime/GeometryProcessing/Source/DynamicMesh/Public/Sampling/MeshOcclusionMapEvaluator.h
--- a/Engine/Plugins/Runtime/GeometryProcessing/Source/DynamicMesh/Public/Sampling/MeshOcclusionMapEvaluator.h
+++ b/Engine/Plugins/Runtime/GeometryProcessing/Source/DynamicMesh/Public/Sampling/MeshOcclusionMapEvaluator.h
@@ -73,6 +73,9 @@ public:
 	template <EMeshOcclusionMapType OcclusionType, ESpace NormalSpace>
 	static void EvaluateColor(const int DataIdx, float*& In, FVector4f& Out, void* EvalData);
 
+	template <EMeshOcclusionMapType OcclusionType, ESpace NormalSpace>
+	static void EvaluateChannel(const int DataIdx, float*& In, float& Out, void* EvalData);
+
 protected:
 	// Cached data
 	const IMeshBakerDetailSampler* DetailSampler = nullptr;
diff --git a/Engine/Plugins/Runtime/GeometryProcessing/Source/DynamicMesh/Public/Sampling/MeshPropertyMapEvaluator.h b/Engine/Plugins/Runtime/GeometryProcessing/Source/DynamicMesh/Public/Sampling/MeshPropertyMapEvaluator.h
--- a/Engine/Plugins/Runtime/GeometryProcessing/Source/DynamicMesh/Public/Sampling/MeshPropertyMapEvaluator.h
+++ b/Engine/Plugins/Runtime/GeometryProcessing/Source/DynamicMesh/Public/Sampling/MeshPropertyMapEvaluator.h
@@ -52,6 +52,7 @@ public:
 	static void EvaluateDefault(float*& Out, void* EvalData);
 
 	static void EvaluateColor(const int DataIdx, float*& In, FVector4f& Out, void* EvalData);
+	static void EvaluateChannel(const int DataIdx, float*& In, float& Out, void* EvalData);
 
 protected:
 	// Cached data
diff --git a/Engine/Plugins/Runtime/GeometryProcessing/Source/DynamicMesh/Public/Sampling/MeshResampleImageEvaluator.h b/Engine/Plugins/Runtime/GeometryProcessing/Source/DynamicMesh/Public/Sampling/MeshResampleImageEvaluator.h
--- a/Engine/Plugins/Runtime/GeometryProcessing/Source/DynamicMesh/Public/Sampling/MeshResampleImageEvaluator.h
+++ b/Engine/Plugins/Runtime/GeometryProcessing/Source/DynamicMesh/Public/Sampling/MeshResampleImageEvaluator.h
@@ -35,6 +35,7 @@ public:
 	static void EvaluateDefault(float*& Out, void* EvalData);
 
 	static void EvaluateColor(const int DataIdx, float*& In, FVector4f& Out, void* EvalData);
+	static void EvaluateChannel(const int DataIdx, float*& In, float& Out, void* EvalData);
 
 protected:
 	// Cached data
