Update FLocalPlanarSimplify to support simplification along straight creases in a mesh, where the sides of the crease may optionally belong to different PolyGroups

Add an edge simplification operation to the PolyEdit tool

diff --git a/Engine/Plugins/Runtime/GeometryProcessing/Source/DynamicMesh/Private/Operations/LocalPlanarSimplify.cpp b/Engine/Plugins/Runtime/GeometryProcessing/Source/DynamicMesh/Private/Operations/LocalPlanarSimplify.cpp
--- a/Engine/Plugins/Runtime/GeometryProcessing/Source/DynamicMesh/Private/Operations/LocalPlanarSimplify.cpp
+++ b/Engine/Plugins/Runtime/GeometryProcessing/Source/DynamicMesh/Private/Operations/LocalPlanarSimplify.cpp
@@ -33,27 +33,63 @@ bool FLocalPlanarSimplify::IsFlat(const FDynamicMesh3& Mesh, int VID, double Dot
 	return bIsFlat;
 }
 
+bool FLocalPlanarSimplify::IsDevelopableAlongEdge(const FDynamicMesh3& Mesh, int EID, int VID, double DotTolerance, FVector3d& NormalA, bool& bIsFlat)
+{
+	FIndex2i EdgeT = Mesh.GetEdgeT(EID);
+	NormalA = Mesh.GetTriNormal(EdgeT.A);
+	bIsFlat = true;
+	FVector3d NormalB;
+	if (EdgeT.B != FDynamicMesh3::InvalidID)
+	{
+		NormalB = Mesh.GetTriNormal(EdgeT.B);
+		// Only consider a second normal if the across-edge triangle normal doesn't match the first normal
+		if (NormalA.Dot(NormalB) < DotTolerance)
+		{
+			bIsFlat = false;
+		}
+	}
+	bool bIsDevelopable = true;
+	Mesh.EnumerateVertexTriangles(VID, [&bIsDevelopable, EdgeT, &Mesh, NormalA, NormalB, DotTolerance, bIsFlat](int32 TID)
+		{
+			if (!bIsDevelopable || EdgeT.Contains(TID))
+			{
+				return;
+			}
+			FVector3d Normal = Mesh.GetTriNormal(TID);
+			bIsDevelopable = Normal.Dot(NormalA) >= DotTolerance;
+			if (!bIsDevelopable && !bIsFlat) // if we didn't match first normal, test the second normal (if we have one)
+			{
+				bIsDevelopable = Normal.Dot(NormalB) >= DotTolerance;
+			}
+		});
+	bIsFlat = bIsFlat && bIsDevelopable;
+
+	return bIsDevelopable;
+}
+
 /**
  * Test if a given edge collapse would cause a triangle flip or other unacceptable decrease in mesh quality
  */
 bool FLocalPlanarSimplify::CollapseWouldHurtTriangleQuality(
-	const FDynamicMesh3& Mesh, const FVector3d& ExpectNormal,
+	const FDynamicMesh3& Mesh, const FVector3d& RemoveVNormal,
 	int32 RemoveV, const FVector3d& RemoveVPos, int32 KeepV, const FVector3d& KeepVPos,
-	double TryToImproveTriQualityThreshold
+	double TryToImproveTriQualityThreshold, bool bHasMultipleNormals
 	)
 {
 	double WorstQualityNewTriangle = FMathd::MaxReal;
 
 	bool bIsHurt = false;
 	Mesh.EnumerateVertexTriangles(RemoveV,
-		[&Mesh, &bIsHurt, &KeepVPos, RemoveV, KeepV, &ExpectNormal, 
-		TryToImproveTriQualityThreshold, &WorstQualityNewTriangle](int32 TID)
+		[&Mesh, &bIsHurt, &KeepVPos, RemoveV, KeepV, &RemoveVNormal,
+		TryToImproveTriQualityThreshold, &WorstQualityNewTriangle, bHasMultipleNormals](int32 TID)
 		{
 			if (bIsHurt)
 			{
 				return;
 			}
 			FIndex3i Tri = Mesh.GetTriangle(TID);
+			// Skip re-computing the triangle normal for flat surfaces
+			FVector3d TriNormal = bHasMultipleNormals ? Mesh.GetTriNormal(TID) : RemoveVNormal;
 			FVector3d Verts[3];
 			for (int Idx = 0; Idx < 3; Idx++)
 			{
@@ -78,7 +114,6 @@ bool FLocalPlanarSimplify::CollapseWouldHurtTriangleQuality(
 			FVector3d Edge2(Verts[2] - Verts[0]);
 			FVector3d VCross(Edge2.Cross(Edge1));
 
-			// TODO: does this tolerance make a difference?  if not, set to zero and remove the Normalize(VCross)
 			double EdgeFlipTolerance = 1.e-5;
 			double Area2 = Normalize(VCross);
 			if (TryToImproveTriQualityThreshold > 0)
@@ -89,7 +124,7 @@ bool FLocalPlanarSimplify::CollapseWouldHurtTriangleQuality(
 				WorstQualityNewTriangle = FMathd::Min(Quality, WorstQualityNewTriangle);
 			}
 
-			bIsHurt = VCross.Dot(ExpectNormal) <= EdgeFlipTolerance;
+			bIsHurt = VCross.Dot(TriNormal) <= EdgeFlipTolerance;
 		}
 	);
 
@@ -127,7 +162,7 @@ bool FLocalPlanarSimplify::CollapseWouldHurtTriangleQuality(
  * Test if a given edge collapse would change the mesh shape or UVs unacceptably
  */
 bool FLocalPlanarSimplify::CollapseWouldChangeShapeOrUVs(
-	const FDynamicMesh3& Mesh, const TSet<int>& CutBoundaryEdgeSet, double DotTolerance, int SourceEID,
+	const FDynamicMesh3& Mesh, const TSet<int>& PathEdgeSet, double DotTolerance, int SourceEID,
 	int32 RemoveV, const FVector3d& RemoveVPos, int32 KeepV, const FVector3d& KeepVPos, const FVector3d& EdgeDir,
 	bool bPreserveTriangleGroups, bool bPreserveUVsForMesh, bool bPreserveVertexUVs, bool bPreserveOverlayUVs,
 	float UVEqualThresholdSq, bool bPreserveVertexNormals, float NormalEqualCosThreshold)
@@ -138,7 +173,18 @@ bool FLocalPlanarSimplify::CollapseWouldChangeShapeOrUVs(
 	bool bHasBadEdge = false;
 
 	int OpposedEdge = -1;
-	int SourceGroupID = Mesh.GetTriangleGroup(Mesh.GetEdgeT(SourceEID).A);
+	FIndex2i EdgeT = Mesh.GetEdgeT(SourceEID);
+	int SourceGroupID = Mesh.GetTriangleGroup(EdgeT.A);
+	int OtherGroupID = -1;
+	// Note: In many cases (e.g. plane cut, new edge insertions) we know there will be only one group,
+	// so we could use a slightly faster path with bAllowTwoGroups==false. But this is probably not
+	// a significant performance difference, so for now we just expose the allows-two-group path.
+	constexpr bool bAllowTwoGroups = true;
+	if (bAllowTwoGroups && EdgeT.B != FDynamicMesh3::InvalidID)
+	{
+		OtherGroupID = Mesh.GetTriangleGroup(EdgeT.B);
+	}
+	
 	Mesh.EnumerateVertexEdges(RemoveV,
 		[&](int32 VertEID)
 		{
@@ -150,17 +196,22 @@ bool FLocalPlanarSimplify::CollapseWouldChangeShapeOrUVs(
 			FDynamicMesh3::FEdge Edge = Mesh.GetEdge(VertEID);
 			if (bPreserveTriangleGroups && Mesh.HasTriangleGroups())
 			{
-				if (SourceGroupID != Mesh.GetTriangleGroup(Edge.Tri.A) ||
-					(Edge.Tri.B != FDynamicMesh3::InvalidID && SourceGroupID != Mesh.GetTriangleGroup(Edge.Tri.B)))
+				int GroupA = Mesh.GetTriangleGroup(Edge.Tri.A);
+				int GroupB = Edge.Tri.B == FDynamicMesh3::InvalidID ? SourceGroupID : Mesh.GetTriangleGroup(Edge.Tri.B);
+
+				if (
+					(GroupA != SourceGroupID && (!bAllowTwoGroups || GroupA != OtherGroupID)) ||
+					(GroupB != SourceGroupID && (!bAllowTwoGroups || GroupB != OtherGroupID))
+					)
 				{
-					// RemoveV is on a group boundary, so the edge collapse would change the shape of the groups
+					// RemoveV is bordering too many groups, so the edge collapse would change the shape of the groups
 					bHasBadEdge = true;
 					return;
 				}
 			}
 
-			// it's a known boundary edge; check if it's the opposite-facing one we need
-			if (CutBoundaryEdgeSet.Contains(VertEID))
+			// it's a path edge; check if it's the opposite-facing one we need
+			if (PathEdgeSet.Contains(VertEID))
 			{
 				if (OpposedEdge != -1)
 				{
@@ -187,12 +238,27 @@ bool FLocalPlanarSimplify::CollapseWouldChangeShapeOrUVs(
 					return;
 				}
 
-				// test that UVs are not too distorted through the collapse
+				// LerpT used for vertex normal and UV interpolation
+				float LerpT = float((OtherVPos - RemoveVPos).Dot(OtherEdgeDir) / (OtherVPos - KeepVPos).Dot(OtherEdgeDir));
+
+				if (bPreserveVertexNormals && Mesh.HasVertexNormals())
+				{
+					FVector3f OtherN = Mesh.GetVertexNormal(OtherV);
+					FVector3f RemoveN = Mesh.GetVertexNormal(RemoveV);
+					FVector3f KeepN = Mesh.GetVertexNormal(KeepV);
+					if (Normalized(Lerp(OtherN, KeepN, LerpT)).Dot(Normalized(RemoveN)) < NormalEqualCosThreshold)
+					{
+						bHasBadEdge = true;
+						return;
+					}
+				}
+
+				// Controls whether any UV distortion is tested
 				if (!bPreserveUVsForMesh)
 				{
 					return;
 				}
-				float LerpT = float( (RemoveVPos - OtherVPos).Dot(OtherEdgeDir) / (KeepVPos - OtherVPos).Dot(OtherEdgeDir) );
+				
 				if (bPreserveVertexUVs && Mesh.HasVertexUVs())
 				{
 					FVector2f OtherUV = Mesh.GetVertexUV(OtherV);
@@ -204,74 +270,108 @@ bool FLocalPlanarSimplify::CollapseWouldChangeShapeOrUVs(
 						return;
 					}
 				}
-				if (bPreserveVertexNormals && Mesh.HasVertexNormals())
-				{
-					FVector3f OtherN = Mesh.GetVertexNormal(OtherV);
-					FVector3f RemoveN = Mesh.GetVertexNormal(RemoveV);
-					FVector3f KeepN = Mesh.GetVertexNormal(KeepV);
-					if (Normalized(Lerp(OtherN, KeepN, LerpT)).Dot(Normalized(RemoveN)) < NormalEqualCosThreshold)
-					{
-						bHasBadEdge = true;
-						return;
-					}
-				}
 				if (bPreserveOverlayUVs && Mesh.HasAttributes())
 				{
 					int NumLayers = Mesh.Attributes()->NumUVLayers();
 					FIndex2i SourceEdgeTris = Mesh.GetEdgeT(SourceEID);
 					FIndex2i OppEdgeTris = Edge.Tri;
 
-					// special handling of seam edge when the edges aren't boundary edges
-					// -- if they're seams, we'd need to check both sides of the seams for a UV match
-					//    but this is complicated and should be quite rare, so we just don't collapse these
-					if (SourceEdgeTris.B != -1 || OppEdgeTris.B != -1)
+					// returns true if collapse will result in acceptable overlay UVs
+					auto CanCollapseOverlayUVs = [&Mesh, KeepV, RemoveV, OtherV, NumLayers, LerpT, UVEqualThresholdSq](int SourceEdgeTID, int OppEdgeTID) -> bool
 					{
-						if (Mesh.Attributes()->IsSeamEdge(SourceEID) ||
-							Mesh.Attributes()->IsSeamEdge(VertEID))
+						FIndex3i SourceBaseTri = Mesh.GetTriangle(SourceEdgeTID);
+						FIndex3i OppBaseTri = Mesh.GetTriangle(OppEdgeTID);
+						int KeepSourceIdx = IndexUtil::FindTriIndex(KeepV, SourceBaseTri);
+						int RemoveSourceIdx = IndexUtil::FindTriIndex(RemoveV, SourceBaseTri);
+						int OtherOppIdx = IndexUtil::FindTriIndex(OtherV, OppBaseTri);
+						if (!ensure(KeepSourceIdx != -1 && RemoveSourceIdx != -1 && OtherOppIdx != -1))
 						{
-							bHasBadEdge = true;
-							return;
+							return false;
 						}
-					}
 
-					FIndex3i SourceBaseTri = Mesh.GetTriangle(SourceEdgeTris.A);
-					FIndex3i OppBaseTri = Mesh.GetTriangle(OppEdgeTris.A);
-					int KeepSourceIdx = IndexUtil::FindTriIndex(KeepV, SourceBaseTri);
-					int RemoveSourceIdx = IndexUtil::FindTriIndex(RemoveV, SourceBaseTri);
-					int OtherOppIdx = IndexUtil::FindTriIndex(OtherV, OppBaseTri);
-					if (!ensure(KeepSourceIdx != -1 && RemoveSourceIdx != -1 && OtherOppIdx != -1))
-					{
-						bHasBadEdge = true;
-						return;
-					}
-
-					// get the UVs per overlay off the triangle(s) attached the two edges
-					for (int UVLayerIdx = 0; UVLayerIdx < NumLayers; UVLayerIdx++)
-					{
-						const FDynamicMeshUVOverlay* UVs = Mesh.Attributes()->GetUVLayer(UVLayerIdx);
-						if (UVs->ElementCount() < 3)
+						// get the UVs per overlay off the triangle(s) attached the two edges
+						for (int UVLayerIdx = 0; UVLayerIdx < NumLayers; UVLayerIdx++)
 						{
-							// overlay is not actually in use; skip it
-							continue;
+							const FDynamicMeshUVOverlay* UVs = Mesh.Attributes()->GetUVLayer(UVLayerIdx);
+							if (UVs->ElementCount() < 3)
+							{
+								// overlay is not actually in use; skip it
+								continue;
+							}
+							FIndex3i SourceT = UVs->GetTriangle(SourceEdgeTID);
+							FIndex3i OppT = UVs->GetTriangle(OppEdgeTID);
+							int KeepE = SourceT[KeepSourceIdx];
+							int RemoveE = SourceT[RemoveSourceIdx];
+							int OtherE = OppT[OtherOppIdx];
+							if (KeepE == -1 || RemoveE == -1 || OtherE == -1)
+							{
+								// overlay is not set on relevant triangles; skip it
+								continue;
+							}
+							FVector2f OtherUV = UVs->GetElement(OtherE);
+							FVector2f RemoveUV = UVs->GetElement(RemoveE);
+							FVector2f KeepUV = UVs->GetElement(KeepE);
+							if (DistanceSquared(Lerp(OtherUV, KeepUV, LerpT), RemoveUV) > UVEqualThresholdSq)
+							{
+								return false;
+							}
 						}
-						FIndex3i SourceT = UVs->GetTriangle(SourceEdgeTris.A);
-						FIndex3i OppT = UVs->GetTriangle(OppEdgeTris.A);
-						int KeepE = SourceT[KeepSourceIdx];
-						int RemoveE = SourceT[RemoveSourceIdx];
-						int OtherE = OppT[OtherOppIdx];
-						if (KeepE == -1 || RemoveE == -1 || OtherE == -1)
+
+						return true;
+					};
+
+					// If we're not on a boundary, check if we're on a seam -- and if so, test that
+					// the collapse won't change the UVs on the other side of the seam
+					if (SourceEdgeTris.B != -1 || OppEdgeTris.B != -1)
+					{
+						// the edges must be both boundary or neither
+						if (SourceEdgeTris.B == -1 || OppEdgeTris.B == -1)
 						{
-							// overlay is not set on relevant triangles; skip it
-							continue;
+							bHasBadEdge = true;
+							return;
 						}
-						FVector2f OtherUV = UVs->GetElement(OtherE);
-						FVector2f RemoveUV = UVs->GetElement(RemoveE);
-						FVector2f KeepUV = UVs->GetElement(KeepE);
-						if ( DistanceSquared( Lerp(OtherUV, KeepUV, LerpT), RemoveUV) > UVEqualThresholdSq)
+						// likewise must be both seam or neither
+						bool bSourceIsSeam = Mesh.Attributes()->IsSeamEdge(SourceEID);
+						bool bOtherIsSeam = Mesh.Attributes()->IsSeamEdge(VertEID);
+						if (bSourceIsSeam != bOtherIsSeam)
 						{
 							bHasBadEdge = true;
 							return;
 						}
+						// if both are seam, make sure that the A and B sides are consistent,
+						// then test the UVs on the other side of the seam
+						if (bSourceIsSeam)
+						{
+							// test whether the A/B sides are consistent for SourceEdgeTris and OppEdgeTris,
+							// based on the ordering of the remove and keep/other vertices in each triangle
+							FIndex3i SourceBaseTri = Mesh.GetTriangle(SourceEdgeTris.A);
+							FIndex3i OppBaseTri = Mesh.GetTriangle(OppEdgeTris.A);
+							int32 SrcTriRmVSubIdx = IndexUtil::FindTriIndex(RemoveV, SourceBaseTri);
+							int32 OppTriRmVSubIdx = IndexUtil::FindTriIndex(RemoveV, OppBaseTri);
+							// Test whether the Triangle "A" for the Source Edge is wound such that the 'Remove' Vertex is before the 'Keep' Vertex
+							bool bSrcAOrderIsRmThenKeep = SourceBaseTri[(SrcTriRmVSubIdx + 1) % 3] == KeepV;
+							// Test whether the Triangle "A" for the Opposite Edge is wound such that the 'Remove' Vertex is before the 'Other' Vertex
+							bool bOppAOrderIsRmThenOther = OppBaseTri[(OppTriRmVSubIdx + 1) % 3] == OtherV;
+							// If the two "A" sides are on the same side of the seam, we expect these two vertex ordering bools *not* to match --
+							// i.e., the oriented edges should be "Keep Remove, then Remove Other" or "Other Remove, then Remove Keep"
+							// If the bools match, we swap the A and B of the Opposite Edge Triangles, to move the A side back to the same side of the seam
+							// as its Source Edge Triangle counterpart.
+							if (bSrcAOrderIsRmThenKeep == bOppAOrderIsRmThenOther)
+							{
+								Swap(OppEdgeTris.A, OppEdgeTris.B);
+							}
+							if (!CanCollapseOverlayUVs(SourceEdgeTris.B, OppEdgeTris.B))
+							{
+								bHasBadEdge = true;
+								return;
+							}
+						}
+					}
+
+					if (!CanCollapseOverlayUVs(SourceEdgeTris.A, OppEdgeTris.A))
+					{
+						bHasBadEdge = true;
+						return;
 					}
 				}
 			}
@@ -292,6 +392,12 @@ void FLocalPlanarSimplify::SimplifyAlongEdges(FDynamicMesh3& Mesh, TSet<int32>&
 {
 	double DotTolerance = FMathd::Cos(SimplificationAngleTolerance * FMathd::DegToRad);
 
+	// For some simplification passes we can safely only consider locally-flat vertices for simplification, ignoring "straight crease" cases because
+	// the new edges that we consider simplifying were generated by cutting across flat triangles. This makes the simplification logic a bit simpler.
+	// You can use this variable to toggle that path.
+	// TODO: We could consider exposing this variable as a parameter (but it does not seem likely to affect performance enough to really be worth doing so)
+	constexpr bool bOnlySimplifyWhereFlat = false;
+
 	// Save the input list of edges to iterate over, so we can edit the set w/ new edges while safely iterating over the old ones
 	TArray<int32> CutBoundaryEdgesArray = InOutEdges.Array();
 
@@ -312,21 +418,29 @@ void FLocalPlanarSimplify::SimplifyAlongEdges(FDynamicMesh3& Mesh, TSet<int32>&
 			FDynamicMesh3::FEdge Edge = Mesh.GetEdge(EID);
 
 			// track whether the neighborhood of the vertex is flat
-			bool Flat[2]{ false, false };
+			bool Developable[2]{ false, false };
 			// normals for each flat vertex
 			FVector3d FlatNormals[2]{ FVector3d::Zero(), FVector3d::Zero() };
-			int NumFlat = 0;
+			bool Flat[2]{ false, false };
+			int NumDevelopable = 0;
 			for (int VIdx = 0; VIdx < 2; VIdx++)
 			{
-				Flat[VIdx] = IsFlat(Mesh, Edge.Vert[VIdx], DotTolerance, FlatNormals[VIdx]);
+				if (bOnlySimplifyWhereFlat)
+				{
+					Flat[VIdx] = Developable[VIdx] = IsFlat(Mesh, Edge.Vert[VIdx], DotTolerance, FlatNormals[VIdx]);
+				}
+				else
+				{
+					Developable[VIdx] = IsDevelopableAlongEdge(Mesh, EID, Edge.Vert[VIdx], DotTolerance, FlatNormals[VIdx], Flat[VIdx]);
+				}
 
-				if (Flat[VIdx])
+				if (Developable[VIdx])
 				{
-					NumFlat++;
+					NumDevelopable++;
 				}
 			}
 
-			if (NumFlat == 0)
+			if (NumDevelopable == 0)
 			{
 				continue;
 			}
@@ -334,7 +448,7 @@ void FLocalPlanarSimplify::SimplifyAlongEdges(FDynamicMesh3& Mesh, TSet<int32>&
 			// see if we can collapse to remove either vertex
 			for (int RemoveVIdx = 0; RemoveVIdx < 2; RemoveVIdx++)
 			{
-				if (!Flat[RemoveVIdx])
+				if (!Developable[RemoveVIdx])
 				{
 					continue;
 				}
@@ -355,7 +469,7 @@ void FLocalPlanarSimplify::SimplifyAlongEdges(FDynamicMesh3& Mesh, TSet<int32>&
 				int SourceEID = EID;
 
 				bHasBadEdge = bHasBadEdge || CollapseWouldHurtTriangleQuality(
-					Mesh, FlatNormals[RemoveVIdx], RemoveV, RemoveVPos, KeepV, KeepVPos, TryToImproveTriQualityThreshold);
+					Mesh, FlatNormals[RemoveVIdx], RemoveV, RemoveVPos, KeepV, KeepVPos, TryToImproveTriQualityThreshold, !Flat[RemoveVIdx]);
 
 				bHasBadEdge = bHasBadEdge || CollapseWouldChangeShapeOrUVs(
 					Mesh, InOutEdges, DotTolerance,
diff --git a/Engine/Plugins/Runtime/GeometryProcessing/Source/DynamicMesh/Public/Operations/LocalPlanarSimplify.h b/Engine/Plugins/Runtime/GeometryProcessing/Source/DynamicMesh/Public/Operations/LocalPlanarSimplify.h
--- a/Engine/Plugins/Runtime/GeometryProcessing/Source/DynamicMesh/Public/Operations/LocalPlanarSimplify.h
+++ b/Engine/Plugins/Runtime/GeometryProcessing/Source/DynamicMesh/Public/Operations/LocalPlanarSimplify.h
@@ -74,21 +74,40 @@ public:
 	 */
 	static bool IsFlat(const FDynamicMesh3& Mesh, int VID, double DotTolerance, FVector3d& OutFirstNormal);
 
+	/**
+	 * Test if the triangles connected to a vertex could be flattened by folding along the given edge
+	 * @param Mesh The mesh to query
+	 * @param EID The edge around which the triangles are allowed to fold
+	 * @param VID The vertex to query
+	 * @param DotTolerance If the dot product of two normals are >= this tolerance, the normals are considered equivalent
+	 * @param The normal of the first triangle attached to the vertex.
+	 * @return Whether all the triangles were coplanar
+	 */
+	static bool IsDevelopableAlongEdge(const FDynamicMesh3& Mesh, int EID, int VID, double DotTolerance, FVector3d& NormalA, bool& bIsFlat);
+
 	/**
 	 * Test if a given edge collapse would cause a triangle flip or other unacceptable decrease in mesh quality
 	 * Specialized for collapsing at flat triangles
+	 * @param Mesh The mesh to query
+	 * @param RemoveVNormal Only used if bHasMultipleNormals==false, i.e. if the surface is locally flat, to avoid recomputing the normal. The normal for all triangles surrounding RemoveV.
+	 * @param RemoveV The vertex to consider removing by an edge collapse
+	 * @param RemoveVPos The position of the vertex RemoveV
+	 * @param KeepV The vertex that is not being removed in the edge collapse
+	 * @param KeepVPos The position of vertex KeepV
+	 * @param TryToImproveTriQualityThreshold Threshold for triangle quality (see comment for class member, above)
+	 * @param bHasMultipleNormals If false, assume the triangulation is locally flat, so we can use ExpectNormal instead of recomputing the expected normal
 	 */
 	static bool CollapseWouldHurtTriangleQuality(
-		const FDynamicMesh3& Mesh, const FVector3d& ExpectNormal,
+		const FDynamicMesh3& Mesh, const FVector3d& RemoveVNormal,
 		int32 RemoveV, const FVector3d& RemoveVPos, int32 KeepV, const FVector3d& KeepVPos,
-		double TryToImproveTriQualityThreshold
+		double TryToImproveTriQualityThreshold, bool bHasMultipleNormals = false
 	);
 
 	/**
 	 * Test if a given edge collapse would change the mesh shape, mesh triangle group shape, or UVs unacceptably
 	 */
 	static bool CollapseWouldChangeShapeOrUVs(
-		const FDynamicMesh3& Mesh, const TSet<int>& CutBoundaryEdgeSet, double DotTolerance, int SourceEID,
+		const FDynamicMesh3& Mesh, const TSet<int>& PathEdgeSet, double DotTolerance, int SourceEID,
 		int32 RemoveV, const FVector3d& RemoveVPos, int32 KeepV, const FVector3d& KeepVPos,
 		const FVector3d& EdgeDir, bool bPreserveTriangleGroups, bool bPreserveUVsForMesh,
 		bool bPreserveVertexUVs, bool bPreserveOverlayUVs, float UVToleranceSq,
diff --git a/Engine/Plugins/Runtime/MeshModelingToolset/Source/MeshModelingTools/Private/EditMeshPolygonsTool.cpp b/Engine/Plugins/Runtime/MeshModelingToolset/Source/MeshModelingTools/Private/EditMeshPolygonsTool.cpp
--- a/Engine/Plugins/Runtime/MeshModelingToolset/Source/MeshModelingTools/Private/EditMeshPolygonsTool.cpp
+++ b/Engine/Plugins/Runtime/MeshModelingToolset/Source/MeshModelingTools/Private/EditMeshPolygonsTool.cpp
@@ -29,6 +29,7 @@
 #include "Operations/MinimalHoleFiller.h"
 #include "Operations/PolygroupRemesh.h"
 #include "Operations/WeldEdgeSequence.h"
+#include "Operations/LocalPlanarSimplify.h"
 #include "Selection/StoredMeshSelectionUtil.h"
 #include "Selection/PolygonSelectionMechanic.h"
 #include "Selections/MeshConnectedComponents.h"
@@ -1321,6 +1322,9 @@ void UEditMeshPolygonsTool::OnTick(float DeltaTime)
 		case EEditMeshPolygonsToolActions::BridgeEdges:
 			ApplyBridgeEdges();
 			break;
+		case EEditMeshPolygonsToolActions::SimplifyAlongEdges:
+			ApplySimplifyAlongEdges();
+			break;
 		case EEditMeshPolygonsToolActions::Retriangulate:
 			ApplyRetriangulate();
 			break;
@@ -2111,6 +2115,50 @@ void UEditMeshPolygonsTool::ApplyStraightenEdges()
 }
 
 
+void UEditMeshPolygonsTool::ApplySimplifyAlongEdges()
+{
+	if (BeginMeshEdgeEditChange() == false)
+	{
+		GetToolManager()->DisplayMessage(LOCTEXT("OnSimplifyAlongEdgesFailed", "Cannot Simplify current selection"), EToolMessageLevel::UserWarning);
+		return;
+	}
+
+	FDynamicMesh3* Mesh = CurrentMesh.Get();
+
+	FDynamicMeshChangeTracker ChangeTracker(Mesh);
+	ChangeTracker.BeginChange();
+
+	// Storage for edge sets is re-used for each selected polygon edge path
+	TSet<int32> SimplifyEdgeSet;
+
+	// Pre-save vertices and triangles along selected edges
+	for (const FSelectedEdge& Edge : ActiveEdgeSelection)
+	{
+		if (Edge.EdgeIDs.Num() > 1)
+		{
+			const TArray<int32>& EdgeVerts = Topology->GetGroupEdgeVertices(Edge.EdgeTopoID);
+			ChangeTracker.SaveVertexOneRingTriangles(EdgeVerts, true);
+		}
+	}
+	
+	// Attempt simplification along edges
+	for (const FSelectedEdge& Edge : ActiveEdgeSelection)
+	{
+		if (Edge.EdgeIDs.Num() > 1)
+		{
+			SimplifyEdgeSet.Reset();
+			SimplifyEdgeSet.Append(Edge.EdgeIDs);
+			FLocalPlanarSimplify LocalSimplify;
+			LocalSimplify.bPreserveVertexNormals = false;
+			LocalSimplify.SimplifyAlongEdges(*Mesh, SimplifyEdgeSet);
+		}
+	}
+
+	FGroupTopologySelection NewSelection;
+	EmitCurrentMeshChangeAndUpdate(LOCTEXT("PolyMeshSimplifyAlongEdgesChange", "Simplify Along Edges"),
+		ChangeTracker.EndChange(), NewSelection);
+}
+
 
 void UEditMeshPolygonsTool::ApplyFillHole()
 {
diff --git a/Engine/Plugins/Runtime/MeshModelingToolset/Source/MeshModelingTools/Public/EditMeshPolygonsTool.h b/Engine/Plugins/Runtime/MeshModelingToolset/Source/MeshModelingTools/Public/EditMeshPolygonsTool.h
--- a/Engine/Plugins/Runtime/MeshModelingToolset/Source/MeshModelingTools/Public/EditMeshPolygonsTool.h
+++ b/Engine/Plugins/Runtime/MeshModelingToolset/Source/MeshModelingTools/Public/EditMeshPolygonsTool.h
@@ -130,6 +130,7 @@ enum class EEditMeshPolygonsToolActions
 	FillHole,
 	BridgeEdges,
 	BevelEdges,
+	SimplifyAlongEdges,
 
 	PlanarProjectionUV,
 
@@ -416,6 +417,10 @@ public:
 	/** Create a new face that connects the selected edges */
 	UFUNCTION(CallInEditor, Category = EdgeEdits, meta = (DisplayName = "Bridge", DisplayPriority = 5))
 	void Bridge() { PostAction(EEditMeshPolygonsToolActions::BridgeEdges); }
+
+	/** Simplify the underlying triangulation along the selected edges, when doing so won't change the shape or UVs, or make low-quality triangles */
+	UFUNCTION(CallInEditor, Category = EdgeEdits, meta = (DisplayPriority = 6))
+	void Simplify() { PostAction(EEditMeshPolygonsToolActions::SimplifyAlongEdges); }
 };
 
 
@@ -687,6 +692,7 @@ protected:
 	void ApplyStraightenEdges();
 	void ApplyFillHole();
 	void ApplyBridgeEdges();
+	void ApplySimplifyAlongEdges();
 
 	void ApplyFlipSingleEdge();
 	void ApplyCollapseSingleEdge();
