fix handling of an outer polygon nested inside a hole in the clipper2 wrapper code

diff --git a/Engine/Plugins/Runtime/GeometryProcessing/Source/GeometryAlgorithms/Private/Curve/ClipperUtils.h b/Engine/Plugins/Runtime/GeometryProcessing/Source/GeometryAlgorithms/Private/Curve/ClipperUtils.h
--- a/Engine/Plugins/Runtime/GeometryProcessing/Source/GeometryAlgorithms/Private/Curve/ClipperUtils.h
+++ b/Engine/Plugins/Runtime/GeometryProcessing/Source/GeometryAlgorithms/Private/Curve/ClipperUtils.h
@@ -1,4 +1,4 @@
-﻿// Copyright Epic Games, Inc. All Rights Reserved.
+// Copyright Epic Games, Inc. All Rights Reserved.
 
 #pragma once
 
@@ -38,12 +38,9 @@ namespace UE::Geometry::Private
 
 	template <typename InputType, typename RealType>
 	static TPolygon2<RealType> ConvertPathToPolygon(const Clipper2Lib::Path<InputType>& InPath, const TVector2<RealType>& InMin = {}, const RealType& InRange = {});
-
-	template <typename InputType, typename RealType>
-	static void ConvertPolyTreeToPolygon(const Clipper2Lib::PolyPath<InputType>* InPaths, UE::Geometry::TGeneralPolygon2<RealType>& OutPolygon, const TVector2<RealType>& InMin = {}, const RealType& InRange = {});
 	 
 	template <typename InputType, typename RealType>
-	static void ConvertPolyTreeToPolygons(const Clipper2Lib::PolyPath<InputType>* InPaths, TArray<UE::Geometry::TGeneralPolygon2<RealType>>& OutPolygons, const TVector2<RealType>& InMin = {} , const RealType& InRange = {});
+	static void ConvertPolyTreeToPolygons(const Clipper2Lib::PolyPath<InputType>* InPaths, TArray<UE::Geometry::TGeneralPolygon2<RealType>>& OutPolygons, const TVector2<RealType>& InMin = {} , const RealType& InRange = {}, int ParentPolygonIdx = -1);
 }
 
 #include "ClipperUtils.inl"
diff --git a/Engine/Plugins/Runtime/GeometryProcessing/Source/GeometryAlgorithms/Private/Curve/ClipperUtils.inl b/Engine/Plugins/Runtime/GeometryProcessing/Source/GeometryAlgorithms/Private/Curve/ClipperUtils.inl
--- a/Engine/Plugins/Runtime/GeometryProcessing/Source/GeometryAlgorithms/Private/Curve/ClipperUtils.inl
+++ b/Engine/Plugins/Runtime/GeometryProcessing/Source/GeometryAlgorithms/Private/Curve/ClipperUtils.inl
@@ -1,4 +1,4 @@
-﻿// Copyright Epic Games, Inc. All Rights Reserved.
+// Copyright Epic Games, Inc. All Rights Reserved.
 
 #if USING_CODE_ANALYSIS
 #include "Curve/ClipperUtils.h"
@@ -117,51 +117,33 @@ namespace UE::Geometry::Private
 		return Polygon;
 	}
 
+
 	template <typename InputType, typename RealType>
-	void ConvertPolyTreeToPolygon(const Clipper2Lib::PolyPath<InputType>* InPaths, UE::Geometry::TGeneralPolygon2<RealType>& OutPolygon, const TVector2<RealType>& InMin, const RealType& InRange)
+	void ConvertPolyTreeToPolygons(const Clipper2Lib::PolyPath<InputType>* InPaths, TArray<UE::Geometry::TGeneralPolygon2<RealType>>& OutPolygons, const TVector2<RealType>& InMin, const RealType& InRange, int ParentPolygonIdx)
 	{
-		const Clipper2Lib::Path<InputType>& InputPolygon = InPaths->Polygon();
-		TPolygon2<RealType> OutputPolygon = ConvertPathToPolygon<InputType, RealType>(InputPolygon, InMin, InRange);
-
-		if(OutputPolygon.VertexCount() > 0)
+		if (OutPolygons.IsEmpty())
 		{
-			checkf(InPaths->IsHole(), TEXT("ConvertPolyTreeToPolygon should only be used for hole polygons."));
-			// Skip containment/orientation tests as it's already performed by Clipper
-			OutPolygon.AddHole(OutputPolygon, false, false);
+			ParentPolygonIdx = -1;
 		}
 
-		for(const Clipper2Lib::PolyPath<InputType>* const Child : *InPaths)
-		{
-			ConvertPolyTreeToPolygon<InputType, RealType>(Child, OutPolygon, InMin, InRange);
-		}
-	}
-
-	template <typename InputType, typename RealType>
-	void ConvertPolyTreeToPolygons(const Clipper2Lib::PolyPath<InputType>* InPaths, TArray<UE::Geometry::TGeneralPolygon2<RealType>>& OutPolygons, const TVector2<RealType>& InMin, const RealType& InRange)
-	{
 		const Clipper2Lib::Path<InputType>& InputPolygon = InPaths->Polygon();
 		TPolygon2<RealType> OutputPolygon = ConvertPathToPolygon<InputType, RealType>(InputPolygon, InMin, InRange);
 
 		if(OutputPolygon.VertexCount() > 0)
 		{
-			checkf(!InPaths->IsHole(), TEXT("ConvertPolyTreeToPolygons should only be used for outer polygons."));
-			OutPolygons.Emplace_GetRef().SetOuter(OutputPolygon);
-		}
-
-		int32 PolygonIdx = OutPolygons.Num() - 1;
-		if(PolygonIdx >= 0)
-		{
-			for(const Clipper2Lib::PolyPath<InputType>* const ChildPath : *InPaths)
+			if (InPaths->IsHole() && ensure(OutPolygons.IsValidIndex(ParentPolygonIdx)))
 			{
-				ConvertPolyTreeToPolygon<InputType, RealType>(ChildPath, OutPolygons[PolygonIdx], InMin, InRange);
+				OutPolygons[ParentPolygonIdx].AddHole(OutputPolygon, false, false);
+			}
+			else
+			{
+				ParentPolygonIdx = OutPolygons.Emplace(OutputPolygon);
 			}
 		}
-		else
+
+		for(const Clipper2Lib::PolyPath<InputType>* const ChildPath : *InPaths)
 		{
-			for(const Clipper2Lib::PolyPath<InputType>* const ChildPath : *InPaths)
-			{
-				ConvertPolyTreeToPolygons<InputType, RealType>(ChildPath, OutPolygons, InMin, InRange);
-			}			
+			ConvertPolyTreeToPolygons<InputType, RealType>(ChildPath, OutPolygons, InMin, InRange, ParentPolygonIdx);
 		}
 	}
 }
