Add support for taking the intersection/union/difference/xor of arrays of polygons with clipper

diff --git a/Engine/Plugins/Runtime/GeometryProcessing/Source/GeometryAlgorithms/Private/Curve/PolygonIntersectionUtils.cpp b/Engine/Plugins/Runtime/GeometryProcessing/Source/GeometryAlgorithms/Private/Curve/PolygonIntersectionUtils.cpp
--- a/Engine/Plugins/Runtime/GeometryProcessing/Source/GeometryAlgorithms/Private/Curve/PolygonIntersectionUtils.cpp
+++ b/Engine/Plugins/Runtime/GeometryProcessing/Source/GeometryAlgorithms/Private/Curve/PolygonIntersectionUtils.cpp
@@ -40,8 +40,84 @@ namespace UE::Geometry::Private
 
 		return bExecuteResult;
 	}
+
+	// Clip polygons in the subject array by the polygons in the clip array
+	template <typename RealType>
+	static bool ClipArrays(
+		const Clipper2Lib::ClipType& InClipType,
+		TArrayView<const FGeneralPolygon2d> SubjPolygons,
+		TArrayView<const FGeneralPolygon2d> ClipPolygons,
+		TArray<FGeneralPolygon2d>& ResultOut)
+	{
+		ResultOut.Reset();
+
+		TAxisAlignedBox2<RealType> InputBounds;
+		for (const FGeneralPolygon2d& Polygon : SubjPolygons)
+		{
+			InputBounds.Contain(Polygon.GetOuter().Bounds());
+		}
+		for (const FGeneralPolygon2d& Polygon : ClipPolygons)
+		{
+			InputBounds.Contain(Polygon.GetOuter().Bounds());
+		}
+		RealType InputRange = (InputBounds.Extents() * 2).GetMax();
+
+		Clipper2Lib::Clipper64 Clipper;
+		for (const FGeneralPolygon2d& Polygon : SubjPolygons)
+		{
+			Clipper.AddSubject(ConvertGeneralizedPolygonToPath<RealType, IntegralType>(Polygon, InputBounds.Min, InputRange));
+		}
+		for (const FGeneralPolygon2d& Polygon : ClipPolygons)
+		{
+			Clipper.AddClip(ConvertGeneralizedPolygonToPath<RealType, IntegralType>(Polygon, InputBounds.Min, InputRange));
+		}
+
+		Clipper2Lib::PolyTree<IntegralType> ResultTree;
+		Clipper2Lib::Paths64 ResultOpenPaths;
+		const bool bExecuteResult = Clipper.Execute(InClipType, Clipper2Lib::FillRule::NonZero, ResultTree, ResultOpenPaths);
+
+		if (bExecuteResult)
+		{
+			ConvertPolyTreeToPolygons<IntegralType, RealType>(&ResultTree, ResultOut, InputBounds.Min, InputRange);
+		}
+
+		return bExecuteResult;
+	}
 };
 
+
+// Array-based operations
+
+namespace UE::Geometry
+{
+bool PolygonsUnion(TArrayView<const FGeneralPolygon2d> Polygons, TArray<FGeneralPolygon2d>& ResultOut, bool bCopyInputOnFailure)
+{
+	// always put everything in the 'subject' array for unions
+	bool bResult = Private::ClipArrays<double>(Clipper2Lib::ClipType::Union, Polygons, TArrayView<const FGeneralPolygon2d>(), ResultOut);
+	// optionally fall back to copying the input polygons on failure
+	if (bCopyInputOnFailure && !bResult)
+	{
+		ResultOut.Append(Polygons);
+	}
+	return bResult;
+}
+bool PolygonsDifference(TArrayView<const FGeneralPolygon2d> SubjPolygons, TArrayView<const FGeneralPolygon2d> ClipPolygons, TArray<FGeneralPolygon2d>& ResultOut)
+{
+	return Private::ClipArrays<double>(Clipper2Lib::ClipType::Difference, SubjPolygons, ClipPolygons, ResultOut);
+}
+bool PolygonsIntersection(TArrayView<const FGeneralPolygon2d> SubjPolygons, TArrayView<const FGeneralPolygon2d> ClipPolygons, TArray<FGeneralPolygon2d>& ResultOut)
+{
+	return Private::ClipArrays<double>(Clipper2Lib::ClipType::Intersection, SubjPolygons, ClipPolygons, ResultOut);
+}
+bool PolygonsExclusiveOr(TArrayView<const FGeneralPolygon2d> SubjPolygons, TArrayView<const FGeneralPolygon2d> ClipPolygons, TArray<FGeneralPolygon2d>& ResultOut)
+{
+	return Private::ClipArrays<double>(Clipper2Lib::ClipType::Xor, SubjPolygons, ClipPolygons, ResultOut);
+}
+} // namespace UE::Geometry
+
+
+// Two-polygon operations
+
 template <EPolygonBooleanOp OperationType, typename GeometryType, typename RealType>
 TBooleanPolygon2Polygon2<OperationType, GeometryType, RealType>::TBooleanPolygon2Polygon2(
 	const GeometryType& InPolygonA, 
diff --git a/Engine/Plugins/Runtime/GeometryProcessing/Source/GeometryAlgorithms/Public/Curve/PolygonIntersectionUtils.h b/Engine/Plugins/Runtime/GeometryProcessing/Source/GeometryAlgorithms/Public/Curve/PolygonIntersectionUtils.h
--- a/Engine/Plugins/Runtime/GeometryProcessing/Source/GeometryAlgorithms/Public/Curve/PolygonIntersectionUtils.h
+++ b/Engine/Plugins/Runtime/GeometryProcessing/Source/GeometryAlgorithms/Public/Curve/PolygonIntersectionUtils.h
@@ -67,4 +67,19 @@ namespace UE::Geometry
 
 	typedef TExclusiveOrPolygon2Polygon2<TGeneralPolygon2<float>, float> FExclusiveOrGeneralPolygon2GeneralPolygon2f;
 	typedef TExclusiveOrPolygon2Polygon2<TGeneralPolygon2<double>, double> FExclusiveOrGeneralPolygon2GeneralPolygon2d;
+
+	// Array-based interfaces to support cases where we have different numbers of polygons to operate on
+	
+	/// Populate ResultOut with the union of the input Polygons
+	/// @param bCopyInputOnFailure		If true, copy the input Polygons to ResultOut if the union is not successfully computed. Otherwise, ResultOut is left empty on failure.
+	GEOMETRYALGORITHMS_API bool PolygonsUnion(TArrayView<const FGeneralPolygon2d> Polygons, TArray<FGeneralPolygon2d>& ResultOut, bool bCopyInputOnFailure);
+
+	/// Populate ResultOut with the difference of PosPolygons minus NegPolygons
+	GEOMETRYALGORITHMS_API bool PolygonsDifference(TArrayView<const FGeneralPolygon2d> PosPolygons, TArrayView<const FGeneralPolygon2d> NegPolygons, TArray<FGeneralPolygon2d>& ResultOut);
+
+	/// Populate ResultOut with the intersection of SubjPolygons with ClipPolygons
+	GEOMETRYALGORITHMS_API bool PolygonsIntersection(TArrayView<const FGeneralPolygon2d> SubjPolygons, TArrayView<const FGeneralPolygon2d> ClipPolygons, TArray<FGeneralPolygon2d>& ResultOut);
+
+	/// Populate ResultOut with the 'exclusive or' of SubjPolygons with ClipPolygons
+	GEOMETRYALGORITHMS_API bool PolygonsExclusiveOr(TArrayView<const FGeneralPolygon2d> SubjPolygons, TArrayView<const FGeneralPolygon2d> ClipPolygons, TArray<FGeneralPolygon2d>& ResultOut);
 }
