Add polygon offset functions that take arrays of general polygons and apply one or two offsets

also:
Update polygon union code to allow the input array to be (optionally) re-used as the result array
Fix the previous offset function to support negative offsets for polygons

diff --git a/Engine/Plugins/Runtime/GeometryProcessing/Source/GeometryAlgorithms/Private/Curve/ClipperUtils.h b/Engine/Plugins/Runtime/GeometryProcessing/Source/GeometryAlgorithms/Private/Curve/ClipperUtils.h
--- a/Engine/Plugins/Runtime/GeometryProcessing/Source/GeometryAlgorithms/Private/Curve/ClipperUtils.h
+++ b/Engine/Plugins/Runtime/GeometryProcessing/Source/GeometryAlgorithms/Private/Curve/ClipperUtils.h
@@ -36,6 +36,9 @@ namespace UE::Geometry::Private
 	template <typename RealType, typename OutputType>
 	static Clipper2Lib::Paths<OutputType> ConvertGeneralizedPolygonToPath(const UE::Geometry::TGeneralPolygon2<RealType>& InPolygon, const TVector2<RealType>& InMin = {}, const RealType& InRange = {});
 
+	template <typename RealType, typename OutputType>
+	static void AddGeneralizedPolygonToPaths(Clipper2Lib::Paths<OutputType>& OutPaths, const UE::Geometry::TGeneralPolygon2<RealType>& InPolygon, const TVector2<RealType>& InMin = {}, const RealType& InRange = {});
+
 	template <typename InputType, typename RealType>
 	static TPolygon2<RealType> ConvertPathToPolygon(const Clipper2Lib::Path<InputType>& InPath, const TVector2<RealType>& InMin = {}, const RealType& InRange = {});
 	 
diff --git a/Engine/Plugins/Runtime/GeometryProcessing/Source/GeometryAlgorithms/Private/Curve/ClipperUtils.inl b/Engine/Plugins/Runtime/GeometryProcessing/Source/GeometryAlgorithms/Private/Curve/ClipperUtils.inl
--- a/Engine/Plugins/Runtime/GeometryProcessing/Source/GeometryAlgorithms/Private/Curve/ClipperUtils.inl
+++ b/Engine/Plugins/Runtime/GeometryProcessing/Source/GeometryAlgorithms/Private/Curve/ClipperUtils.inl
@@ -86,15 +86,20 @@ namespace UE::Geometry::Private
 	}
 
 	template <typename RealType, typename OutputType>
-	Clipper2Lib::Paths<OutputType> ConvertGeneralizedPolygonToPath(const UE::Geometry::TGeneralPolygon2<RealType>& InPolygon, const TVector2<RealType>& InMin, const RealType& InRange)
+	void AddGeneralizedPolygonToPaths(Clipper2Lib::Paths<OutputType>& OutPaths, const UE::Geometry::TGeneralPolygon2<RealType>& InPolygon, const TVector2<RealType>& InMin, const RealType& InRange)
 	{
-		Clipper2Lib::Paths<OutputType> Paths;
-		Paths.reserve(1 + InPolygon.GetHoles().Num());
-		Paths.push_back(ConvertPolygonToPath<RealType, OutputType>(InPolygon.GetOuter(), InMin, InRange));
+		OutPaths.reserve(OutPaths.size() + 1 + InPolygon.GetHoles().Num());
+		OutPaths.push_back(ConvertPolygonToPath<RealType, OutputType>(InPolygon.GetOuter(), InMin, InRange));
 		for(const TPolygon2<RealType>& HolePath : InPolygon.GetHoles())
 		{
-			Paths.push_back(ConvertPolygonToPath<RealType, OutputType>(HolePath, InMin, InRange));
+			OutPaths.push_back(ConvertPolygonToPath<RealType, OutputType>(HolePath, InMin, InRange));
 		}
+	}
+	template <typename RealType, typename OutputType>
+	Clipper2Lib::Paths<OutputType> ConvertGeneralizedPolygonToPath(const UE::Geometry::TGeneralPolygon2<RealType>& InPolygon, const TVector2<RealType>& InMin, const RealType& InRange)
+	{
+		Clipper2Lib::Paths<OutputType> Paths;
+		AddGeneralizedPolygonToPaths(Paths, InPolygon, InMin, InRange);
 		return Paths;
 	}
 
diff --git a/Engine/Plugins/Runtime/GeometryProcessing/Source/GeometryAlgorithms/Private/Curve/PolygonIntersectionUtils.cpp b/Engine/Plugins/Runtime/GeometryProcessing/Source/GeometryAlgorithms/Private/Curve/PolygonIntersectionUtils.cpp
--- a/Engine/Plugins/Runtime/GeometryProcessing/Source/GeometryAlgorithms/Private/Curve/PolygonIntersectionUtils.cpp
+++ b/Engine/Plugins/Runtime/GeometryProcessing/Source/GeometryAlgorithms/Private/Curve/PolygonIntersectionUtils.cpp
@@ -49,8 +49,6 @@ namespace UE::Geometry::Private
 		TArrayView<const FGeneralPolygon2d> ClipPolygons,
 		TArray<FGeneralPolygon2d>& ResultOut)
 	{
-		ResultOut.Reset();
-
 		TAxisAlignedBox2<RealType> InputBounds;
 		for (const FGeneralPolygon2d& Polygon : SubjPolygons)
 		{
@@ -78,6 +76,7 @@ namespace UE::Geometry::Private
 
 		if (bExecuteResult)
 		{
+			ResultOut.Reset();
 			ConvertPolyTreeToPolygons<IntegralType, RealType>(&ResultTree, ResultOut, InputBounds.Min, InputRange);
 		}
 
@@ -95,8 +94,9 @@ bool PolygonsUnion(TArrayView<const FGeneralPolygon2d> Polygons, TArray<FGeneral
 	// always put everything in the 'subject' array for unions
 	bool bResult = Private::ClipArrays<double>(Clipper2Lib::ClipType::Union, Polygons, TArrayView<const FGeneralPolygon2d>(), ResultOut);
 	// optionally fall back to copying the input polygons on failure
-	if (bCopyInputOnFailure && !bResult)
+	if (bCopyInputOnFailure && !bResult && ResultOut.GetData() != Polygons.GetData())
 	{
+		ResultOut.Reset(Polygons.Num());
 		ResultOut.Append(Polygons);
 	}
 	return bResult;
diff --git a/Engine/Plugins/Runtime/GeometryProcessing/Source/GeometryAlgorithms/Private/Curve/PolygonOffsetUtils.cpp b/Engine/Plugins/Runtime/GeometryProcessing/Source/GeometryAlgorithms/Private/Curve/PolygonOffsetUtils.cpp
--- a/Engine/Plugins/Runtime/GeometryProcessing/Source/GeometryAlgorithms/Private/Curve/PolygonOffsetUtils.cpp
+++ b/Engine/Plugins/Runtime/GeometryProcessing/Source/GeometryAlgorithms/Private/Curve/PolygonOffsetUtils.cpp
@@ -32,9 +32,12 @@ namespace UE::Geometry::Private
 		}
 
 		Clipper2Lib::ClipperOffset ClipperOffset;
-		ClipperOffset.ReverseSolution(FMath::Sign(InOffset) < 0);
 
-		InputBounds.ExpandBy(FMath::Max(0, InOffset * InOffset)); // pad to allow for offset geometry
+		bool bOneSidedOffset = InEndType == EPolygonOffsetEndType::Polygon;
+
+		RealType OffsetExpand = bOneSidedOffset ? FMath::Max(0, InOffset) : FMath::Abs(InOffset);
+		RealType MaxMiterExpand = InJoinType == EPolygonOffsetJoinType::Miter ? InMiterLimit : 0;
+		InputBounds.ExpandBy(OffsetExpand + MaxMiterExpand); // pad to allow for offset geometry
 		RealType InputRange = InputBounds.GetSize().GetMax();
 		
 		const Clipper2Lib::Paths<IntegralType> Paths = ConvertPolygonsToPaths<RealType, IntegralType>(InPolygons, InputBounds.Min, InputRange);
@@ -44,11 +47,13 @@ namespace UE::Geometry::Private
 		ClipperOffset.MiterLimit(InMiterLimit);
 
 		RealType ScaledOffset = FMath::Floor(static_cast<RealType>((InOffset / InputRange) * static_cast<RealType>(IntRange))); // scale to account for value normalization
-		Clipper2Lib::Paths64 OffsetResultPaths = ClipperOffset.Execute(FMath::Abs(ScaledOffset));
+		Clipper2Lib::Paths64 OffsetResultPaths = ClipperOffset.Execute(ScaledOffset);
 
 		if(OffsetResultPaths.size() <= 0)
 		{
-			return false;
+			// if the offset removed everything, result should be empty and there's nothing else to do
+			OutResult.Reset();
+			return true;
 		}
 
 		bool bExecuteResult = true;
@@ -65,16 +70,108 @@ namespace UE::Geometry::Private
 
 			if(bExecuteResult)
 			{
+				OutResult.Reset();
 				ConvertPolyTreeToPolygons<IntegralType, RealType>(&UnionResultPolyTree, OutResult, InputBounds.Min, InputRange);
 			}
 		}
 		else
 		{
+			OutResult.Reset();
 			OutResult.Add(ConvertPathToPolygon(OffsetResultPaths.back(), InputBounds.Min, InputRange));
 		}
 
 		return bExecuteResult;
 	}
+
+	template <typename RealType>
+	static bool	ApplyOffsetsToGeneralPolygons(
+		const TArrayView<const UE::Geometry::TGeneralPolygon2<RealType>> InPolygons,
+		const EPolygonOffsetJoinType& InJoinType,
+		const EPolygonOffsetEndType& InEndType,
+		TArrayView<const RealType> InOffsets,
+		const RealType& InMiterLimit,
+		TArray<UE::Geometry::TGeneralPolygon2<RealType>>& OutResult,
+		bool bCopyInputOnFailure)
+	{
+		// Get combined bounds (min, max) of points 
+		FAxisAlignedBox2d AccumBounds = FAxisAlignedBox2d::Empty();
+		for (const FGeneralPolygon2d& Polygon : InPolygons)
+		{
+			AccumBounds.Contain(Polygon.Bounds());
+		}
+
+		bool bOneSidedOffset = InEndType == EPolygonOffsetEndType::Polygon;
+
+		RealType MaxExpandBounds = 0;
+		for (const RealType& Offset : InOffsets)
+		{
+			MaxExpandBounds = FMath::Max(MaxExpandBounds, bOneSidedOffset ? FMath::Max(0, Offset) : FMath::Abs(Offset));
+		}
+		RealType MaxMiterExpand = InJoinType == EPolygonOffsetJoinType::Miter ? InMiterLimit : 0;
+		AccumBounds.Expand(MaxExpandBounds + MaxMiterExpand); // pad to allow for offset geometry
+		RealType InputRange = (AccumBounds.Extents() * 2).GetMax();
+
+
+		Clipper2Lib::Paths64 CurrentPaths;
+		for (const FGeneralPolygon2d& Polygon : InPolygons)
+		{
+			AddGeneralizedPolygonToPaths<RealType, IntegralType>(CurrentPaths, Polygon, AccumBounds.Min, InputRange);
+		}
+
+		int NumOffsets = InOffsets.Num();
+		for (int OffsetIdx = 0; OffsetIdx < NumOffsets; ++OffsetIdx)
+		{
+			Clipper2Lib::ClipperOffset ClipperOffset;
+			ClipperOffset.MergeGroups(OffsetIdx + 1 != NumOffsets); // This disables union clipping for the final offset, so we can perform it later to get a polytree
+			ClipperOffset.AddPaths(CurrentPaths, static_cast<Clipper2Lib::JoinType>(InJoinType), static_cast<Clipper2Lib::EndType>(InEndType));
+			ClipperOffset.MiterLimit(InMiterLimit);
+
+			RealType ScaledOffset = FMath::Floor(static_cast<RealType>((InOffsets[OffsetIdx] / InputRange) * static_cast<RealType>(IntRange))); // scale to account for value normalization
+			CurrentPaths = ClipperOffset.Execute(ScaledOffset);
+
+			if (CurrentPaths.size() <= 0)
+			{
+				// if the offset removed everything, result should be empty and there's nothing else to do
+				OutResult.Reset();
+				return true;
+			}
+		}
+
+		
+		bool bExecuteResult = true;
+		if (CurrentPaths.size() > 1)
+		{
+			// ...then union to merge and get polytree
+			Clipper2Lib::Clipper64 Clipper;
+			Clipper.PreserveCollinear = false;
+			Clipper.AddSubject(CurrentPaths);
+
+			Clipper2Lib::PolyTree64 UnionResultPolyTree;
+			Clipper2Lib::Paths64 UnionResultOpenPaths;
+			bExecuteResult = Clipper.Execute(Clipper2Lib::ClipType::Union, Clipper2Lib::FillRule::NonZero, UnionResultPolyTree, UnionResultOpenPaths);
+
+			if (bExecuteResult)
+			{
+				OutResult.Reset();
+				ConvertPolyTreeToPolygons<IntegralType, RealType>(&UnionResultPolyTree, OutResult, AccumBounds.Min, InputRange);
+			}
+			else
+			{
+				if (bCopyInputOnFailure && OutResult.GetData() != InPolygons.GetData())
+				{
+					OutResult.Reset();
+					OutResult.Append(InPolygons);
+				}
+			}
+		}
+		else
+		{
+			OutResult.Reset();
+			OutResult.Add(ConvertPathToPolygon(CurrentPaths.back(), AccumBounds.Min, InputRange));
+		}
+
+		return bExecuteResult;
+	}
 }; 
 
 template <typename GeometryType, typename RealType>
@@ -96,4 +193,36 @@ namespace UE::Geometry
 {
 	template class GEOMETRYALGORITHMS_API TOffsetPolygon2<TArrayView<TVector2<float>>, float>;
 	template class GEOMETRYALGORITHMS_API TOffsetPolygon2<TArrayView<TVector2<double>>, double>;
-}
+
+
+	bool PolygonsOffset(double Offset, 
+		TArrayView<const FGeneralPolygon2d> Polygons, TArray<FGeneralPolygon2d>& ResultOut,
+		bool bCopyInputOnFailure,
+		double MiterLimit,
+		EPolygonOffsetJoinType JoinType,
+		EPolygonOffsetEndType EndType)
+	{
+		using namespace UE::Geometry::Private;
+
+		TArrayView<const double> OffsetsView(&Offset, 1);
+		bool bSuccess = ApplyOffsetsToGeneralPolygons<double>(Polygons, JoinType, EndType, OffsetsView, MiterLimit, ResultOut, bCopyInputOnFailure);
+
+		return bSuccess;
+	}
+
+	bool PolygonsOffsets(double FirstOffset, double SecondOffset, 
+		TArrayView<const FGeneralPolygon2d> Polygons, TArray<FGeneralPolygon2d>& ResultOut,
+		bool bCopyInputOnFailure,
+		double MiterLimit,
+		EPolygonOffsetJoinType JoinType,
+		EPolygonOffsetEndType EndType)
+	{
+		using namespace UE::Geometry::Private;
+
+		double Offsets[2]{ FirstOffset, SecondOffset };
+		TArrayView<const double> OffsetsView(Offsets, 2);
+		bool bSuccess = ApplyOffsetsToGeneralPolygons<double>(Polygons, JoinType, EndType, OffsetsView, MiterLimit, ResultOut, bCopyInputOnFailure);
+
+		return bSuccess;
+	}
+}
\ No newline at end of file
diff --git a/Engine/Plugins/Runtime/GeometryProcessing/Source/GeometryAlgorithms/Public/Curve/PolygonOffsetUtils.h b/Engine/Plugins/Runtime/GeometryProcessing/Source/GeometryAlgorithms/Public/Curve/PolygonOffsetUtils.h
--- a/Engine/Plugins/Runtime/GeometryProcessing/Source/GeometryAlgorithms/Public/Curve/PolygonOffsetUtils.h
+++ b/Engine/Plugins/Runtime/GeometryProcessing/Source/GeometryAlgorithms/Public/Curve/PolygonOffsetUtils.h
@@ -45,4 +45,41 @@ namespace UE::Geometry
 
 	typedef TOffsetPolygon2<TArrayView<TVector2<float>>, float> FOffsetPolygon2f;
 	typedef TOffsetPolygon2<TArrayView<TVector2<double>>, double> FOffsetPolygon2d;
+
+
+	// Offset an array of general polygons.
+	// Note: Input polygons should not overlap; to process overlapping polygons, call PolygonsUnion first.
+	// @param Offset		The amount to offset the polygons.
+	// @param Polygons		The polygons to offset
+	// @param ResultOut		On success, will be filled with the resulting polygons. Note: It is ok to pass the same array as both the input and output.
+	// @param bCopyInputOnFailure Whether to copy the input polygons to the result on failure.
+	// @param MiterLimit	If JoinType is Miter, this limits how far each miter join can extend
+	// @param JoinType		How to join/extend corners between two edges
+	// @param EndType		Should generally be Polygon. Otherwise, input is interpreted as a path, and offset is a 'thickening' applied to each side.
+	GEOMETRYALGORITHMS_API bool PolygonsOffset(
+		double Offset,
+		TArrayView<const FGeneralPolygon2d> Polygons, TArray<FGeneralPolygon2d>& ResultOut,
+		bool bCopyInputOnFailure,
+		double MiterLimit,
+		EPolygonOffsetJoinType JoinType = EPolygonOffsetJoinType::Square,
+		EPolygonOffsetEndType EndType = EPolygonOffsetEndType::Polygon);
+
+	// Perform two offsets of the input polygons, intended to support "opening" or "closing" operations
+	// Note: Input polygons should not overlap; to process overlapping polygons, call PolygonsUnion first.
+	// @param FirstOffset	The initial amount to offset the polygons.
+	// @param SecondOffset	The second amount to offset the polygons.
+	// @param Polygons		The polygons to offset
+	// @param ResultOut		On success, will be filled with the resulting polygons. Note: It is ok to pass the same array as both the input and output.
+	// @param bCopyInputOnFailure Whether to copy the input polygons to the result on failure.
+	// @param MiterLimit	If JoinType is Miter, this limits how far each miter join can extend
+	// @param JoinType		How to join/extend corners between two edges
+	// @param EndType		Should generally be Polygon. Otherwise, input is interpreted as a path, and offset is a 'thickening' applied to each side.
+	GEOMETRYALGORITHMS_API bool PolygonsOffsets(
+		double FirstOffset,
+		double SecondOffset,
+		TArrayView<const FGeneralPolygon2d> Polygons, TArray<FGeneralPolygon2d>& ResultOut,
+		bool bCopyInputOnFailure,
+		double MiterLimit,
+		EPolygonOffsetJoinType JoinType = EPolygonOffsetJoinType::Square,
+		EPolygonOffsetEndType EndType = EPolygonOffsetEndType::Polygon);
 }
