CombineMeshInstance: add support for combining simple collision and simple box-merging strategy to combine collision shapes.
GeometryCore: Make TAxisAlignedBox3::DistanceSquared const (should have been). Add TQuaternion * scalar and operator+ functions. Add TTransformSRT3::GetScale3D for API compatibility with engine FTransform
Add FSimpleShapeSet3d::Append() variant that takes a FTransformSequence3d

diff --git a/Engine/Plugins/Experimental/MeshModelingToolsetExp/Source/GeometryProcessingAdapters/Private/CombineMeshInstancesImpl.cpp b/Engine/Plugins/Experimental/MeshModelingToolsetExp/Source/GeometryProcessingAdapters/Private/CombineMeshInstancesImpl.cpp
--- a/Engine/Plugins/Experimental/MeshModelingToolsetExp/Source/GeometryProcessingAdapters/Private/CombineMeshInstancesImpl.cpp
+++ b/Engine/Plugins/Experimental/MeshModelingToolsetExp/Source/GeometryProcessingAdapters/Private/CombineMeshInstancesImpl.cpp
@@ -29,6 +29,9 @@
 #include "DynamicMesh/Operations/MergeCoincidentMeshEdges.h"
 #include "Operations/RemoveOccludedTriangles.h"
 
+#include "Physics/CollisionGeometryConversion.h"
+#include "Physics/PhysicsDataCollection.h"
+
 #include "TransformSequence.h"
 
 
@@ -76,7 +79,7 @@ struct FMeshInstanceSet
 struct FSourceGeometry
 {
 	TArray<UE::Geometry::FDynamicMesh3> SourceMeshLODs;
-	//UE::Geometry::FSimpleShapeSet3d CollisionShapes;
+	UE::Geometry::FSimpleShapeSet3d CollisionShapes;
 };
 
 struct FOptimizedGeometry
@@ -251,9 +254,27 @@ void InitializeAssemblySourceMeshesFromLOD(
 				Target.SourceMeshLODs[k] = Target.SourceMeshLODs[k-1];
 			}
 		}
+
 	});
 
 
+	// not clear that it is safe to do this in parallel...
+	for (int32 Index = 0; Index < NumSets; ++Index)
+	{
+		TUniquePtr<FMeshInstanceSet>& InstanceSet = Assembly.InstanceSets[Index];
+		FSourceGeometry& Target = Assembly.SourceMeshGeometry[Index];
+
+		UStaticMesh* StaticMesh = InstanceSet->SourceAsset;
+		UBodySetup* BodySetup = StaticMesh->GetBodySetup();
+		if (BodySetup)
+		{
+			UE::Geometry::GetShapeSet(BodySetup->AggGeom, Target.CollisionShapes);
+			// todo: detect boxes?
+		}
+	}
+
+
+
 	//ParallelFor(NumSets, [&](int32 Index)
 	//{
 	//	TUniquePtr<FMeshInstanceSet>& InstanceSet = Assembly.InstanceSets[Index];
@@ -914,12 +935,166 @@ void BuildCombinedMesh(
 	{
 		CombinedMeshLODs.Add(MoveTemp(MeshLODs[LODLevel].Mesh));
 	}
+}
+
+
+
+
+
+
+/**
+ * Construct a new OrientedBox that contains both A and B. The main problem is to
+ * determine the new Orientation, this is done by a 0.5 slerp of the orientations of A and B.
+ * The new local Origin and Extents are then computed in this new orientation.
+ */
+FOrientedBox3d MergeBoxes(const FOrientedBox3d& A, const FOrientedBox3d& B)
+{
+	FOrientedBox3d NewBox;
+	NewBox.Frame.Origin = (A.Center() + B.Center()) * 0.5;
+
+	FQuaterniond RotationA(A.Frame.Rotation), RotationB(B.Frame.Rotation);
+	if (RotationA.Dot(RotationB) < 0)
+	{
+		RotationB = -RotationB;
+	}
+
+	// this is just a slerp?
+	FQuaterniond HalfRotation = RotationA + RotationB;
+	HalfRotation.Normalize();
+	NewBox.Frame.Rotation = HalfRotation;
+
+	// likely faster to compute the frame X/Y/Z instead of calling ToFramePoint each time...
+	FAxisAlignedBox3d LocalBounds(FVector3d::Zero(), FVector3d::Zero());
+	A.EnumerateCorners([&](FVector3d P)
+	{
+		LocalBounds.Contain( NewBox.Frame.ToFramePoint(P) );
+	});
+	B.EnumerateCorners([&](FVector3d P)
+	{
+		LocalBounds.Contain( NewBox.Frame.ToFramePoint(P) );
+	});
+
+	// update origin and extents
+	NewBox.Frame.Origin = NewBox.Frame.FromFramePoint( LocalBounds.Center() );
+	NewBox.Extents = 0.5 * LocalBounds.Diagonal();
+
+	return NewBox;
+}
+
+
+
+static void CombineCollisionShapes(
+	FSimpleShapeSet3d& CollisionShapes,
+	double AxisToleranceDelta = 0.01)
+{
+	// only going to merge boxes for now
+	TArray<FOrientedBox3d> Boxes;
+	for (FBoxShape3d Box : CollisionShapes.Boxes)
+	{
+		Boxes.Add(Box.Box);
+	}
+
+	// want to merge larger-volume boxes first
+	Boxes.Sort([&](const FOrientedBox3d& A, const FOrientedBox3d& B)
+	{
+		return A.Volume() > B.Volume();
+	});
+
+	auto CalcOffsetVolume = [](FOrientedBox3d Box, double AxisDelta) -> double
+	{
+		Box.Extents.X = FMathd::Max(0, Box.Extents.X+AxisDelta);
+		Box.Extents.Y = FMathd::Max(0, Box.Extents.Y+AxisDelta);
+		Box.Extents.Z = FMathd::Max(0, Box.Extents.Z+AxisDelta);
+		return Box.Volume();
+	};
+
+	double DotTol = 0.99;
+	auto HasMatchingAxis = [DotTol](const FVector3d& Axis, const FOrientedBox3d& Box)
+	{
+		for (int32 k = 0; k < 3; ++k)
+		{
+			if (FMathd::Abs(Axis.Dot(Box.GetAxis(k))) > DotTol)
+			{
+				return true;
+			}
+		}
+		return false;
+	};
+
+	bool bFoundMerge = true;
+	while (bFoundMerge)
+	{
+		bFoundMerge = false;
+
+		int32 N = Boxes.Num();
+		for (int32 i = 0; i < N; ++i)
+		{
+			FOrientedBox3d Box1 = Boxes[i];
+
+			for (int32 j = i + 1; j < N; ++j)
+			{
+				FOrientedBox3d Box2 = Boxes[j];
+
+				// should we just be appending box2 to Box1? prevents getting skewed boxes...
+				FOrientedBox3d NewBox = MergeBoxes(Box1, Box2);
+
+				// check if newbox is still aligned w/ box2?
+				bool bAllAxesAligned = true;
+				for (int32 k = 0; k < 3; ++k)
+				{
+					bAllAxesAligned = bAllAxesAligned && HasMatchingAxis(Box1.GetAxis(k), NewBox) && HasMatchingAxis(Box2.GetAxis(k), NewBox);
+				}
+				if (!bAllAxesAligned)
+				{
+					continue;
+				}
+
+				double SumVolume = Box1.Volume() + Box2.Volume();
+				if ( (CalcOffsetVolume(NewBox, AxisToleranceDelta) > SumVolume) &&
+						(CalcOffsetVolume(NewBox, -AxisToleranceDelta) < SumVolume) )
+				{
+					bFoundMerge = true;
+					Boxes[i] = NewBox;
+					Boxes.RemoveAtSwap(j);
+					j = N;
+					N--;
+				}
+			}
+		}
+	}
 
+	CollisionShapes.Boxes.Reset();
+	for (FOrientedBox3d Box : Boxes)
+	{
+		CollisionShapes.Boxes.Add(FBoxShape3d(Box));
+	}
 }
 
 
+void BuildCombinedCollisionShapes(
+	const FMeshInstanceAssembly& Assembly,
+	IGeometryProcessing_CombineMeshInstances::FOptions CombineOptions,
+	FSimpleShapeSet3d& CombinedCollisionShapes)
+{
+	int32 NumSets = Assembly.InstanceSets.Num();
 
+	for ( int32 SetIndex = 0; SetIndex < NumSets; ++SetIndex )
+	{
+		const TUniquePtr<FMeshInstanceSet>& InstanceSet = Assembly.InstanceSets[SetIndex];
+		const FSourceGeometry& SourceGeometry = Assembly.SourceMeshGeometry[SetIndex];
+		for ( const FMeshInstance& Instance : InstanceSet->Instances )
+		{
+			bool bIsDecorativePart = (Instance.DetailLevel == EMeshDetailLevel::Decorative);
+			if ( ! bIsDecorativePart )
+			{
+				CombinedCollisionShapes.Append( SourceGeometry.CollisionShapes, Instance.WorldTransform );
+			}
+		}
+	}
 
+	// trivially merge any adjacent boxes that merge to a perfect combined-box
+	CombineCollisionShapes(CombinedCollisionShapes, 0.01);
+}
 
 
 
@@ -1008,8 +1183,16 @@ void FCombineMeshInstancesImpl::CombineMeshInstances(
 	TArray<UE::Geometry::FDynamicMesh3> CombinedMeshLODs;
 	BuildCombinedMesh(InstanceAssembly, Options, CombinedMeshLODs);
 
+	FSimpleShapeSet3d CombinedCollisionShapes;
+	BuildCombinedCollisionShapes(InstanceAssembly, Options, CombinedCollisionShapes);
+	FPhysicsDataCollection PhysicsData;
+	PhysicsData.Geometry = CombinedCollisionShapes;
+	PhysicsData.CopyGeometryToAggregate();		// need FPhysicsDataCollection to convert to agg geom, should fix this
+
+
 	ResultsOut.CombinedMeshes.SetNum(1);
 	IGeometryProcessing_CombineMeshInstances::FOutputMesh& OutputMesh = ResultsOut.CombinedMeshes[0];
 	OutputMesh.MeshLODs = MoveTemp(CombinedMeshLODs);
 	OutputMesh.MaterialSet = InstanceAssembly.UniqueMaterials;
+	OutputMesh.SimpleCollisionShapes = PhysicsData.AggGeom;
 }
\ No newline at end of file
diff --git a/Engine/Plugins/Runtime/GeometryProcessing/Source/DynamicMesh/Private/ShapeApproximation/SimpleShapeSet3.cpp b/Engine/Plugins/Runtime/GeometryProcessing/Source/DynamicMesh/Private/ShapeApproximation/SimpleShapeSet3.cpp
--- a/Engine/Plugins/Runtime/GeometryProcessing/Source/DynamicMesh/Private/ShapeApproximation/SimpleShapeSet3.cpp
+++ b/Engine/Plugins/Runtime/GeometryProcessing/Source/DynamicMesh/Private/ShapeApproximation/SimpleShapeSet3.cpp
@@ -379,14 +379,15 @@ void FSimpleShapeSet3d::RemoveContainedGeometry()
 	LevelSets = MoveTemp(NewSet.LevelSets);
 }
 
-
-static void TransformSphereShape(FSphereShape3d& SphereShape, const FTransform3d& Transform)
+template<typename TransformType>
+static void TransformSphereShape(FSphereShape3d& SphereShape, const TransformType& Transform)
 {
 	double RadiusScale = Transform.GetScale3D().Length() / FVector3d::One().Length();
 	SphereShape.Sphere.Center = Transform.TransformPosition(SphereShape.Sphere.Center);
 	SphereShape.Sphere.Radius *= RadiusScale;
 }
-static void TransformBoxShape(FBoxShape3d& BoxShape, const FTransform3d& Transform)
+template<typename TransformType>
+static void TransformBoxShape(FBoxShape3d& BoxShape, const TransformType& Transform)
 {
 	FVector3d CornerVec = BoxShape.Box.Frame.PointAt(BoxShape.Box.Extents) - BoxShape.Box.Frame.Origin;
 	BoxShape.Box.Frame.Transform(Transform);
@@ -395,7 +396,8 @@ static void TransformBoxShape(FBoxShape3d& BoxShape, const FTransform3d& Transfo
 	BoxShape.Box.Extents.Y = CornerVec.Dot(BoxShape.Box.AxisY());
 	BoxShape.Box.Extents.Z = CornerVec.Dot(BoxShape.Box.AxisZ());
 }
-static void TransformCapsuleShape(UE::Geometry::FCapsuleShape3d& CapsuleShape, const FTransform3d& Transform)
+template<typename TransformType>
+static void TransformCapsuleShape(UE::Geometry::FCapsuleShape3d& CapsuleShape, const TransformType& Transform)
 {
 	FVector3d P0 = Transform.TransformPosition(CapsuleShape.Capsule.Segment.StartPoint());
 	FVector3d P1 = Transform.TransformPosition(CapsuleShape.Capsule.Segment.EndPoint());
@@ -412,28 +414,30 @@ static void TransformCapsuleShape(UE::Geometry::FCapsuleShape3d& CapsuleShape, c
 	CapsuleShape.Capsule.Radius *= RadiusScale;
 }
 
-static void TransformLevelSetShape(UE::Geometry::FLevelSetShape3d& LevelSetShape, const FTransform3d& Transform)
+template<typename TransformType>
+static void TransformLevelSetShape(UE::Geometry::FLevelSetShape3d& LevelSetShape, const TransformType& Transform)
 {
 	const FTransform3d Sub(-0.5 * LevelSetShape.CellSize * FVector3d::One());
 	const FTransform3d Add(0.5 * LevelSetShape.CellSize * FVector3d::One());
-	LevelSetShape.GridTransform = LevelSetShape.GridTransform * Sub * Transform * Add;
+	LevelSetShape.GridTransform = LevelSetShape.GridTransform * Sub * (FTransform3d)Transform * Add;
 }
 
 
-
-static void TransformSphereShape(FSphereShape3d& SphereShape, const TArray<FTransform3d>& TransformSequence)
+template<typename TransformArrayType>
+static void TransformSphereShapeByArray(FSphereShape3d& SphereShape, const TransformArrayType& TransformSequence)
 {
-	for (const FTransform3d& XForm : TransformSequence)
+	for (const auto& XForm : TransformSequence)
 	{
 		SphereShape.Sphere.Center = XForm.TransformPosition(SphereShape.Sphere.Center);
 		double RadiusScale = XForm.GetScale3D().Length() / FVector3d::One().Length();
 		SphereShape.Sphere.Radius *= RadiusScale;
 	}
 }
-static void TransformBoxShape(FBoxShape3d& BoxShape, const TArray<FTransform3d>& TransformSequence)
+template<typename TransformArrayType>
+static void TransformBoxShapeByArray(FBoxShape3d& BoxShape, const TransformArrayType& TransformSequence)
 {
 	FVector3d CornerVec = BoxShape.Box.Frame.PointAt(BoxShape.Box.Extents) - BoxShape.Box.Frame.Origin;
-	for (const FTransform3d& XForm : TransformSequence)
+	for (const auto& XForm : TransformSequence)
 	{
 		BoxShape.Box.Frame.Transform(XForm);
 		CornerVec = XForm.TransformVector(CornerVec);
@@ -442,7 +446,8 @@ static void TransformBoxShape(FBoxShape3d& BoxShape, const TArray<FTransform3d>&
 	BoxShape.Box.Extents.Y = CornerVec.Dot(BoxShape.Box.AxisY());
 	BoxShape.Box.Extents.Z = CornerVec.Dot(BoxShape.Box.AxisZ());
 }
-static void TransformCapsuleShape(UE::Geometry::FCapsuleShape3d& CapsuleShape, const TArray<FTransform3d>& TransformSequence)
+template<typename TransformArrayType>
+static void TransformCapsuleShapeByArray(UE::Geometry::FCapsuleShape3d& CapsuleShape, const TransformArrayType& TransformSequence)
 {
 	FVector3d P0 = CapsuleShape.Capsule.Segment.StartPoint();
 	FVector3d P1 = CapsuleShape.Capsule.Segment.EndPoint();
@@ -452,7 +457,7 @@ static void TransformCapsuleShape(UE::Geometry::FCapsuleShape3d& CapsuleShape, c
 	FVector3d InitialSideVec = CapsuleFrame.PointAt(FVector3d(CurRadius, CurRadius, 0)) - CapsuleFrame.Origin;
 	FVector3d NewSideVec = InitialSideVec;
 
-	for (const FTransform3d& XForm : TransformSequence)
+	for (const auto& XForm : TransformSequence)
 	{
 		P0 = XForm.TransformPosition(P0);
 		P1 = XForm.TransformPosition(P1);
@@ -465,15 +470,15 @@ static void TransformCapsuleShape(UE::Geometry::FCapsuleShape3d& CapsuleShape, c
 	double RadiusScale = NewSideVec.Length() / InitialSideVec.Length();
 	CapsuleShape.Capsule.Radius *= RadiusScale;
 }
-
-static void TransformLevelSetShape(UE::Geometry::FLevelSetShape3d& LevelSetShape, const TArray<FTransform3d>& TransformSequence)
+template<typename TransformArrayType>
+static void TransformLevelSetShapeByArray(UE::Geometry::FLevelSetShape3d& LevelSetShape, const TransformArrayType& TransformSequence)
 {
 	const FTransform3d Sub(-0.5 * LevelSetShape.CellSize * FVector3d::One());
 	const FTransform3d Add(0.5 * LevelSetShape.CellSize * FVector3d::One());
 
 	LevelSetShape.GridTransform = LevelSetShape.GridTransform * Sub;
 
-	for (const FTransform3d& XForm : TransformSequence)
+	for (const auto& XForm : TransformSequence)
 	{
 		LevelSetShape.GridTransform = LevelSetShape.GridTransform * XForm;
 	}
@@ -554,19 +559,19 @@ void FSimpleShapeSet3d::Append(const FSimpleShapeSet3d& OtherShapeSet, const TAr
 {
 	for (FSphereShape3d SphereShape : OtherShapeSet.Spheres)
 	{
-		TransformSphereShape(SphereShape, TransformSequence);
+		TransformSphereShapeByArray(SphereShape, TransformSequence);
 		Spheres.Add(SphereShape);
 	}
 
 	for (FBoxShape3d BoxShape : OtherShapeSet.Boxes)
 	{
-		TransformBoxShape(BoxShape, TransformSequence);
+		TransformBoxShapeByArray(BoxShape, TransformSequence);
 		Boxes.Add(BoxShape);
 	}
 
 	for (UE::Geometry::FCapsuleShape3d CapsuleShape : OtherShapeSet.Capsules)
 	{
-		TransformCapsuleShape(CapsuleShape, TransformSequence);
+		TransformCapsuleShapeByArray(CapsuleShape, TransformSequence);
 		Capsules.Add(CapsuleShape);
 	}
 
@@ -584,11 +589,51 @@ void FSimpleShapeSet3d::Append(const FSimpleShapeSet3d& OtherShapeSet, const TAr
 	for (const FLevelSetShape3d& LevelSetShape : OtherShapeSet.LevelSets)
 	{
 		LevelSets.Add(LevelSetShape);
-		TransformLevelSetShape(LevelSets.Last(), TransformSequence);
+		TransformLevelSetShapeByArray(LevelSets.Last(), TransformSequence);
 	}
 }
 
 
+void FSimpleShapeSet3d::Append(const FSimpleShapeSet3d& OtherShapeSet, const FTransformSequence3d& TransformSequence)
+{
+	for (FSphereShape3d SphereShape : OtherShapeSet.Spheres)
+	{
+		TransformSphereShapeByArray(SphereShape, TransformSequence.GetTransforms());
+		Spheres.Add(SphereShape);
+	}
+
+	for (FBoxShape3d BoxShape : OtherShapeSet.Boxes)
+	{
+		TransformBoxShapeByArray(BoxShape, TransformSequence.GetTransforms());
+		Boxes.Add(BoxShape);
+	}
+
+	for (UE::Geometry::FCapsuleShape3d CapsuleShape : OtherShapeSet.Capsules)
+	{
+		TransformCapsuleShapeByArray(CapsuleShape, TransformSequence.GetTransforms());
+		Capsules.Add(CapsuleShape);
+	}
+
+	Convexes.Reserve(Convexes.Num() + OtherShapeSet.Convexes.Num());
+	for (const FConvexShape3d& ConvexShape : OtherShapeSet.Convexes)
+	{
+		Convexes.Add(ConvexShape);
+		for (const FTransformSRT3d& XForm : TransformSequence.GetTransforms())
+		{
+			MeshTransforms::ApplyTransform(Convexes.Last().Mesh, XForm, true);
+		}
+	}
+
+	LevelSets.Reserve(LevelSets.Num() + OtherShapeSet.LevelSets.Num());
+	for (const FLevelSetShape3d& LevelSetShape : OtherShapeSet.LevelSets)
+	{
+		LevelSets.Add(LevelSetShape);
+		TransformLevelSetShapeByArray(LevelSets.Last(), TransformSequence.GetTransforms());
+	}
+}
+
+
+
 
 void FSimpleShapeSet3d::FilterByVolume(int32 MaximumCount)
 {
diff --git a/Engine/Plugins/Runtime/GeometryProcessing/Source/DynamicMesh/Public/ShapeApproximation/SimpleShapeSet3.h b/Engine/Plugins/Runtime/GeometryProcessing/Source/DynamicMesh/Public/ShapeApproximation/SimpleShapeSet3.h
--- a/Engine/Plugins/Runtime/GeometryProcessing/Source/DynamicMesh/Public/ShapeApproximation/SimpleShapeSet3.h
+++ b/Engine/Plugins/Runtime/GeometryProcessing/Source/DynamicMesh/Public/ShapeApproximation/SimpleShapeSet3.h
@@ -11,6 +11,7 @@
 #include "DynamicMesh/DynamicMesh3.h"
 #include "Spatial/DenseGrid3.h"
 #include "TransformTypes.h"
+#include "TransformSequence.h"
 
 
 namespace UE {
@@ -178,6 +179,10 @@ struct DYNAMICMESH_API FSimpleShapeSet3d
 	 */
 	void Append(const FSimpleShapeSet3d& OtherShapeSet, const TArray<FTransform3d>& TransformSequence);
 
+	/**
+	 * Append elements of another shape set with given transforms applied
+	 */
+	void Append(const FSimpleShapeSet3d& OtherShapeSet, const FTransformSequence3d& TransformSequence);
 
 	/**
 	 * Remove any of the elements that are fully contained in larger elements
diff --git a/Engine/Source/Developer/GeometryProcessingInterfaces/Public/GeometryProcessingInterfaces/CombineMeshInstances.h b/Engine/Source/Developer/GeometryProcessingInterfaces/Public/GeometryProcessingInterfaces/CombineMeshInstances.h
--- a/Engine/Source/Developer/GeometryProcessingInterfaces/Public/GeometryProcessingInterfaces/CombineMeshInstances.h
+++ b/Engine/Source/Developer/GeometryProcessingInterfaces/Public/GeometryProcessingInterfaces/CombineMeshInstances.h
@@ -5,6 +5,7 @@
 #include "CoreMinimal.h"
 #include "Features/IModularFeature.h"
 #include "DynamicMesh/DynamicMesh3.h"
+#include "PhysicsEngine/AggregateGeom.h"
 
 class UPrimitiveComponent;
 class UStaticMesh;
@@ -91,6 +92,8 @@ public:
 	{
 		TArray<UE::Geometry::FDynamicMesh3> MeshLODs;
 		TArray<UMaterialInterface*> MaterialSet;
+
+		FKAggregateGeom SimpleCollisionShapes;
 	};
 
 
diff --git a/Engine/Source/Runtime/GeometryCore/Public/BoxTypes.h b/Engine/Source/Runtime/GeometryCore/Public/BoxTypes.h
--- a/Engine/Source/Runtime/GeometryCore/Public/BoxTypes.h
+++ b/Engine/Source/Runtime/GeometryCore/Public/BoxTypes.h
@@ -500,7 +500,7 @@ struct TAxisAlignedBox3
 		return dx * dx + dy * dy + dz * dz;
 	}
 
-	RealType DistanceSquared(const TAxisAlignedBox3<RealType>& Box)
+	RealType DistanceSquared(const TAxisAlignedBox3<RealType>& Box) const
 	{
 		// compute lensqr( max(0, abs(center1-center2) - (extent1+extent2)) )
 		RealType delta_x = TMathUtil<RealType>::Abs((Box.Min.X + Box.Max.X) - (Min.X + Max.X))
diff --git a/Engine/Source/Runtime/GeometryCore/Public/Quaternion.h b/Engine/Source/Runtime/GeometryCore/Public/Quaternion.h
--- a/Engine/Source/Runtime/GeometryCore/Public/Quaternion.h
+++ b/Engine/Source/Runtime/GeometryCore/Public/Quaternion.h
@@ -277,6 +277,27 @@ TQuaternion<RealType> operator*(const TQuaternion<RealType>& A, const TQuaternio
 	return TQuaternion<RealType>(X, Y, Z, W);
 }
 
+
+template<typename RealType>
+TQuaternion<RealType> operator*(RealType Scalar, const TQuaternion<RealType>& Q) 
+{
+	return TQuaternion<RealType>(Scalar * Q.X, Scalar * Q.Y, Scalar * Q.Z, Scalar * Q.W);
+}
+
+template<typename RealType>
+TQuaternion<RealType> operator*(const TQuaternion<RealType>& Q, RealType Scalar) 
+{
+	return TQuaternion<RealType>(Scalar * Q.X, Scalar * Q.Y, Scalar * Q.Z, Scalar * Q.W);
+}
+
+
+template<typename RealType>
+TQuaternion<RealType> operator+(const TQuaternion<RealType>& A, const TQuaternion<RealType>& B) 
+{
+	return TQuaternion<RealType>(A.X + B.X, A.Y + B.Y, A.Z + B.Z, A.W + B.W);
+}
+
+
 template<typename RealType>
 TQuaternion<RealType> operator -(const TQuaternion<RealType>& A, const TQuaternion<RealType>& B) 
 {
diff --git a/Engine/Source/Runtime/GeometryCore/Public/TransformTypes.h b/Engine/Source/Runtime/GeometryCore/Public/TransformTypes.h
--- a/Engine/Source/Runtime/GeometryCore/Public/TransformTypes.h
+++ b/Engine/Source/Runtime/GeometryCore/Public/TransformTypes.h
@@ -149,6 +149,14 @@ public:
 		return Scale3D;
 	}
 
+	/**
+	 * @return Scale portion of transform
+	 */
+	const TVector<RealType>& GetScale3D() const
+	{
+		return Scale3D;
+	}
+
 	/**
 	 * set Scale portion of transform
 	 */
