Remove checkSlow(false) when encountering a numerical issue since it's valid for those cases to be reached and we return the Numerical Failure flag.
Otherwise, this also causes geometry processing unit tests to fail when testing for numerical failures where invalid matrices are supplied.

#rnx

diff --git a/Engine/Plugins/Runtime/GeometryProcessing/Source/DynamicMesh/Private/Solvers/Internal/PowerMethodSolver.cpp b/Engine/Plugins/Runtime/GeometryProcessing/Source/DynamicMesh/Private/Solvers/Internal/PowerMethodSolver.cpp
--- a/Engine/Plugins/Runtime/GeometryProcessing/Source/DynamicMesh/Private/Solvers/Internal/PowerMethodSolver.cpp
+++ b/Engine/Plugins/Runtime/GeometryProcessing/Source/DynamicMesh/Private/Solvers/Internal/PowerMethodSolver.cpp
@@ -45,7 +45,6 @@ FPowerMethod::ESolveStatus FPowerMethod::Solve(ScalarType& OutEigenValue, RealVe
         ScalarType SquaredMagnitude = OutEigenVector.squaredNorm();
         if (FMath::IsFinite(SquaredMagnitude) == false || FMath::IsNearlyZero(SquaredMagnitude))
         {
-            checkSlow(false);
             return ESolveStatus::NumericalIssue;
         } 
 
@@ -56,7 +55,6 @@ FPowerMethod::ESolveStatus FPowerMethod::Solve(ScalarType& OutEigenValue, RealVe
         OutEigenValue = this->RayleighQuotient(OutEigenVector);  
         if (FMath::IsFinite(OutEigenValue) == false) 
         {
-            checkSlow(false);
             return ESolveStatus::NumericalIssue;
         }
 
@@ -161,7 +159,6 @@ bool FPowerMethod::Iteration(RealVectorType& EVector, bool bComputeLargest)
         Solver->Solve(EVector, EVectorSolve);
         if (Solver->bSucceeded() == false) 
         {
-            checkSlow(false);
             return false;
         }
         
@@ -181,7 +178,6 @@ bool FPowerMethod::Iteration(RealVectorType& EVector, bool bComputeLargest)
             Solver->Solve(EVector, EVectorSolve);
             if (Solver->bSucceeded() == false) 
             {
-                checkSlow(false);
                 return false;
             } 
 
