Hole Fill Tool:
- change the smoothing weight scheme to IDT Cotangent
- re-check for mesh seams to constrain between iterations of remeshing

diff --git a/Engine/Plugins/Runtime/GeometryProcessing/Source/DynamicMesh/Private/Operations/SmoothHoleFiller.cpp b/Engine/Plugins/Runtime/GeometryProcessing/Source/DynamicMesh/Private/Operations/SmoothHoleFiller.cpp
--- a/Engine/Plugins/Runtime/GeometryProcessing/Source/DynamicMesh/Private/Operations/SmoothHoleFiller.cpp
+++ b/Engine/Plugins/Runtime/GeometryProcessing/Source/DynamicMesh/Private/Operations/SmoothHoleFiller.cpp
@@ -90,7 +90,7 @@ namespace
 		FMeshRegionOperator Region(Mesh, SmoothTriangles.AsArray());
 		FDynamicMesh3& SmoothMesh = Region.Region.GetSubmesh();
 
-		ELaplacianWeightScheme UseScheme = ELaplacianWeightScheme::ClampedCotangent;
+		ELaplacianWeightScheme UseScheme = ELaplacianWeightScheme::IDTCotanget;
 		TUniquePtr<UE::Solvers::IConstrainedMeshSolver> Smoother = UE::MeshDeformation::ConstructConstrainedMeshSmoother(
 			UseScheme, SmoothMesh);
 		check(Smoother);
@@ -188,18 +188,12 @@ FSmoothHoleFiller::FSmoothHoleFiller(FDynamicMesh3& Mesh, const FEdgeLoop& FillL
 {
 }
 
-void FSmoothHoleFiller::DefaultConfigureRemesher(FSubRegionRemesher& Remesher, bool bConstrainROIBoundary)
-{
-	check(RemeshingTargetEdgeLength > 0.0);
-	Remesher.SetTargetEdgeLength(RemeshingTargetEdgeLength);
-
-	Remesher.SmoothSpeedT = RemeshingSmoothAlpha;
-	Remesher.bEnableSmoothing = (RemeshingSmoothAlpha > 0.0);
-	Remesher.SmoothType = FRemesher::ESmoothTypes::MeanValue;
 
+void FSmoothHoleFiller::ConstrainSubregionSeams(FSubRegionRemesher& Remesher, bool bConstrainROIBoundary)
+{
 	FMeshConstraints Constraints;
 	const bool bAllowSplits = !bConstrainROIBoundary;
-	const bool bAllowSmoothing = false;
+	constexpr bool bAllowSmoothing = false;
 
 	// Constrain seam edges and vertices in the EdgeROI
 	FMeshConstraintsUtil::ConstrainSeamsInEdgeROI(Constraints,
@@ -226,6 +220,18 @@ void FSmoothHoleFiller::DefaultConfigureRemesher(FSubRegionRemesher& Remesher, b
 	{
 		Remesher.SetExternalConstraints(MoveTemp(Constraints));
 	}
+}
+
+void FSmoothHoleFiller::DefaultConfigureRemesher(FSubRegionRemesher& Remesher, bool bConstrainROIBoundary)
+{
+	check(RemeshingTargetEdgeLength > 0.0);
+	Remesher.SetTargetEdgeLength(RemeshingTargetEdgeLength);
+
+	Remesher.SmoothSpeedT = RemeshingSmoothAlpha;
+	Remesher.bEnableSmoothing = (RemeshingSmoothAlpha > 0.0);
+	Remesher.SmoothType = FRemesher::ESmoothTypes::MeanValue;
+
+	ConstrainSubregionSeams(Remesher, bConstrainROIBoundary);
 
 	Remesher.ProjectionMode = FMeshRefinerBase::ETargetProjectionMode::NoProjection;
 }
@@ -285,7 +291,12 @@ bool FSmoothHoleFiller::Fill(int32 GroupID)
 
 		for (int k = 0; k < InitialRemeshPasses; ++k)
 		{
-			Remesher->UpdateROI();
+			if (k > 0)
+			{
+				Remesher->UpdateROI();
+				// If the ROI has expanded, it may include seam edges that were not constrained in the original Remesher set up
+				ConstrainSubregionSeams(*Remesher, true);
+			}
 			Remesher->BasicRemeshPass();
 		}
 
@@ -345,7 +356,9 @@ void FSmoothHoleFiller::SmoothAndRemeshPreserveRegion(FMeshFaceSelection& Triang
 	if (bRemeshAfterSmooth)
 	{
 		FRestrictedSubRegionRemesher Remesher(&Mesh, TriangleSelection.AsSet());
+		Remesher.UpdateROI();
 		DefaultConfigureRemesher(Remesher, true);
+		
 
 		FDynamicMesh3 ProjectionTargetMeshCopy;
 		TUniquePtr<FDynamicMeshAABBTree3> ProjectionTargetSpatial = nullptr;
@@ -368,7 +381,12 @@ void FSmoothHoleFiller::SmoothAndRemeshPreserveRegion(FMeshFaceSelection& Triang
 
 		for (int k = 0; k < PostSmoothRemeshPasses; ++k)
 		{
-			Remesher.UpdateROI();
+			if (k > 0)
+			{
+				Remesher.UpdateROI();
+				// If the ROI has expanded, it may include seam edges that were not constrained in the original Remesher set up
+				ConstrainSubregionSeams(Remesher, true);
+			}
 			Remesher.BasicRemeshPass();
 		}
 
@@ -419,7 +437,12 @@ void FSmoothHoleFiller::SmoothAndRemesh(FMeshFaceSelection& TriangleSelection)
 
 		for (int k = 0; k < PostSmoothRemeshPasses; ++k)
 		{
-			Remesher.UpdateROI();
+			if (k > 0)
+			{
+				Remesher.UpdateROI();
+				// If the ROI has expanded, it may include seam edges that were not constrained in the original Remesher set up
+				ConstrainSubregionSeams(Remesher, false);
+			}
 			Remesher.BasicRemeshPass();
 		}
 
diff --git a/Engine/Plugins/Runtime/GeometryProcessing/Source/DynamicMesh/Public/Operations/SmoothHoleFiller.h b/Engine/Plugins/Runtime/GeometryProcessing/Source/DynamicMesh/Public/Operations/SmoothHoleFiller.h
--- a/Engine/Plugins/Runtime/GeometryProcessing/Source/DynamicMesh/Public/Operations/SmoothHoleFiller.h
+++ b/Engine/Plugins/Runtime/GeometryProcessing/Source/DynamicMesh/Public/Operations/SmoothHoleFiller.h
@@ -94,6 +94,8 @@ protected:
 
 	void SmoothAndRemesh(FMeshFaceSelection& tris);
 
+	void ConstrainSubregionSeams(FSubRegionRemesher& Remesher, bool bConstrainROIBoundary);
+
 	void DefaultConfigureRemesher(FSubRegionRemesher& Remesher, bool bConstrainROIBoundary);
 
 };
