MeshModelingTools: Fix crashes in smoothing code on non-compacted meshes.

diff --git a/Engine/Plugins/Runtime/GeometryProcessing/Source/DynamicMesh/Public/Solvers/PrecomputedMeshWeightData.h b/Engine/Plugins/Runtime/GeometryProcessing/Source/DynamicMesh/Public/Solvers/PrecomputedMeshWeightData.h
--- a/Engine/Plugins/Runtime/GeometryProcessing/Source/DynamicMesh/Public/Solvers/PrecomputedMeshWeightData.h
+++ b/Engine/Plugins/Runtime/GeometryProcessing/Source/DynamicMesh/Public/Solvers/PrecomputedMeshWeightData.h
@@ -210,17 +210,17 @@ namespace UE
 			const int32 NumTris = TriangleLinearization.NumTris();
 			TriangleDataArray.SetNumUninitialized(NumTris);
 
-			const auto& ToTriIdx = TriangleLinearization.ToIndex();
+			const auto& ToTriId = TriangleLinearization.ToId();
 
 			for (int32 i = 0; i < NumTris; ++i)
 			{
 				// Current triangle
 
-				const int32 TriIdx = ToTriIdx[i];
+				const int32 TriId = ToTriId[i];
 
 				// Compute all the geometric data needed for this triangle.
 
-				TriangleDataArray[i].Initialize(DynamicMesh, TriIdx);
+				TriangleDataArray[i].Initialize(DynamicMesh, TriId);
 			}
 		}
 
