Fix computation of edge collapse lerp parameter to fix (for example) artefacts in the UV layer when remeshing the stairs primitive

diff --git a/Engine/Plugins/Runtime/GeometryProcessing/Source/DynamicMesh/Private/Remesher.cpp b/Engine/Plugins/Runtime/GeometryProcessing/Source/DynamicMesh/Private/Remesher.cpp
--- a/Engine/Plugins/Runtime/GeometryProcessing/Source/DynamicMesh/Private/Remesher.cpp
+++ b/Engine/Plugins/Runtime/GeometryProcessing/Source/DynamicMesh/Private/Remesher.cpp
@@ -176,26 +176,32 @@ FRemesher::EProcessResult FRemesher::ProcessEdge(int edgeID)
 		//  funcs will do it internally...
 		//  (or maybe we can collapse if cd exists? edge-collapse doesn't check for it explicitly...)
 
-		double collapse_t = 0.5;		// need to know t-value along edge to update lerpable attributes properly
-		FVector3d vNewPos = (vA + vB) * collapse_t;
+		int iKeep;
+		int iCollapse;
+		double collapse_t;		// need to know t-value along edge to update lerpable attributes properly
+		FVector3d vNewPos;
 
-		int iKeep = b;
-		int iCollapse = a;
-		// if either vtx is fixed, collapse to that position
+		// If either vtx is fixed, make that vertex iKeep and set collapse_t to 0 so vNewPos corresponds to the kept vtx
+		// Otherwise set collapse_t and vNewPos using the projected edge midpoint and measure collapse_t from iKeep
 		if (collapse_to == b)
 		{
-			collapse_t = 1.0;
+			iKeep = b;
+			iCollapse = a;
+			collapse_t = 0;
 			vNewPos = vB;
 		}
 		else if (collapse_to == a)
 		{
-			iKeep = a; iCollapse = b;
+			iKeep = a;
+			iCollapse = b;
 			collapse_t = 0;
 			vNewPos = vA;
 		}
 		else
 		{
-			vNewPos = GetProjectedCollapsePosition(iKeep, vNewPos);
+			iKeep = b;
+			iCollapse = a;
+			vNewPos = GetProjectedCollapsePosition(iKeep, (vA + vB) * 0.5);
 			double div = Distance(vA, vB);
 			collapse_t = (div < FMathd::ZeroTolerance) ? 0.5 : (Distance(vNewPos, Mesh->GetVertex(iKeep))) / div;
 			collapse_t = VectorUtil::Clamp(collapse_t, 0.0, 1.0);
diff --git a/Engine/Source/Runtime/GeometryCore/Public/DynamicMesh/DynamicMesh3.h b/Engine/Source/Runtime/GeometryCore/Public/DynamicMesh/DynamicMesh3.h
--- a/Engine/Source/Runtime/GeometryCore/Public/DynamicMesh/DynamicMesh3.h
+++ b/Engine/Source/Runtime/GeometryCore/Public/DynamicMesh/DynamicMesh3.h
@@ -1270,7 +1270,7 @@ public:
 	 * Collapse the edge between the two vertices, if topologically possible.
 	 * @param KeepVertID index of the vertex that should be kept
 	 * @param RemoveVertID index of the vertex that should be removed
-	 * @param EdgeParameterT vKeep is moved to Lerp(KeepPos, RemovePos, collapse_t)
+	 * @param EdgeParameterT vKeep is moved to Lerp(KeepPos, RemovePos, EdgeParameterT)
 	 * @param CollapseInfo returned information about new and modified mesh elements
 	 * @return Ok on success, or enum value indicates why operation cannot be applied. Mesh remains unmodified on error.
 	 */
diff --git a/Engine/Source/Runtime/GeometryCore/Public/DynamicMesh/InfoTypes.h b/Engine/Source/Runtime/GeometryCore/Public/DynamicMesh/InfoTypes.h
--- a/Engine/Source/Runtime/GeometryCore/Public/DynamicMesh/InfoTypes.h
+++ b/Engine/Source/Runtime/GeometryCore/Public/DynamicMesh/InfoTypes.h
@@ -189,7 +189,7 @@ struct FEdgeCollapseInfo
 	FIndex2i RemovedEdges;			// the edges that were removed (second is InvalidID for boundary edge)
 	FIndex2i KeptEdges;				// the edges that were kept (second is InvalidID for boundary edge)
 
-	double CollapseT;				// interpolation parameter along edge for new vertex in range [0,1]
+	double CollapseT;				// interpolation parameter along edge for new vertex in range [0,1] where 0 => KeptVertex and 1 => RemovedVertex
 };
 
 /** Information about mesh elements modified by MergeEdges() */
