Fixed a static analysis error.

#rnx

diff --git a/Engine/Plugins/Runtime/GeometryProcessing/Source/DynamicMesh/Private/GroupTopology.cpp b/Engine/Plugins/Runtime/GeometryProcessing/Source/DynamicMesh/Private/GroupTopology.cpp
--- a/Engine/Plugins/Runtime/GeometryProcessing/Source/DynamicMesh/Private/GroupTopology.cpp
+++ b/Engine/Plugins/Runtime/GeometryProcessing/Source/DynamicMesh/Private/GroupTopology.cpp
@@ -619,7 +619,9 @@ bool FGroupTopology::GenerateBoundaryAndGroupEdges(FGroup& Group,
 				return false;
 			}
 		}
-		return ensure(false); // Shouldn't be able to get here
+
+		ensure(false); // Shouldn't be able to get here
+		return false;
 	};
 
 	// Go through the boundary, find the corners, and use the intervening edges to create
