GeometryProcessing: Handle the case where appending a triangle fails in  FDynamicMeshEditor::SplitMesh(). This conceivably should never be possible, but it was observed in the wild due to a broken input FDynamicMesh3, and it's straightforward to try to handle here.

diff --git a/Engine/Plugins/Runtime/GeometryProcessing/Source/DynamicMesh/Private/DynamicMeshEditor.cpp b/Engine/Plugins/Runtime/GeometryProcessing/Source/DynamicMesh/Private/DynamicMeshEditor.cpp
--- a/Engine/Plugins/Runtime/GeometryProcessing/Source/DynamicMesh/Private/DynamicMeshEditor.cpp
+++ b/Engine/Plugins/Runtime/GeometryProcessing/Source/DynamicMesh/Private/DynamicMeshEditor.cpp
@@ -2272,6 +2272,7 @@ bool FDynamicMeshEditor::SplitMesh(const FDynamicMesh3* SourceMesh, TArray<FDyna
 			}
 		}
 
+		bool bCreatedNewVertex[3] = {false, false, false};
 		FIndex3i NewTri;
 		for (int j = 0; j < 3; ++j)
 		{
@@ -2279,6 +2280,7 @@ bool FDynamicMeshEditor::SplitMesh(const FDynamicMesh3* SourceMesh, TArray<FDyna
 			int NewVID = IndexMaps.GetNewVertex(SourceVID);
 			if (NewVID == IndexMaps.InvalidID())
 			{
+				bCreatedNewVertex[j] = true;
 				NewVID = Mesh.AppendVertex(*SourceMesh, SourceVID);
 				IndexMaps.SetVertex(SourceVID, NewVID);
 			}
@@ -2286,8 +2288,33 @@ bool FDynamicMeshEditor::SplitMesh(const FDynamicMesh3* SourceMesh, TArray<FDyna
 		}
 
 		int NewTID = Mesh.AppendTriangle(NewTri, NewGID);
-		IndexMaps.SetTriangle(SourceTID, NewTID);
-		AppendTriangleAttributes(SourceMesh, SourceTID, &Mesh, NewTID, IndexMaps, UnusedInvalidResultAccumulator);
+
+		// conceivably this should never happen, but it did occur due to other mesh issues,
+		// and it can be handled here without much effort
+		if (NewTID < 0)
+		{
+			// append failed, try creating separate new vertices
+			for (int j = 0; j < 3; ++j)
+			{
+				if ( bCreatedNewVertex[j] == false )
+				{
+					int SourceVID = Tri[j];
+					NewTri[j] = Mesh.AppendVertex(*SourceMesh, SourceVID);
+				}
+			}
+			NewTID = Mesh.AppendTriangle(NewTri, NewGID);
+		}
+
+		if ( NewTID >= 0 )
+		{
+			IndexMaps.SetTriangle(SourceTID, NewTID);
+			AppendTriangleAttributes(SourceMesh, SourceTID, &Mesh, NewTID, IndexMaps, UnusedInvalidResultAccumulator);
+		}
+		else
+		{
+			checkSlow(false);
+			// something has gone very wrong, skip this triangle
+		}
 	}
 
 	for (int Idx = 0; Idx < NumMeshes; Idx++)
@@ -2296,4 +2323,4 @@ bool FDynamicMeshEditor::SplitMesh(const FDynamicMesh3* SourceMesh, TArray<FDyna
 	}
 	
 	return true;
-}
\ No newline at end of file
+}
