GeometryProcessing: fix issue in MeshSimplification where explicitly-constrained vertices could result in collapses ignoring geometric constraint

diff --git a/Engine/Plugins/Runtime/GeometryProcessing/Source/DynamicMesh/Private/MeshSimplification.cpp b/Engine/Plugins/Runtime/GeometryProcessing/Source/DynamicMesh/Private/MeshSimplification.cpp
--- a/Engine/Plugins/Runtime/GeometryProcessing/Source/DynamicMesh/Private/MeshSimplification.cpp
+++ b/Engine/Plugins/Runtime/GeometryProcessing/Source/DynamicMesh/Private/MeshSimplification.cpp
@@ -1371,6 +1371,18 @@ ESimplificationResult TMeshSimplification<QuadricErrorType>::CollapseEdge(int ed
 	else
 	{
 		vNewPos = (collapse_to == a) ? vA : vB;
+
+		// If geometric error constraint is requested, ensure that the fixed vNewPos satisfies the constraint.
+		// This must be done, otherwise the free vertex being collapsed to a fixed vertex will be allowed
+		// to violate the geometric constraint
+		if (GeometricErrorConstraint != EGeometricErrorCriteria::None)
+		{
+			if (CheckIfCollapseWithinGeometricTolerance(iKeep, iCollapse, vNewPos, t0, t1) == false)
+			{
+				ProfileEndCollapse();
+				return ESimplificationResult::Failed_GeometricDeviation;
+			}
+		}
 	}
 
 	// check if this collapse will create a normal flip. Also checks
