Fix speckles in border regions when creating textures using BakeTX tool

diff --git a/Engine/Plugins/Runtime/GeometryProcessing/Source/DynamicMesh/Private/Sampling/MeshMapBaker.cpp b/Engine/Plugins/Runtime/GeometryProcessing/Source/DynamicMesh/Private/Sampling/MeshMapBaker.cpp
--- a/Engine/Plugins/Runtime/GeometryProcessing/Source/DynamicMesh/Private/Sampling/MeshMapBaker.cpp
+++ b/Engine/Plugins/Runtime/GeometryProcessing/Source/DynamicMesh/Private/Sampling/MeshMapBaker.cpp
@@ -21,6 +21,131 @@ FBoxFilter FMeshMapBaker::BoxFilter(BoxFilterRadius);
 FBSplineFilter FMeshMapBaker::BSplineFilter(BCFilterRadius);
 FMitchellNetravaliFilter FMeshMapBaker::MitchellNetravaliFilter(BCFilterRadius);
 
+namespace
+{
+
+void ComputeGutterTexelsUsingFilterKernelCoverage(
+	FImageOccupancyMap& OccupancyMap,
+	FImageTile Tile,
+	FImageDimensions Dimensions,
+	int32 FilterKernelSize,
+	bool(*IsInFilterRegionEval)(const FVector2d& Dist))
+{
+	const int64 TexelsPerTile = Tile.Num();
+	const int64 TexelsPerPaddedTile = OccupancyMap.Tile.Num();
+	const int32 SamplesPerTexel = OccupancyMap.PixelSampler.Num();
+
+	TArray64<int8> TexelIsGutter;
+	TexelIsGutter.Init(true, Tile.Num());
+
+	// First pass: Identify Gutter texels as the texels whose centers are not covered by any filter kernels
+	// placed at interior/border sample points of texels in the PaddedTile.
+	for (int64 TexelIndexInPaddedTile = 0; TexelIndexInPaddedTile < TexelsPerPaddedTile; ++TexelIndexInPaddedTile)
+	{
+		const FVector2i TexelCoordsInImage = OccupancyMap.Tile.GetSourceCoords(TexelIndexInPaddedTile);
+
+		// Check filter kernel coverage for each interior sample in the image texels
+		for (int32 SampleIndexInTexel = 0; SampleIndexInTexel < SamplesPerTexel; ++SampleIndexInTexel)
+		{
+			const int64 SampleIndexInPaddedTile = TexelIndexInPaddedTile * SamplesPerTexel + SampleIndexInTexel;
+			if (!OccupancyMap.IsInterior(SampleIndexInPaddedTile))
+			{
+				continue; // Not an interior sample, do nothing
+			}
+
+			const FVector2i KernelStartCoordsInImage(
+				FMath::Clamp(TexelCoordsInImage.X - FilterKernelSize, 0, Dimensions.GetWidth()),
+				FMath::Clamp(TexelCoordsInImage.Y - FilterKernelSize, 0, Dimensions.GetHeight())
+				);
+			const FVector2i KernelEndCoordsInImage(
+				FMath::Clamp(TexelCoordsInImage.X + FilterKernelSize + 1, 0, Dimensions.GetWidth()),
+				FMath::Clamp(TexelCoordsInImage.Y + FilterKernelSize + 1, 0, Dimensions.GetHeight())
+				);
+			const FImageTile KernelTile(KernelStartCoordsInImage, KernelEndCoordsInImage);
+			const int64 NbrTexelsPerKernelTile = KernelTile.Num();
+
+			for (int64 NbrTexelIndexInKernel = 0; NbrTexelIndexInKernel < NbrTexelsPerKernelTile; NbrTexelIndexInKernel++)
+			{
+				const FVector2i NbrTexelCoordsInImage = KernelTile.GetSourceCoords(NbrTexelIndexInKernel);
+
+				// Skip neighbour texels in the padding
+				if (!Tile.Contains(NbrTexelCoordsInImage.X, NbrTexelCoordsInImage.Y))
+				{
+					continue;
+				}
+
+				const int64 NbrTexelIndexInPaddedTile = OccupancyMap.Tile.GetIndexFromSourceCoords(NbrTexelCoordsInImage);
+				const int32 NbrTexelUVChart = OccupancyMap.TexelQueryUVChart[NbrTexelIndexInPaddedTile];
+				// Note: The occupancy map assigned UV chart indices to texels so all samples use this value
+				const int32 SampleUVChart = OccupancyMap.TexelQueryUVChart[TexelIndexInPaddedTile];
+
+				// Compute the filter coverage based on the UV distance from the neighbor texel center to the sample position
+				// Note: No contribution if the sample and neighbor texel are on different UV charts
+				if (SampleUVChart == NbrTexelUVChart)
+				{
+					// See :GridAlignedFilterKernels
+					FVector2d SampleUVInImage;
+					{
+						const FVector2d TexelSize = Dimensions.GetTexelSize();
+						const int64 TexelIndexInImage = Dimensions.GetIndex(TexelCoordsInImage.X, TexelCoordsInImage.Y);
+						const FVector2d TexelCenterUV = Dimensions.GetTexelUV(TexelIndexInImage);
+						const FVector2d SampleUVInTexel = OccupancyMap.PixelSampler.Sample(SampleIndexInTexel);
+						SampleUVInImage = TexelCenterUV - 0.5 * TexelSize + SampleUVInTexel * TexelSize;
+					}
+
+					const FVector2d NbrTexelCenterUV = Dimensions.GetTexelUV(NbrTexelCoordsInImage);
+
+					const FVector2d TexelDistance = Dimensions.GetTexelDistance(NbrTexelCenterUV, SampleUVInImage);
+					if (IsInFilterRegionEval(TexelDistance))
+					{
+						const int64 NbrTexelIndexInTile = Tile.GetIndexFromSourceCoords(NbrTexelCoordsInImage);
+						TexelIsGutter[NbrTexelIndexInTile] = false;
+					}
+				}
+			}
+		} // end sample loop
+	} // end texel loop
+
+	// Second pass: Construct Gutter texel to Interior texel mapping
+	for (int64 GutterTexelIndexInTile = 0; GutterTexelIndexInTile < TexelsPerTile; ++GutterTexelIndexInTile)
+	{
+		const FVector2i GutterTexelCoordsInImage = Tile.GetSourceCoords(GutterTexelIndexInTile);
+		const int64 GutterTexelIndexInImage = Dimensions.GetIndex(GutterTexelCoordsInImage);
+		if (!TexelIsGutter[GutterTexelIndexInTile])
+		{
+			continue; // Not a Gutter texel, do nothing
+		}
+
+		for (int32 SampleIndexInTexel = 0; SampleIndexInTexel < SamplesPerTexel; ++SampleIndexInTexel)
+		{
+			const int64 GutterTexelIndexInPaddedTile = OccupancyMap.Tile.GetIndexFromSourceCoords(GutterTexelCoordsInImage);
+			const int64 SampleIndexInPaddedTile = GutterTexelIndexInPaddedTile * SamplesPerTexel + SampleIndexInTexel;
+			if (OccupancyMap.TexelType[SampleIndexInPaddedTile] == OccupancyMap.GutterTexel)
+			{
+				// To avoid artifacts with mipmapping the gutter texels store the nearest interior texel and a
+				// post pass copies those nearest interior texel values to each gutter pixel. The mapped texel
+				// should be an interior texel since interior texels cover the UV mesh, so when we snap the
+				// closest UV point to the nearest texel it should be an interior texel
+				//
+				// TODO There can be interior texels closer to the gutter the texel than the one we find here. The texel
+				// found here is the interior texel that partially covers the UV mesh. Since we set interior texels from
+				// kernel coverage there will be a nearer texel. We could fix this by using the TexelQueryUV to define a
+				// search direction along which we can find the right texel, or better, use filter kernels to propogate
+				// known data into the gutter. Search :NearestInteriorGutterTexel for a related problem in the OccupancyMap
+				//
+				const FVector2d SampleUV = (FVector2d)OccupancyMap.TexelQueryUV[SampleIndexInPaddedTile];
+				const FVector2i NearestTexelCoordsInImage = Dimensions.UVToCoords(SampleUV);
+				const int64 MappedTexelIndexInImage = Dimensions.GetIndex(NearestTexelCoordsInImage);
+				const TTuple<int64, int64> GutterTexel(GutterTexelIndexInImage, MappedTexelIndexInImage);
+				OccupancyMap.GutterTexels.Add(GutterTexel);
+
+				break;
+			}
+		} // end sample loop
+	} // end tile pixel loop
+}
+
+} // end anonymous namespace
 
 void FMeshMapBaker::InitBake()
 {
@@ -324,7 +449,9 @@ void FMeshMapBaker::Bake()
 		FImageOccupancyMap OccupancyMap;
 		OccupancyMap.GutterSize = GutterSize;
 		OccupancyMap.Initialize(Dimensions, PaddedTile, SamplesPerPixel);
-		OccupancyMap.ComputeFromUVSpaceMesh(FlatMesh, [this](int32 TriangleID) { return FlatMesh.GetTriangleGroup(TriangleID); }, TargetMeshUVCharts);
+		const auto GetTriangleIDFunc = [this](int32 TriangleID) { return FlatMesh.GetTriangleGroup(TriangleID); };
+		OccupancyMap.ClassifySamplesFromUVSpaceMesh(FlatMesh, GetTriangleIDFunc, TargetMeshUVCharts);
+		ComputeGutterTexelsUsingFilterKernelCoverage(OccupancyMap, Tile, Dimensions, FilterKernelSize, IsInFilterRegionEval);
 		GutterTexelsPerTile[TileIdx] = OccupancyMap.GutterTexels;
 
 		const int64 NumTilePixels = Tile.Num();
@@ -367,22 +494,53 @@ void FMeshMapBaker::Bake()
 						const int64 LinearIdx = OccupancyMapLinearIdx * NumSamples + SampleIdx;
 						if (OccupancyMap.IsInterior(LinearIdx))
 						{
-							const FVector2d UVPosition = (FVector2d)OccupancyMap.TexelQueryUV[LinearIdx];
-							const int32 UVTriangleID = OccupancyMap.TexelQueryTriangle[LinearIdx];
+							// This block calls BakeSample on all interior/border samples of the OccupancyMap. In the code
+							// below QueryUVInImage is the UV coordinate in the multisampled texel image where the
+							// evaluators are evaluated. This coordinate is computed by the OccupancyMap and is the same
+							// as the sample's UV coordinate in the sample grid implied by the multisampled image
+							// (SampleUVInImage) if SampleUVInImage lies in the UV unwrap mesh. If SampleUVInImage is
+							// not in the UV mesh this sample is a border sample and the samples texel will not fully
+							// overlap with UV unwrap mesh. The evaluators must be evaluated at a UV position in the
+							// unwrap mesh so we must use QueryUVInImage position. The filter weights which apply the
+							// evaluated result to surrounding texels, however, could be could be placed at either of
+							// these positions. Until UE 5.2 we used QueryUVInImage but this contributed to issue UE-169350.
+							// QueryUVInImage points can be arbitrarily located relative to the texel centers which
+							// meant we ended up with kernels that narrowly overlapped with texel centers resulting
+							// in tiny filter weights and cases where the overall filter weight for the texel was zero
+							// (persumably because the Mitchell-Netravali filter function becomes negative near the
+							// edges so contributions from different samples can sum to zero). This resulted in speckles
+							// (aka black texels surrounded by neighbours of a very different value) in the final image.
+							// To fix this we now place filter kernels at SampleUVInImage which eliminates the
+							// small/arbitrary overlap between the filter kernel and texel centers (in fact the kernel
+							// overlaps are limited to a small number of possiblities determined by the sample positions
+							// in the texel and the kernel radius ie independent of the position of the uv unwrap mesh).
+							// This is also relevant for the code labeled :GridAlignedFilterKernels
+
+							FVector2d SampleUVInImage;
+							{
+								const FVector2d TexelSize = Dimensions.GetTexelSize();
+								const int64 TexelIndexInImage = Dimensions.GetIndex(ImageCoords.X, ImageCoords.Y);
+								const FVector2d TexelCenterUV = Dimensions.GetTexelUV(TexelIndexInImage);
+								const FVector2d SampleUVInTexel = OccupancyMap.PixelSampler.Sample(SampleIdx);
+								SampleUVInImage = TexelCenterUV - 0.5 * TexelSize + SampleUVInTexel * TexelSize;
+							}
+
+							const FVector2d QueryUVInImage = (FVector2d)OccupancyMap.TexelQueryUV[LinearIdx];
 
 							// Compute the per-sample correspondence data 
 							// Note: Since we check LinearIdx is an interior sample above we know we'll get a valid
 							// SampleInfo because interior samples all have valid UVTriangleIDs.
 							FMeshUVSampleInfo SampleInfo;
-							if (MeshUVSampler.QuerySampleInfo(UVTriangleID, UVPosition, SampleInfo))
+							const int32 UVTriangleID = OccupancyMap.TexelQueryTriangle[LinearIdx];
+							if (MeshUVSampler.QuerySampleInfo(UVTriangleID, QueryUVInImage, SampleInfo))
 							{
 								FMeshMapEvaluator::FCorrespondenceSample Sample;
 								bool bSampleValid = ComputeCorrespondenceSample(SampleInfo, Sample);
 								if (bSampleValid)
 								{
-									BakeSample(*TileBuffer, Sample, UVPosition, ImageCoords, OccupancyMap);
+									BakeSample(*TileBuffer, Sample, SampleUVInImage, ImageCoords, OccupancyMap);
 								}
-								InteriorSampleCallback(bSampleValid, Sample, UVPosition, ImageCoords);
+								InteriorSampleCallback(bSampleValid, Sample, SampleUVInImage, ImageCoords);
 							}
 						}
 					}
@@ -546,7 +704,7 @@ void FMeshMapBaker::Bake()
 void FMeshMapBaker::BakeSample(
 	FMeshMapTileBuffer& TileBuffer,
 	const FMeshMapEvaluator::FCorrespondenceSample& Sample,
-	const FVector2d& UVPosition,
+	const FVector2d& SampleFilterUVPosition,
 	const FVector2i& ImageCoords,
 	const FImageOccupancyMap& OccupancyMap)
 {
@@ -565,7 +723,8 @@ void FMeshMapBaker::BakeSample(
 	const int64 OccupancyMapSampleIdx = OccupancyMap.Tile.GetIndexFromSourceCoords(ImageCoords);
 	const int32 SampleUVChart = OccupancyMap.TexelQueryUVChart[OccupancyMapSampleIdx];
 
-	auto AddFn = [this, &ImageCoords, &UVPosition, &Tile, &TileBuffer, &OccupancyMap, SampleUVChart](const TArray<int32>& EvaluatorIds, const float* SourceBuffer, float Weight) -> void
+	auto AddFn = [this, &ImageCoords, &SampleFilterUVPosition, &Tile, &TileBuffer, &OccupancyMap, SampleUVChart]
+		(const TArray<int32>& EvaluatorIds, const float* SourceBuffer, float Weight) -> void
 	{
 		const FVector2i BoxFilterStart(
 			FMath::Clamp(ImageCoords.X - FilterKernelSize, 0, Dimensions.GetWidth()),
@@ -592,7 +751,8 @@ void FMeshMapBaker::BakeSample(
 			// Note: There will be no contribution if the sample and pixel are on different UV charts
 			float FilterWeight = Weight * static_cast<float>(SampleUVChart == BufferTilePixelUVChart);
 			{
-				FVector2d TexelDistance = Dimensions.GetTexelDistance(Dimensions.GetTexelUV(SourceCoords), UVPosition);
+				const FVector2d PixelCenterUVPosition = Dimensions.GetTexelUV(SourceCoords);
+				const FVector2d TexelDistance = Dimensions.GetTexelDistance(PixelCenterUVPosition, SampleFilterUVPosition);
 				FilterWeight *= TextureFilterEval(TexelDistance);
 			}
 
@@ -620,7 +780,8 @@ void FMeshMapBaker::BakeSample(
 		}
 	};
 
-	auto OverwriteFn = [this, &ImageCoords, &Tile, &TileBuffer](const TArray<int32>& EvaluatorIds, const float* SourceBuffer) -> void
+	auto OverwriteFn = [this, &ImageCoords, &Tile, &TileBuffer]
+		(const TArray<int32>& EvaluatorIds, const float* SourceBuffer) -> void
 	{
 		const int64 BufferTilePixelLinearIdx = Tile.GetIndexFromSourceCoords(ImageCoords); 
 		float* PixelBuffer = TileBuffer.GetPixel(BufferTilePixelLinearIdx);
@@ -646,7 +807,7 @@ void FMeshMapBaker::BakeSample(
 
 	if (SampleFilterF)
 	{
-		const float SampleMaskWeight = FMath::Clamp(SampleFilterF(ImageCoords, UVPosition, Sample.BaseSample.TriangleIndex), 0.0f, 1.0f);
+		const float SampleMaskWeight = FMath::Clamp(SampleFilterF(ImageCoords, SampleFilterUVPosition, Sample.BaseSample.TriangleIndex), 0.0f, 1.0f);
 		AddFn(EvaluatorIdsForMode(FMeshMapEvaluator::EAccumulateMode::Add), Buffer, SampleMaskWeight);
 		AddFn(EvaluatorIdsForMode(FMeshMapEvaluator::EAccumulateMode::Add), BakeDefaults.GetData(), 1.0f - SampleMaskWeight);
 		OverwriteFn(EvaluatorIdsForMode(FMeshMapEvaluator::EAccumulateMode::Overwrite), (SampleMaskWeight == 0) ? BakeDefaults.GetData() : Buffer);
@@ -726,15 +887,19 @@ void FMeshMapBaker::InitFilter()
 	case EBakeFilterType::None:
 		FilterKernelSize = 0;
 		TextureFilterEval = &EvaluateFilter<EBakeFilterType::None>;
+		IsInFilterRegionEval = &EvaluateIsInFilterRegion<EBakeFilterType::None>;
 		break;
 	case EBakeFilterType::Box:
 		TextureFilterEval = &EvaluateFilter<EBakeFilterType::Box>;
+		IsInFilterRegionEval = &EvaluateIsInFilterRegion<EBakeFilterType::Box>;
 		break;
 	case EBakeFilterType::BSpline:
 		TextureFilterEval = &EvaluateFilter<EBakeFilterType::BSpline>;
+		IsInFilterRegionEval = &EvaluateIsInFilterRegion<EBakeFilterType::BSpline>;
 		break;
 	case EBakeFilterType::MitchellNetravali:
 		TextureFilterEval = &EvaluateFilter<EBakeFilterType::MitchellNetravali>;
+		IsInFilterRegionEval = &EvaluateIsInFilterRegion<EBakeFilterType::MitchellNetravali>;
 		break;
 	}
 }
@@ -762,6 +927,28 @@ float FMeshMapBaker::EvaluateFilter(const FVector2d& Dist)
 	return Result;
 }
 
+template<FMeshMapBaker::EBakeFilterType BakeFilterType>
+bool FMeshMapBaker::EvaluateIsInFilterRegion(const FVector2d& Dist)
+{
+	bool Result = false;
+	if constexpr(BakeFilterType == EBakeFilterType::None)
+	{
+		Result = true;
+	}
+	else if constexpr(BakeFilterType == EBakeFilterType::Box)
+	{
+		Result = BoxFilter.IsInFilterRegion(Dist);
+	}
+	else if constexpr(BakeFilterType == EBakeFilterType::BSpline)
+	{
+		Result = BSplineFilter.IsInFilterRegion(Dist);
+	}
+	else if constexpr(BakeFilterType == EBakeFilterType::MitchellNetravali)
+	{
+		Result = MitchellNetravaliFilter.IsInFilterRegion(Dist);
+	}
+	return Result;
+}
 
 void FMeshMapBaker::ComputeUVCharts(const FDynamicMesh3& Mesh, TArray<int32>& MeshUVCharts)
 {
diff --git a/Engine/Plugins/Runtime/GeometryProcessing/Source/DynamicMesh/Public/Sampling/MeshMapBaker.h b/Engine/Plugins/Runtime/GeometryProcessing/Source/DynamicMesh/Public/Sampling/MeshMapBaker.h
--- a/Engine/Plugins/Runtime/GeometryProcessing/Source/DynamicMesh/Public/Sampling/MeshMapBaker.h
+++ b/Engine/Plugins/Runtime/GeometryProcessing/Source/DynamicMesh/Public/Sampling/MeshMapBaker.h
@@ -152,7 +152,7 @@ protected:
 	void BakeSample(
 		FMeshMapTileBuffer& TileBuffer,
 		const FMeshMapEvaluator::FCorrespondenceSample& Sample,
-		const FVector2d& UVPosition,
+		const FVector2d& QueryUVPosition,
 		const FVector2i& ImageCoords,
 		const FImageOccupancyMap& OccupancyMap);
 
@@ -211,10 +211,15 @@ protected:
 	/** Texture filter function */
 	using TextureFilterFn = float(*)(const FVector2d& Dist);
 	TextureFilterFn TextureFilterEval = nullptr;
+	using IsInFilterRegionFn = bool(*)(const FVector2d& Dist);
+	IsInFilterRegionFn IsInFilterRegionEval = nullptr;
 
 	template<EBakeFilterType BakeFilterType>
 	static float EvaluateFilter(const FVector2d& Dist);
 
+	template<EBakeFilterType BakeFilterType>
+	static bool EvaluateIsInFilterRegion(const FVector2d& Dist);
+
 	/** The total size of the temporary float buffer for BakeSample. */
 	int32 BakeSampleBufferSize = 0;
 
diff --git a/Engine/Source/Runtime/GeometryCore/Public/Image/BCSplineFilter.h b/Engine/Source/Runtime/GeometryCore/Public/Image/BCSplineFilter.h
--- a/Engine/Source/Runtime/GeometryCore/Public/Image/BCSplineFilter.h
+++ b/Engine/Source/Runtime/GeometryCore/Public/Image/BCSplineFilter.h
@@ -37,6 +37,9 @@ public:
 
 	/** @return the filter weight given a 2D distance vector, in Texel units. */
 	float GetWeight(const FVector2d& Dist) const;
+
+	/** @return true if the given 2D distance vector, in Texel units, is in the region where the filter is defined and false otherwise */
+	bool IsInFilterRegion(const FVector2d& Dist) const;
 };
 
 template<EBCSplineType SplineType, bool bIsRadial>
@@ -111,6 +114,16 @@ float FBCSplineFilter<SplineType, bIsRadial>::GetWeight(const FVector2d& Dist) c
 	}
 }
 
+template<EBCSplineType SplineType, bool bIsRadial>
+bool FBCSplineFilter<SplineType, bIsRadial>::IsInFilterRegion(const FVector2d& Dist) const
+{
+	if constexpr(bIsRadial)
+	{
+		return (FMathf::Abs((float)Dist.Length() / Radius) < 2.0f);
+	}
+	return (FMathf::Abs((float)Dist.X / Radius) < 2.0f) && (FMathf::Abs((float)Dist.Y / Radius) < 2.0f);
+}
+
 using FBSplineFilter = FBCSplineFilter<EBCSplineType::BSpline, /*bIsRadial*/ false>;
 using FCatmullRomFilter = FBCSplineFilter<EBCSplineType::CatmullRom, /*bIsRadial*/ false>;
 using FMitchellNetravaliFilter = FBCSplineFilter<EBCSplineType::MitchellNetravali, /*bIsRadial*/ false>;
diff --git a/Engine/Source/Runtime/GeometryCore/Public/Image/BoxFilter.h b/Engine/Source/Runtime/GeometryCore/Public/Image/BoxFilter.h
--- a/Engine/Source/Runtime/GeometryCore/Public/Image/BoxFilter.h
+++ b/Engine/Source/Runtime/GeometryCore/Public/Image/BoxFilter.h
@@ -35,6 +35,12 @@ public:
 		// center texel weight changes compared to the weight for 0 < Radius < 1, which seems more
 		// intuitive for users.
 		//
+		return Dist.X >= -Radius && Dist.X <= Radius && Dist.Y >= -Radius && Dist.Y <= Radius;
+	}
+
+	/** @return true if the given 2D distance vector, in Texel units, is in the region where the filter is defined and false otherwise */
+	bool IsInFilterRegion(const FVector2d& Dist) const
+	{
 		return Dist.X >= -Radius && Dist.X <= Radius && Dist.Y >= -Radius && Dist.Y <= Radius; 
 	}
 };
diff --git a/Engine/Source/Runtime/GeometryCore/Public/Image/ImageOccupancyMap.h b/Engine/Source/Runtime/GeometryCore/Public/Image/ImageOccupancyMap.h
--- a/Engine/Source/Runtime/GeometryCore/Public/Image/ImageOccupancyMap.h
+++ b/Engine/Source/Runtime/GeometryCore/Public/Image/ImageOccupancyMap.h
@@ -26,6 +26,9 @@ namespace Geometry
  * This is not just a cache. For 'border' texels where the texel center is technically outside the mesh/polygon,
  * but the texel rectangle may still overlap the shape, the nearest UV/Triangle is stored.
  * This simplifies computing samples around the borders such that the shape is covered under linear interpolatione/etc.
+ * 
+ * TODO Some of the names in this class are misleading because the name mentions "Texel" but should actually say "Sample"
+ * this is because the class initially supported only one sample per texel but later multi-sampling was added
  */
 class FImageOccupancyMap
 {
@@ -44,15 +47,34 @@ public:
 	//const int8 BorderTexel = 2;
 	const int8 GutterTexel = 3;
 
-	/** Texel Type for each texel in image, Size = Width x Height x Multisamples */
+	/**
+	 * Classification of each sample in Tile.
+	 * Size = Tile.GetWidth() x Tile.GetHeight() x SamplesPerPixel.
+	 */
 	TArray64<int8> TexelType;
-	/** Count of samples in the texel that are interior. */
+
+	/**
+	 * Count of interior samples for each texel in Tile.
+	 * Size = Tile.GetWidth() x Tile.GetHeight().
+	 */
 	TArray64<int32> TexelInteriorSamples;
-	/** UV for each texel in image. Only set for Interior texels */
+
+	/**
+	 * UV for each sample in Tile.
+	 * Size = Tile.GetWidth() x Tile.GetHeight() x SamplesPerPixel.
+	 */
 	TArray64<FVector2f> TexelQueryUV;
-	/** integer/Triangle ID for each texel in image. Only set for Interior texels. */
+
+	/**
+	 * Integer/Triangle ID for each sample in Tile.
+	 * Size = Tile.GetWidth() x Tile.GetHeight() x SamplesPerPixel.
+	 */
 	TArray64<int32> TexelQueryTriangle;
-	/** UV Chart ID for each texel in image. Only set if UVSpaceMeshTriCharts is provided during ComputeFromUVSpaceMesh */
+
+	/**
+	 * UV Chart ID for each texel in Tile. Only set if UVSpaceMeshTriCharts is provided
+	 * Size = Tile.GetWidth() x Tile.GetHeight().
+	 */
 	TArray64<int32> TexelQueryUVChart;
 
 	/** Set of Gutter Texels. Pair is <LinearIndexOfGutterTexel, LinearIndexOfNearestInteriorTexel>, so
@@ -65,20 +87,20 @@ public:
 	bool bParallel = true;
 
 
-	void Initialize(FImageDimensions DimensionsIn, int32 SamplesPerPixelIn = 1)
+	void Initialize(FImageDimensions DimensionsIn, int32 SamplesPerPixel = 1)
 	{
 		check(DimensionsIn.IsSquare()); // are we sure it works otherwise?
 		Dimensions = DimensionsIn;
 		Tile = FImageTile(FVector2i(0,0), FVector2i(Dimensions.GetWidth(), Dimensions.GetHeight()));
-		InitializePixelSampler(SamplesPerPixelIn);
+		InitializePixelSampler(SamplesPerPixel);
 	}
 
-	void Initialize(FImageDimensions DimensionsIn, const FImageTile& TileIn, int32 SamplesPerPixelIn = 1)
+	void Initialize(FImageDimensions DimensionsIn, const FImageTile& TileIn, int32 SamplesPerPixel = 1)
 	{
 		check(DimensionsIn.IsSquare());
 		Dimensions = DimensionsIn;
 		Tile = TileIn;
-		InitializePixelSampler(SamplesPerPixelIn);
+		InitializePixelSampler(SamplesPerPixel);
 	}
 
 	/**
@@ -101,21 +123,14 @@ public:
 	//	TexelType[Dimensions.GetIndex(X, Y)] = Type;
 	//}
 
-
-	/**
-	 * Computes the image occupancy map from a UV space mesh.
-	 * 
-	 * @param UVSpaceMesh UV space mesh to compute the occupancy map from.
-	 * @param GetTriangleIDFunc Lambda to remap a texel's nearest triangle ID
-	 * @param UVSpaceMeshTriCharts Optional UVSpaceMesh triangle ID to UV chart map 
-	 */
 	template<typename MeshType, typename GetTriangleIDFuncType>
-	bool ComputeFromUVSpaceMesh(const MeshType& UVSpaceMesh, 
+	bool ClassifySamplesFromUVSpaceMesh(
+		const MeshType& UVSpaceMesh, 
 		GetTriangleIDFuncType GetTriangleIDFunc = [](int32 TriangleID) { return TriangleID; },
 		const TArray<int32>* UVSpaceMeshTriCharts = nullptr)
 	{
-		TRACE_CPUPROFILER_EVENT_SCOPE(FImageOccupancyMap::ComputeFromUVSpaceMesh);
-		
+		TRACE_CPUPROFILER_EVENT_SCOPE(FImageOccupancyMap::ClassifySamples);
+
 		// make flat mesh
 		TMeshAABBTree3<MeshType> FlatSpatial(&UVSpaceMesh, true);
 
@@ -134,36 +149,25 @@ public:
 		IMeshSpatial::FQueryOptions QueryOptions;
 		QueryOptions.MaxDistance = (double)GutterSize * TexelDiag;
 
-		// find interior texels
-		TAtomic<int64> TotalGutterCounter = 0;
-		TArray< TArray64<TTuple<int64, int64>> > GutterTexelsPerScanline;
-		GutterTexelsPerScanline.SetNum(Tile.GetHeight());
-
 		const int32 NumUVSpaceMeshTriCharts = UVSpaceMeshTriCharts ? UVSpaceMeshTriCharts->Num() : 0;
 
+		// Classify texel samples
 		ParallelFor(Tile.GetHeight(),
-					[this, &UVSpaceMesh, &GetTriangleIDFunc, &FlatSpatial, TexelDiag, &QueryOptions, &TexelSize, &TotalGutterCounter, &GutterTexelsPerScanline, UVSpaceMeshTriCharts, NumUVSpaceMeshTriCharts]
+					[this, &UVSpaceMesh, &GetTriangleIDFunc, &FlatSpatial, TexelDiag, &QueryOptions, &TexelSize, UVSpaceMeshTriCharts, NumUVSpaceMeshTriCharts]
 		(int32 ImgY)
 		{
 			for (int32 ImgX = 0; ImgX < Tile.GetWidth(); ++ImgX)
 			{
-				// With multiple samples, a texel may contain multiple gutter samples.
-				// Since we copy nearest interior texel over top our gutter texels, we
-				// should only consider a texel a gutter texel iff all samples in a
-				// texel are gutter texels. This bool will track that as we iterate
-				// over each sample of the texel.
-				bool bIsGutterTexel = true;
-				bool bHitTriID = false;
-				TTuple<int64, int64> GutterNearestTexel;
+				const FVector2i SourceCoords = Tile.GetSourceCoords(ImgX, ImgY);
+				const int64 SourceTexelLinearIdx = Dimensions.GetIndex(SourceCoords.X, SourceCoords.Y);
+				const int64 TexelLinearIdx = Tile.GetIndex(ImgX, ImgY);
+				const FVector2d SourceTexelCenterUV = Dimensions.GetTexelUV(SourceTexelLinearIdx);
+
 				for (int32 Sample = 0; Sample < PixelSampler.Num(); ++Sample)
 				{
-					const FVector2i SourceCoords = Tile.GetSourceCoords(ImgX, ImgY);
-					const int64 SourceTexelLinearIdx = Dimensions.GetIndex(SourceCoords.X, SourceCoords.Y);
-					const int64 TexelLinearIdx = Tile.GetIndex(ImgX, ImgY);
 					const int64 SampleLinearIdx = TexelLinearIdx * PixelSampler.Num() + Sample;
 					const FVector2d SampleUV = PixelSampler.Sample(Sample);
-					FVector2d UVPoint = Dimensions.GetTexelUV(SourceTexelLinearIdx);
-					UVPoint = UVPoint - 0.5 * TexelSize + SampleUV * TexelSize;
+					const FVector2d UVPoint = SourceTexelCenterUV - 0.5 * TexelSize + SampleUV * TexelSize;
 					const FVector3d UVPoint3d(UVPoint.X, UVPoint.Y, 0);
 
 					double NearDistSqr;
@@ -181,46 +185,97 @@ public:
 							TexelQueryUV[SampleLinearIdx] = (FVector2f)UVPoint;
 							TexelQueryTriangle[SampleLinearIdx] = TriId;
 							++TexelInteriorSamples[TexelLinearIdx];
-							bIsGutterTexel = false;
 						}
 						else if (NearDistSqr < TexelDiag * TexelDiag)
 						{
 							const FDistPoint3Triangle3d DistQuery = TMeshQueries<MeshType>::TriangleDistance(UVSpaceMesh, NearestTriID, UVPoint3d);
 							FVector2d NearestUV = GetXY(DistQuery.ClosestTrianglePoint);
+
 							// nudge point into triangle to improve numerical behavior of things like barycentric coord calculation
+							// TODO Hmm this nudging could invalidate TriId? Maybe users should do it explicitly and then deal with/ignore this problem themselves
 							NearestUV += (10.0 * FMathf::ZeroTolerance) * Normalized(NearestUV - UVPoint);
 
+							// TODO Use BorderTexel here, and we should probably enable those
 							TexelType[SampleLinearIdx] = InteriorTexel;
 							TexelQueryUV[SampleLinearIdx] = (FVector2f)NearestUV;
 							TexelQueryTriangle[SampleLinearIdx] = TriId;
 							++TexelInteriorSamples[TexelLinearIdx];
-							bIsGutterTexel = false;
 						}
 						else
 						{
 							TexelType[SampleLinearIdx] = GutterTexel;
 							const FDistPoint3Triangle3d DistQuery = TMeshQueries<MeshType>::TriangleDistance(UVSpaceMesh, NearestTriID, UVPoint3d);
 							const FVector2d NearestUV = GetXY(DistQuery.ClosestTrianglePoint);
-							const FVector2i NearestCoords = Dimensions.UVToCoords(NearestUV);
 
-							// To avoid artifacts with mipmapping the gutter texels store the nearest
-							// interior texel and a post pass copies those nearest interior texel values
-							// to each gutter pixel.
-							const int64 NearestLinearIdx = Dimensions.GetIndex(NearestCoords);
-							GutterNearestTexel = TTuple<int64, int64>(SourceTexelLinearIdx, NearestLinearIdx);
+							TexelQueryUV[SampleLinearIdx] = (FVector2f)NearestUV;
+							TexelQueryTriangle[SampleLinearIdx] = TriId;
 						}
+
 						if (TriId < NumUVSpaceMeshTriCharts && TexelQueryUVChart[TexelLinearIdx] == IndexConstants::InvalidID)
 						{
 							TexelQueryUVChart[TexelLinearIdx] = (*UVSpaceMeshTriCharts)[TriId];
 						}
-						bHitTriID = true;
 					}
-				}
-				if (bIsGutterTexel && bHitTriID)
+				} // end Sample loop
+			} // end scanline loop
+		}, !bParallel ? EParallelForFlags::ForceSingleThread : EParallelForFlags::None);
+
+		return true;
+	}
+
+	bool ComputeGutterTexelsFromGutterSamples()
+	{
+		TAtomic<int64> TotalGutterCounter = 0;
+		TArray< TArray64<TTuple<int64, int64>> > GutterTexelsPerScanline;
+		GutterTexelsPerScanline.SetNum(Tile.GetHeight());
+
+		ParallelFor(Tile.GetHeight(), [this, &TotalGutterCounter, &GutterTexelsPerScanline]
+		(int32 ImgY)
+		{
+			for (int32 ImgX = 0; ImgX < Tile.GetWidth(); ++ImgX)
+			{
+				const FVector2i SourceCoords = Tile.GetSourceCoords(ImgX, ImgY);
+				const int64 SourceTexelLinearIdx = Dimensions.GetIndex(SourceCoords.X, SourceCoords.Y);
+				const int64 TexelLinearIdx = Tile.GetIndex(ImgX, ImgY);
+				const TTuple<int64, int64> InvalidGutterNearestTexel(-1, -1);
+
+				// With multiple samples, a texel may contain multiple gutter samples.
+				// Since we copy nearest interior texel over top our gutter texels, we
+				// should only consider a texel a gutter texel iff all samples in a
+				// texel are gutter texels.
+				TTuple<int64, int64> GutterNearestTexel = InvalidGutterNearestTexel;
+				for (int32 Sample = 0; Sample < PixelSampler.Num(); ++Sample)
+				{
+					const int64 SampleLinearIdx = TexelLinearIdx * PixelSampler.Num() + Sample;
+
+					if (TexelType[SampleLinearIdx] == InteriorTexel)
+					{
+						GutterNearestTexel = InvalidGutterNearestTexel;
+						break;
+					}
+
+					if (TexelType[SampleLinearIdx] == GutterTexel)
+					{
+						const FVector2d NearestUV = (FVector2d)TexelQueryUV[SampleLinearIdx];
+						const FVector2i NearestCoords = Dimensions.UVToCoords(NearestUV);
+
+						// To avoid artifacts with mipmapping the gutter texels store the nearest
+						// interior texel and a post pass copies those nearest interior texel values
+						// to each gutter pixel.
+						//
+						// TODO There can be interior texels closer to the gutter the texel than the one we find here.
+						// Search :NearestInteriorGutterTexel for a related problem in the MeshMapBaker.
+						//
+						const int64 NearestLinearIdx = Dimensions.GetIndex(NearestCoords);
+						GutterNearestTexel = TTuple<int64, int64>(SourceTexelLinearIdx, NearestLinearIdx);
+					}
+				} // end Sample loop
+
+				if (GutterNearestTexel != InvalidGutterNearestTexel)
 				{
 					GutterTexelsPerScanline[ImgY].Add(GutterNearestTexel);
 				}
-			}
+			} // end scanline loop
 
 			TotalGutterCounter += GutterTexelsPerScanline[ImgY].Num();
 
@@ -236,6 +291,28 @@ public:
 		return true;
 	}
 
+	/**
+	 * Computes the image occupancy map from a UV space mesh. This function classifies all the samples in the image,
+	 * and computes the GutterTexels mapping, if you dont want the mapping you can call ClassifySamplesFromUVSpaceMesh
+	 * and compute it yourself (eg for texture baking we use filter kernel coverage to compute the GutterTexels mapping)
+	 * 
+	 * @param UVSpaceMesh UV space mesh to compute the occupancy map from.
+	 * @param GetTriangleIDFunc Lambda to remap a texel's nearest triangle ID
+	 * @param UVSpaceMeshTriCharts Optional UVSpaceMesh triangle ID to UV chart map
+	 */
+	template<typename MeshType, typename GetTriangleIDFuncType>
+	bool ComputeFromUVSpaceMesh(
+		const MeshType& UVSpaceMesh, 
+		GetTriangleIDFuncType GetTriangleIDFunc = [](int32 TriangleID) { return TriangleID; },
+		const TArray<int32>* UVSpaceMeshTriCharts = nullptr)
+	{
+		ClassifySamplesFromUVSpaceMesh(UVSpaceMesh, GetTriangleIDFunc, UVSpaceMeshTriCharts);
+
+		ComputeGutterTexelsFromGutterSamples();
+
+		return true;
+	}
+
 
 
 
diff --git a/Engine/Source/Runtime/GeometryCore/Public/Image/ImageTile.h b/Engine/Source/Runtime/GeometryCore/Public/Image/ImageTile.h
--- a/Engine/Source/Runtime/GeometryCore/Public/Image/ImageTile.h
+++ b/Engine/Source/Runtime/GeometryCore/Public/Image/ImageTile.h
@@ -65,6 +65,12 @@ public:
 		return GetIndex(LocalCoords.X, LocalCoords.Y);
 	}
 
+	/** @return the coordinates into this tile given the linear index into this tile. */
+	FVector2i GetCoords(const int64 LinearIdx) const
+	{
+		return FVector2i((int32)(LinearIdx % Width), (int32)(LinearIdx / Width));
+	}
+
 	/** @return true if the given source coordinates are contained by this tile. */
 	bool Contains(const int32 X, const int32 Y) const
 	{
