fix EdgeLoop.cpp calling FindEdgeFromTri with arguments in the incorrect order

diff --git a/Engine/Plugins/Runtime/GeometryProcessing/Source/DynamicMesh/Private/EdgeLoop.cpp b/Engine/Plugins/Runtime/GeometryProcessing/Source/DynamicMesh/Private/EdgeLoop.cpp
--- a/Engine/Plugins/Runtime/GeometryProcessing/Source/DynamicMesh/Private/EdgeLoop.cpp
+++ b/Engine/Plugins/Runtime/GeometryProcessing/Source/DynamicMesh/Private/EdgeLoop.cpp
@@ -318,7 +318,7 @@ bool UE::Geometry::ConvertTriOrderedEdgeLoopToLoop(const FDynamicMesh3& Mesh,
 		EdgeLoop->Reserve(N);
 		for (int32 k = 0; k < N; ++k)
 		{
-			int32 EdgeID = Mesh.FindEdgeFromTri( TriOrderedEdgesLoopOut[k].TriangleID, VertexLoop[k], VertexLoop[(k+1)%N] );
+			int32 EdgeID = Mesh.FindEdgeFromTri(VertexLoop[k], VertexLoop[(k+1)%N], TriOrderedEdgesLoopOut[k].TriangleID);
 			EdgeLoop->Add(EdgeID);
 		}
 	}
