Fixed mismatched copy constructor/assignment operator in FMeshFaceSelection and FTangentTri2.

diff --git a/Engine/Plugins/Runtime/GeometryProcessing/Source/DynamicMesh/Public/Operations/MeshGeodesicSurfaceTracer.h b/Engine/Plugins/Runtime/GeometryProcessing/Source/DynamicMesh/Public/Operations/MeshGeodesicSurfaceTracer.h
--- a/Engine/Plugins/Runtime/GeometryProcessing/Source/DynamicMesh/Public/Operations/MeshGeodesicSurfaceTracer.h
+++ b/Engine/Plugins/Runtime/GeometryProcessing/Source/DynamicMesh/Public/Operations/MeshGeodesicSurfaceTracer.h
@@ -148,8 +148,6 @@ class GeodesicSingleTriangleUtils::FTangentTri2
 public:
 	FTangentTri2() = default;
 
-	FTangentTri2(const FTangentTri2& other) = default;
-
 	// The PrimaryEdge will be aligned with the x-axis in 2d ( and must be one of the edges of the triangle indicated by TriID)
 	FTangentTri2(const FDynamicMesh3& Mesh, int32 TriID, int32 PrimaryEdgeID);
 
diff --git a/Engine/Plugins/Runtime/GeometryProcessing/Source/DynamicMesh/Public/Selections/MeshFaceSelection.h b/Engine/Plugins/Runtime/GeometryProcessing/Source/DynamicMesh/Public/Selections/MeshFaceSelection.h
--- a/Engine/Plugins/Runtime/GeometryProcessing/Source/DynamicMesh/Public/Selections/MeshFaceSelection.h
+++ b/Engine/Plugins/Runtime/GeometryProcessing/Source/DynamicMesh/Public/Selections/MeshFaceSelection.h
@@ -30,11 +30,6 @@ public:
 	{
 		Mesh = mesh;
 	}
-	FMeshFaceSelection(const FMeshFaceSelection& copy)
-	{
-		Mesh = copy.Mesh;
-		Selected = copy.Selected;
-	}
 
 	// convert vertex selection to face selection. Require at least minCount verts of
 	// tri to be selected (valid values are 1,2,3)
