Fixed mismatched copy constructor/assignment operator in FMeshSurfaceDirection.

diff --git a/Engine/Plugins/Runtime/GeometryProcessing/Source/DynamicMesh/Public/Operations/MeshGeodesicSurfaceTracer.h b/Engine/Plugins/Runtime/GeometryProcessing/Source/DynamicMesh/Public/Operations/MeshGeodesicSurfaceTracer.h
--- a/Engine/Plugins/Runtime/GeometryProcessing/Source/DynamicMesh/Public/Operations/MeshGeodesicSurfaceTracer.h
+++ b/Engine/Plugins/Runtime/GeometryProcessing/Source/DynamicMesh/Public/Operations/MeshGeodesicSurfaceTracer.h
@@ -307,8 +307,6 @@ struct GeodesicSingleTriangleUtils::FMeshSurfaceDirection
 		, Dir(FVector2d(FMath::Cos(PolarAngle), FMath::Sin(PolarAngle)))
 	{}
 
-	FMeshSurfaceDirection(const FMeshSurfaceDirection& Other) = default;
-	
 
 	int32 EdgeID;           // Mesh Edge, has implicit order (from Edge.A to Edge.B)
 	FVector2d  Dir;         // Normalized direction relative to the frame defined by the mesh edge EdgeID aligned with the x-axis
