Fixed mismatched copy constructor/assignment operator in FMeshEdgeSelection.

diff --git a/Engine/Plugins/Runtime/GeometryProcessing/Source/DynamicMesh/Public/Selections/MeshEdgeSelection.h b/Engine/Plugins/Runtime/GeometryProcessing/Source/DynamicMesh/Public/Selections/MeshEdgeSelection.h
--- a/Engine/Plugins/Runtime/GeometryProcessing/Source/DynamicMesh/Public/Selections/MeshEdgeSelection.h
+++ b/Engine/Plugins/Runtime/GeometryProcessing/Source/DynamicMesh/Public/Selections/MeshEdgeSelection.h
@@ -31,11 +31,6 @@ public:
 	{
 		Mesh = mesh;
 	}
-	FMeshEdgeSelection(const FMeshEdgeSelection& copy)
-	{
-		Mesh = copy.Mesh;
-		Selected = copy.Selected;
-	}
 
 
 	// convert vertex selection to edge selection. Require at least minCount verts of edge to be selected
