Fixed mismatched copy constructor/assignment operator in UE::MeshDeformation::CotanTriangleData and UE::MeshDeformation::MeanValueTriangleData.

diff --git a/Engine/Plugins/Runtime/GeometryProcessing/Source/DynamicMesh/Private/Solvers/PrecomputedMeshWeightData.cpp b/Engine/Plugins/Runtime/GeometryProcessing/Source/DynamicMesh/Private/Solvers/PrecomputedMeshWeightData.cpp
--- a/Engine/Plugins/Runtime/GeometryProcessing/Source/DynamicMesh/Private/Solvers/PrecomputedMeshWeightData.cpp
+++ b/Engine/Plugins/Runtime/GeometryProcessing/Source/DynamicMesh/Private/Solvers/PrecomputedMeshWeightData.cpp
@@ -5,24 +5,6 @@
 using namespace UE::Geometry;
 
 
-UE::MeshDeformation::CotanTriangleData::CotanTriangleData(const CotanTriangleData& Other)
-{
-
-	Cotangent[0] = Other.Cotangent[0];
-	Cotangent[1] = Other.Cotangent[1];
-	Cotangent[2] = Other.Cotangent[2];
-
-	VoronoiArea[0] = Other.VoronoiArea[0];
-	VoronoiArea[1] = Other.VoronoiArea[1];
-	VoronoiArea[2] = Other.VoronoiArea[2];
-
-	OppositeEdge[0] = Other.OppositeEdge[0];
-	OppositeEdge[1] = Other.OppositeEdge[1];
-	OppositeEdge[2] = Other.OppositeEdge[2];
-
-}
-
-
 void UE::MeshDeformation::CotanTriangleData::Initialize(const FDynamicMesh3& DynamicMesh, int32 SrcTriId)
 {
 	TriId = SrcTriId;
@@ -126,21 +108,6 @@ void UE::MeshDeformation::CotanTriangleData::Initialize(const FDynamicMesh3& Dyn
 }
 
 
-UE::MeshDeformation::MeanValueTriangleData::MeanValueTriangleData(const MeanValueTriangleData& Other)
-	: TriId(Other.TriId)
-	, TriVtxIds(Other.TriVtxIds)
-	, TriEdgeIds(Other.TriEdgeIds)
-	, bDegenerate(Other.bDegenerate)
-{
-	EdgeLength[0] = Other.EdgeLength[0];
-	EdgeLength[1] = Other.EdgeLength[1];
-	EdgeLength[2] = Other.EdgeLength[2];
-
-	TanHalfAngle[0] = Other.TanHalfAngle[0];
-	TanHalfAngle[1] = Other.TanHalfAngle[1];
-	TanHalfAngle[2] = Other.TanHalfAngle[2];
-}
-
 void UE::MeshDeformation::MeanValueTriangleData::Initialize(const FDynamicMesh3& DynamicMesh, int32 SrcTriId)
 {
 	TriId = SrcTriId;
diff --git a/Engine/Plugins/Runtime/GeometryProcessing/Source/DynamicMesh/Public/Solvers/PrecomputedMeshWeightData.h b/Engine/Plugins/Runtime/GeometryProcessing/Source/DynamicMesh/Public/Solvers/PrecomputedMeshWeightData.h
--- a/Engine/Plugins/Runtime/GeometryProcessing/Source/DynamicMesh/Public/Solvers/PrecomputedMeshWeightData.h
+++ b/Engine/Plugins/Runtime/GeometryProcessing/Source/DynamicMesh/Public/Solvers/PrecomputedMeshWeightData.h
@@ -53,8 +53,6 @@ namespace UE
 
 			CotanTriangleData() = default;
 
-			CotanTriangleData(const CotanTriangleData& Other);
-
 			CotanTriangleData(const FDynamicMesh3& DynamicMesh, int32 TriId)
 			{
 				Initialize(DynamicMesh, TriId);
@@ -155,8 +153,6 @@ namespace UE
 
 		public:
 
-			MeanValueTriangleData(const MeanValueTriangleData& Other);
-
 			void Initialize(const FDynamicMesh3& DynamicMesh, int32 SrcTriId);
 
 			// return Tan(angle / 2) for the corner indicated by this vert id.
