MeshModelingTools: Make the box finding code in MeshSimpleShapeApproximation prefer to orient the boxes such that they align with the world frame better.

diff --git a/Engine/Plugins/Runtime/GeometryProcessing/Source/DynamicMesh/Private/ShapeApproximation/MeshSimpleShapeApproximation.cpp b/Engine/Plugins/Runtime/GeometryProcessing/Source/DynamicMesh/Private/ShapeApproximation/MeshSimpleShapeApproximation.cpp
--- a/Engine/Plugins/Runtime/GeometryProcessing/Source/DynamicMesh/Private/ShapeApproximation/MeshSimpleShapeApproximation.cpp
+++ b/Engine/Plugins/Runtime/GeometryProcessing/Source/DynamicMesh/Private/ShapeApproximation/MeshSimpleShapeApproximation.cpp
@@ -19,6 +19,58 @@
 
 using namespace UE::Geometry;
 
+namespace FMeshSimpleShapeApproximationLocals
+{
+	// Take the box, and while preserving its shape in space, swap its axes around so that they
+	// are pointed roughly toward the world x/y/z. This makes it behave more intuitively under
+	// transformations (for instance, scaling Z grows roughly in the Z direction rather than
+	// to the side if the original box was actually on its "side").
+	FOrientedBox3d ReparameterizeBoxCloserToWorldFrame(const FOrientedBox3d& Box)
+	{
+		FVector3d Axes[3]{ Box.AxisX(), Box.AxisY(), Box.AxisZ() };	
+
+		int32 BestZ = FMath::Max3Index(
+			FMath::Abs(Axes[0].Z),
+			FMath::Abs(Axes[1].Z),
+			FMath::Abs(Axes[2].Z));
+		FVector3d NewZ = Axes[BestZ] * (Axes[BestZ].Z > 0 ? 1 : -1);
+
+		// Pick the best Y out of the two remaining
+		double DotY[3]{
+			Axes[0].Y,
+			Axes[1].Y,
+			Axes[2].Y
+		};
+		DotY[BestZ] = 0; // don't pick this one
+
+		int32 BestY = FMath::Max3Index(
+			FMath::Abs(DotY[0]),
+			FMath::Abs(DotY[1]),
+			FMath::Abs(DotY[2]));
+		FVector3d NewY = Axes[BestY] * (DotY[BestY] > 0 ? 1 : -1);
+
+		// Sum of BestY and BestZ will be either 0+1=1, 0+2=2, or 1+2=3, and the 
+		// corresponding leftover index will be 2, 1, or 0.
+		int32 BestX = 3 - (BestY + BestZ);
+		// Static analyzer probably won't like that, so here's a clamp for safety
+		BestX = FMath::Clamp(BestX, 0, 2);
+
+		// Sanity check to make sure we picked unique axes
+		if (!ensure(BestX != BestY 
+			&& BestY != BestZ 
+			&& BestX != BestZ))
+		{
+			return Box;
+		}
+
+		return FOrientedBox3d (
+			FFrame3d(Box.Frame.Origin,
+				NewY.Cross(NewZ), // better to do this than risk picking the wrong direction for BestX
+				NewY,
+				NewZ),
+			FVector3d(Box.Extents[BestX], Box.Extents[BestY], Box.Extents[BestZ]));
+	}
+}
 
 void FMeshSimpleShapeApproximation::DetectAndCacheSimpleShapeType(const FDynamicMesh3* SourceMesh, FSourceMeshCache& CacheOut)
 {
@@ -62,6 +114,8 @@ bool FMeshSimpleShapeApproximation::GetDetectedSimpleShape(
 	FSimpleShapeSet3d& ShapeSetOut,
 	FCriticalSection& ShapeSetLock)
 {
+	using namespace FMeshSimpleShapeApproximationLocals;
+
 	if (Cache.DetectedType == EDetectedSimpleShapeType::Sphere && bDetectSpheres)
 	{
 		ShapeSetLock.Lock();
@@ -71,8 +125,10 @@ bool FMeshSimpleShapeApproximation::GetDetectedSimpleShape(
 	}
 	else if (Cache.DetectedType == EDetectedSimpleShapeType::Box && bDetectBoxes)
 	{
+		FOrientedBox3d BoxToUse = ReparameterizeBoxCloserToWorldFrame(Cache.DetectedBox);
+
 		ShapeSetLock.Lock();
-		ShapeSetOut.Boxes.Add(Cache.DetectedBox);
+		ShapeSetOut.Boxes.Add(BoxToUse);
 		ShapeSetLock.Unlock();
 		return true;
 	}
@@ -206,6 +262,8 @@ void FMeshSimpleShapeApproximation::Generate_AlignedBoxes(FSimpleShapeSet3d& Sha
 
 void FMeshSimpleShapeApproximation::Generate_OrientedBoxes(FSimpleShapeSet3d& ShapeSetOut, FProgressCancel* Progress)
 {
+	using namespace FMeshSimpleShapeApproximationLocals;
+
 	FCriticalSection GeometryLock;
 	ParallelFor(SourceMeshes.Num(), [&](int32 idx)
 	{
@@ -225,8 +283,10 @@ void FMeshSimpleShapeApproximation::Generate_OrientedBoxes(FSimpleShapeSet3d& Sh
 
 		if (FitResult.bHaveBox)
 		{
+			FOrientedBox3d BoxToUse = ReparameterizeBoxCloserToWorldFrame(FitResult.Box);
+
 			GeometryLock.Lock();
-			ShapeSetOut.Boxes.Add(FBoxShape3d(FitResult.Box));
+			ShapeSetOut.Boxes.Add(FBoxShape3d(BoxToUse));
 			GeometryLock.Unlock();
 		}
 	});
