Renamed FSparseMatrixD.h to SparseMatrix.h

#rnx

diff --git a/Engine/Plugins/Runtime/GeometryProcessing/Source/DynamicMesh/Private/Solvers/Internal/ConstrainedPoissonSolver.h b/Engine/Plugins/Runtime/GeometryProcessing/Source/DynamicMesh/Private/Solvers/Internal/ConstrainedPoissonSolver.h
--- a/Engine/Plugins/Runtime/GeometryProcessing/Source/DynamicMesh/Private/Solvers/Internal/ConstrainedPoissonSolver.h
+++ b/Engine/Plugins/Runtime/GeometryProcessing/Source/DynamicMesh/Private/Solvers/Internal/ConstrainedPoissonSolver.h
@@ -5,7 +5,7 @@
 #include "CoreMinimal.h"
 #include "DynamicMesh/DynamicMesh3.h"
 
-#include "FSparseMatrixD.h"
+#include "SparseMatrix.h"
 #include "MatrixSolver.h"
 #include "Math/NumericLimits.h"
 
diff --git a/Engine/Plugins/Runtime/GeometryProcessing/Source/DynamicMesh/Private/Solvers/Internal/FSparseMatrixD.h b/Engine/Plugins/Runtime/GeometryProcessing/Source/DynamicMesh/Private/Solvers/Internal/FSparseMatrixD.h
deleted file mode 100644
diff --git a/Engine/Plugins/Runtime/GeometryProcessing/Source/DynamicMesh/Private/Solvers/Internal/LaplacianOperators.h b/Engine/Plugins/Runtime/GeometryProcessing/Source/DynamicMesh/Private/Solvers/Internal/LaplacianOperators.h
--- a/Engine/Plugins/Runtime/GeometryProcessing/Source/DynamicMesh/Private/Solvers/Internal/LaplacianOperators.h
+++ b/Engine/Plugins/Runtime/GeometryProcessing/Source/DynamicMesh/Private/Solvers/Internal/LaplacianOperators.h
@@ -7,7 +7,7 @@
 #include "Solvers/MeshLaplacian.h"
 #include "Solvers/MeshLinearization.h"
 #include "Solvers/LaplacianMatrixAssembly.h"
-#include "FSparseMatrixD.h"
+#include "SparseMatrix.h"
 
 // this is a private header
 using namespace UE::Geometry;
diff --git a/Engine/Plugins/Runtime/GeometryProcessing/Source/DynamicMesh/Private/Solvers/Internal/MatrixBase.h b/Engine/Plugins/Runtime/GeometryProcessing/Source/DynamicMesh/Private/Solvers/Internal/MatrixBase.h
--- a/Engine/Plugins/Runtime/GeometryProcessing/Source/DynamicMesh/Private/Solvers/Internal/MatrixBase.h
+++ b/Engine/Plugins/Runtime/GeometryProcessing/Source/DynamicMesh/Private/Solvers/Internal/MatrixBase.h
@@ -3,7 +3,7 @@
 #pragma once 
 
 #include "CoreMinimal.h"
-#include "FSparseMatrixD.h"
+#include "SparseMatrix.h"
 
 enum class EMatrixSolverType
 {
diff --git a/Engine/Plugins/Runtime/GeometryProcessing/Source/DynamicMesh/Private/Solvers/Internal/MeshUVSolver.h b/Engine/Plugins/Runtime/GeometryProcessing/Source/DynamicMesh/Private/Solvers/Internal/MeshUVSolver.h
--- a/Engine/Plugins/Runtime/GeometryProcessing/Source/DynamicMesh/Private/Solvers/Internal/MeshUVSolver.h
+++ b/Engine/Plugins/Runtime/GeometryProcessing/Source/DynamicMesh/Private/Solvers/Internal/MeshUVSolver.h
@@ -10,7 +10,7 @@
 #include "Math/RandomStream.h"
 #include "MatrixBase.h"
 #include "DenseMatrix.h"
-#include "FSparseMatrixD.h"
+#include "SparseMatrix.h"
 
 namespace UE
 {
diff --git a/Engine/Plugins/Runtime/GeometryProcessing/Source/DynamicMesh/Private/Solvers/Internal/SparseMatrix.h b/Engine/Plugins/Runtime/GeometryProcessing/Source/DynamicMesh/Private/Solvers/Internal/SparseMatrix.h
new file mode 100644
--- /dev/null
+++ b/Engine/Plugins/Runtime/GeometryProcessing/Source/DynamicMesh/Private/Solvers/Internal/SparseMatrix.h
@@ -0,0 +1,124 @@
+// Copyright Epic Games, Inc. All Rights Reserved.
+
+#pragma once
+
+#include "CoreMinimal.h"
+#include "Util/ElementLinearization.h"   // TVector3Arrays
+
+#include "Solvers/MatrixInterfaces.h"
+#include "Solvers/LaplacianMatrixAssembly.h"
+
+// According to http://eigen.tuxfamily.org/index.php?title=Main_Page 
+// SimplicialCholesky, AMD ordering, and constrained_cg are disabled.
+
+#ifndef EIGEN_MPL2_ONLY
+#define EIGEN_MPL2_ONLY
+#endif
+
+#if defined(_MSC_VER) && USING_CODE_ANALYSIS
+#pragma warning(push)
+#pragma warning(disable : 6011)
+#pragma warning(disable : 6387)
+#pragma warning(disable : 6313)
+#pragma warning(disable : 6294)
+#endif
+
+#if defined(__clang__) && __clang_major__ > 9
+#pragma clang diagnostic push
+#pragma clang diagnostic ignored "-Wimplicit-int-float-conversion"
+#endif
+
+PRAGMA_DEFAULT_VISIBILITY_START
+THIRD_PARTY_INCLUDES_START
+#include <Eigen/Sparse>
+THIRD_PARTY_INCLUDES_END
+PRAGMA_DEFAULT_VISIBILITY_END
+
+#if defined(__clang__) && __clang_major__ > 9
+#pragma clang diagnostic pop
+#endif
+
+#if defined(_MSC_VER) && USING_CODE_ANALYSIS
+#pragma warning(pop)
+#endif
+
+
+
+// NB: The LU solver likes ColMajor but the CG sovler likes RowMajor
+//     Also, to change everything to float / double just change the scalar type here
+
+typedef Eigen::SparseMatrix<double, Eigen::ColMajor>  FSparseMatrixD;
+
+
+
+
+//
+// Extension of TSparseMatrixAssembler suitable for eigen sparse matrix
+//
+class FEigenSparseMatrixAssembler : public UE::Solvers::TSparseMatrixAssembler<double>
+{
+public:
+	typedef FSparseMatrixD::Scalar    ScalarT;
+	typedef Eigen::Triplet<ScalarT>  MatrixTripletT;
+
+	TUniquePtr<FSparseMatrixD> Matrix;
+	std::vector<MatrixTripletT> EntryTriplets;
+
+	FEigenSparseMatrixAssembler(int32 RowsI, int32 ColsJ)
+	{
+		Matrix = MakeUnique<FSparseMatrixD>(RowsI, ColsJ);
+
+		ReserveEntriesFunc = [this](int32 NumElements)
+		{
+			EntryTriplets.reserve(NumElements);
+		};
+
+		AddEntryFunc = [this](int32 i, int32 j, double Value)
+		{
+			EntryTriplets.push_back(MatrixTripletT(i, j, Value));
+		};
+	}
+
+	void ExtractResult(FSparseMatrixD& Result)
+	{
+		Matrix->setFromTriplets(EntryTriplets.begin(), EntryTriplets.end());
+		Matrix->makeCompressed();
+
+		Result.swap(*Matrix);
+	}
+};
+
+
+
+/**
+* A struct of arrays representation used to hold vertex positions
+* in three vectors that can interface with the eigen library
+*/
+class FSOAPositions : public UE::Geometry::TVector3Arrays<double>
+{
+public:
+	typedef typename FSparseMatrixD::Scalar ScalarType;
+	typedef Eigen::Matrix<ScalarType, Eigen::Dynamic, 1>  RealVectorType;
+	typedef Eigen::Map<RealVectorType> VectorType;
+	typedef Eigen::Map<const RealVectorType> ConstVectorType;
+
+	FSOAPositions(int32 Size) : UE::Geometry::TVector3Arrays<double>(Size)
+	{ }
+
+	FSOAPositions()
+	{}
+
+	VectorType Array(int32 i)
+	{
+		TArray<ScalarType>& Column = (i == 0) ? XVector : (i == 1) ? YVector : ZVector;
+		return (Column.Num() > 0) ? VectorType(&Column[0], Column.Num()) : VectorType(nullptr, 0);
+	}
+	ConstVectorType Array(int32 i) const
+	{
+		const TArray<ScalarType>& Column = (i == 0) ? XVector : (i == 1) ? YVector : ZVector;
+		return (Column.Num() > 0) ? ConstVectorType(&Column[0], Column.Num()) : ConstVectorType(nullptr, 0);
+	}
+
+};
+
+
