-Wrapped FSparseMatrixD in namespace UE::Geometry.
-Removed unused includes.
-Added SparseMatrix.cpp
-Added a function for slicing the sparse matrix.

#rnx

diff --git a/Engine/Plugins/Runtime/GeometryProcessing/Source/DynamicMesh/Private/Solvers/Internal/MatrixBase.h b/Engine/Plugins/Runtime/GeometryProcessing/Source/DynamicMesh/Private/Solvers/Internal/MatrixBase.h
--- a/Engine/Plugins/Runtime/GeometryProcessing/Source/DynamicMesh/Private/Solvers/Internal/MatrixBase.h
+++ b/Engine/Plugins/Runtime/GeometryProcessing/Source/DynamicMesh/Private/Solvers/Internal/MatrixBase.h
@@ -2,9 +2,13 @@
 
 #pragma once 
 
-#include "CoreMinimal.h"
 #include "SparseMatrix.h"
 
+namespace UE
+{
+namespace Geometry
+{
+
 enum class EMatrixSolverType
 {
 	LU /* SparseLU*/,
@@ -97,3 +101,5 @@ public:
 
 };
 
+}
+}
\ No newline at end of file
diff --git a/Engine/Plugins/Runtime/GeometryProcessing/Source/DynamicMesh/Private/Solvers/Internal/MatrixSolver.cpp b/Engine/Plugins/Runtime/GeometryProcessing/Source/DynamicMesh/Private/Solvers/Internal/MatrixSolver.cpp
--- a/Engine/Plugins/Runtime/GeometryProcessing/Source/DynamicMesh/Private/Solvers/Internal/MatrixSolver.cpp
+++ b/Engine/Plugins/Runtime/GeometryProcessing/Source/DynamicMesh/Private/Solvers/Internal/MatrixSolver.cpp
@@ -2,7 +2,10 @@
 
 #include "MatrixSolver.h"
 
-using namespace UE::Geometry;
+namespace UE
+{
+namespace Geometry
+{
 
 IMatrixSolverBase::~IMatrixSolverBase() {};
 IIterativeMatrixSolverBase::~IIterativeMatrixSolverBase() {};
@@ -43,3 +46,6 @@ TUniquePtr<IMatrixSolverBase> ContructMatrixSolver(const EMatrixSolverType& Matr
 
 	return ResultPtr;
 }
+
+}
+}
\ No newline at end of file
diff --git a/Engine/Plugins/Runtime/GeometryProcessing/Source/DynamicMesh/Private/Solvers/Internal/MatrixSolver.h b/Engine/Plugins/Runtime/GeometryProcessing/Source/DynamicMesh/Private/Solvers/Internal/MatrixSolver.h
--- a/Engine/Plugins/Runtime/GeometryProcessing/Source/DynamicMesh/Private/Solvers/Internal/MatrixSolver.h
+++ b/Engine/Plugins/Runtime/GeometryProcessing/Source/DynamicMesh/Private/Solvers/Internal/MatrixSolver.h
@@ -2,7 +2,6 @@
 
 #pragma once 
 
-#include "CoreMinimal.h"
 #include "Async/ParallelFor.h"
 
 #include "MatrixBase.h"
@@ -34,7 +33,10 @@ PRAGMA_DEFAULT_VISIBILITY_END
 
 #include "ProfilingDebugging/ScopedTimers.h"
 
-
+namespace UE
+{
+namespace Geometry
+{
 // Matrix Solver Factory
 TUniquePtr<IMatrixSolverBase> ContructMatrixSolver(const EMatrixSolverType& MatrixSolverType);
 
@@ -268,3 +270,5 @@ typedef TIterativeMatrixSolver<Eigen::ConjugateGradient<FSparseMatrixD, Eigen::L
 typedef TIterativeMatrixSolver<Eigen::BiCGSTAB<FSparseMatrixD, Eigen::IncompleteLUT<FSparseMatrixD::Scalar, int> >> FBiCGMatrixSolver;
 
 
+}
+}
\ No newline at end of file
diff --git a/Engine/Plugins/Runtime/GeometryProcessing/Source/DynamicMesh/Private/Solvers/Internal/SparseMatrix.cpp b/Engine/Plugins/Runtime/GeometryProcessing/Source/DynamicMesh/Private/Solvers/Internal/SparseMatrix.cpp
new file mode 100644
--- /dev/null
+++ b/Engine/Plugins/Runtime/GeometryProcessing/Source/DynamicMesh/Private/Solvers/Internal/SparseMatrix.cpp
@@ -0,0 +1,104 @@
+// Copyright Epic Games, Inc. All Rights Reserved.
+
+#include "SparseMatrix.h"
+
+namespace UE
+{
+namespace Geometry 
+{
+
+bool SliceSparseMatrix(const FSparseMatrixD& InMatrix, 
+                       const TArray<int>& InRowsToSlice,
+                       const TArray<int>& InColsToSlice,
+                       FSparseMatrixD& OutSlicedMatrix)
+{
+    const int32 NumRows = static_cast<int32>(InMatrix.rows());
+    const int32 NumCols = static_cast<int32>(InMatrix.cols());
+    const int32 NumNonZeros = static_cast<int32>(InMatrix.nonZeros());
+    const int32 NumRowsSliced = InRowsToSlice.Num();
+    const int32 NumColsSliced = InColsToSlice.Num();
+
+    if (NumRowsSliced == 0 && NumColsSliced == 0)
+    {
+        // if both slice masks are empty then simply return the copy
+        OutSlicedMatrix = InMatrix;
+        return true;
+    }
+    
+    if (NumRows == 0 || NumCols == 0)
+    {
+        return true; // empty matrix, nothing to do
+    }
+
+    bool bUseAllRows = false;
+    if (NumRowsSliced == 0)
+    {
+        bUseAllRows = true;
+    }
+
+    bool bUseAllCols = false;
+    if (NumColsSliced == 0)
+    {
+        bUseAllCols = true;
+    }
+
+    // Map row/col indices in the input matrix to their row/col indices in the sliced matrix
+    TMap<int, int> ToSlicedRowIndexMap;
+    ToSlicedRowIndexMap.Reserve(NumRowsSliced);
+
+    TMap<int, int> ToSlicedColIndexMap;
+    ToSlicedColIndexMap.Reserve(NumColsSliced);
+
+    for (int Idx = 0; Idx < NumRowsSliced; ++Idx)
+    {
+        const int Row = InRowsToSlice[Idx];
+        if (Row < 0 || Row >= NumRows)
+        {
+            return false;
+        }
+        ToSlicedRowIndexMap.Add(Row, Idx);
+    }
+
+    for (int Idx = 0; Idx < NumColsSliced; ++Idx)
+    {	
+        const int Col = InColsToSlice[Idx];
+        if (Col < 0 || Col >= NumCols)
+        {
+            return false;
+        }
+        ToSlicedColIndexMap.Add(Col, Idx);
+    }
+
+    // Take a guess at the number of the nonzero elements in the sliced matrix
+    std::vector<Eigen::Triplet<FSparseMatrixD::Scalar>> Triplets;
+    const float NonZeroProbabilitytPerEntry = static_cast<float>(NumNonZeros) / static_cast<float>(NumRows * NumCols);
+    const float SlicedNumEntries = static_cast<float>(NumRowsSliced * NumColsSliced);
+    const int SlicedNonZeroEstimate = static_cast<int>(NonZeroProbabilitytPerEntry * SlicedNumEntries);
+    Triplets.reserve(SlicedNonZeroEstimate);
+
+    // Iterate over all of the non-zero elements and check if they are part of the new sliced matrix.
+    for (int OuterIdx = 0; OuterIdx < InMatrix.outerSize(); ++OuterIdx)
+    {
+        for (typename FSparseMatrixD::InnerIterator It(InMatrix, OuterIdx); It; ++It)
+        {	
+            const int Row = static_cast<int>(It.row());
+            const int Col = static_cast<int>(It.col());
+            if ((bUseAllRows || ToSlicedRowIndexMap.Contains(Row)) && (bUseAllCols || ToSlicedColIndexMap.Contains(Col)))
+            {	
+                // Remap the row and col indices to the sliced matrix row and col indices
+                const int SlicedRow = ToSlicedRowIndexMap[Row];
+                const int SlicedCol = ToSlicedColIndexMap[Col];
+                checkSlow(SlicedRow < NumRowsSliced && SlicedCol < NumColsSliced);
+                Triplets.emplace_back(SlicedRow, SlicedCol, It.value());
+            }
+        }
+    }
+
+    OutSlicedMatrix.resize(NumRowsSliced, NumColsSliced);
+    OutSlicedMatrix.setFromTriplets(Triplets.begin(), Triplets.end());
+
+    return true;
+}
+
+}
+}
\ No newline at end of file
diff --git a/Engine/Plugins/Runtime/GeometryProcessing/Source/DynamicMesh/Private/Solvers/Internal/SparseMatrix.h b/Engine/Plugins/Runtime/GeometryProcessing/Source/DynamicMesh/Private/Solvers/Internal/SparseMatrix.h
--- a/Engine/Plugins/Runtime/GeometryProcessing/Source/DynamicMesh/Private/Solvers/Internal/SparseMatrix.h
+++ b/Engine/Plugins/Runtime/GeometryProcessing/Source/DynamicMesh/Private/Solvers/Internal/SparseMatrix.h
@@ -2,11 +2,11 @@
 
 #pragma once
 
-#include "CoreMinimal.h"
 #include "Util/ElementLinearization.h"   // TVector3Arrays
 
-#include "Solvers/MatrixInterfaces.h"
-#include "Solvers/LaplacianMatrixAssembly.h"
+#include "Solvers/MatrixInterfaces.h"    // TSparseMatrixAssembler
+
+#include "Containers/Array.h"
 
 // According to http://eigen.tuxfamily.org/index.php?title=Main_Page 
 // SimplicialCholesky, AMD ordering, and constrained_cg are disabled.
@@ -43,6 +43,11 @@ PRAGMA_DEFAULT_VISIBILITY_END
 #endif
 
 
+namespace UE
+{
+namespace Geometry
+{
+
 
 // NB: The LU solver likes ColMajor but the CG sovler likes RowMajor
 //     Also, to change everything to float / double just change the scalar type here
@@ -121,4 +126,20 @@ public:
 
 };
 
-
+/** 
+ * Slice the matrix such that OutSlicedMatrix = InMatrix[InRowsToSlice, InColsToSlice].
+ * 
+ * @param InMatrix The matrix we are slicing.
+ * @param InRowsToSlice Rows slice mask. If empty all rows will be used. 
+ * @param InColsToSlice Columns slice mask. If empty all columns will be used.
+ * @param OutSlicedMatrix The sliced matrix.
+ * 
+ * @return true if slicing succeeds, false otherwise. Slicing can fail if one of the input parameters is invalid.
+ */
+bool SliceSparseMatrix(const FSparseMatrixD& InMatrix, 
+					   const TArray<int>& InRowsToSlice,
+					   const TArray<int>& InColsToSlice,
+					   FSparseMatrixD& OutSlicedMatrix);
+
+}
+}
\ No newline at end of file
