Added slice dense matrix functions.

diff --git a/Engine/Plugins/Runtime/GeometryProcessing/Source/DynamicMesh/Private/Solvers/Internal/DenseMatrix.cpp b/Engine/Plugins/Runtime/GeometryProcessing/Source/DynamicMesh/Private/Solvers/Internal/DenseMatrix.cpp
new file mode 100644
--- /dev/null
+++ b/Engine/Plugins/Runtime/GeometryProcessing/Source/DynamicMesh/Private/Solvers/Internal/DenseMatrix.cpp
@@ -0,0 +1,90 @@
+// Copyright Epic Games, Inc. All Rights Reserved.
+
+#include "DenseMatrix.h"
+#include "Math/RandomStream.h"
+
+namespace UE
+{
+namespace Geometry
+{
+
+	template <typename ScalarType, int RowsAtCompileTime, int ColsAtCompileTime>
+	void RandomDenseMatrix(TDenseMatrix<ScalarType, RowsAtCompileTime, ColsAtCompileTime>& OutMatrix,
+		const int32 InRowsAtRunTime,
+		const int32 InColsAtRunTime,
+		const FRandomStream& InRandomStream)
+	{
+		auto RandNumFunc = [&InRandomStream]() { return InRandomStream.FRandRange(-1.0f, 1.0f); };
+		OutMatrix = TDenseMatrix<ScalarType, RowsAtCompileTime, ColsAtCompileTime>::NullaryExpr(InRowsAtRunTime, InColsAtRunTime, RandNumFunc);
+	}
+
+	template <typename ScalarType, int RowsAtCompileTime, int ColsAtCompileTime>
+	bool SliceDenseMatrix(const TDenseMatrix<ScalarType, RowsAtCompileTime,ColsAtCompileTime>& InMatrix, 
+	                      const TArray<int>& InRowsToSlice,
+	                      const TArray<int>& InColsToSlice,
+	                      TDenseMatrixX<ScalarType>& OutSlicedMatrix)
+	{
+        const int32 NumRows = static_cast<int32>(InMatrix.rows());
+        const int32 NumCols = static_cast<int32>(InMatrix.cols());
+        int32 NumRowsSliced = InRowsToSlice.Num();
+        int32 NumColsSliced = InColsToSlice.Num();
+
+        if (NumRowsSliced == 0 && NumColsSliced == 0)
+        {
+            // if both slice masks are empty then simply return the copy
+            OutSlicedMatrix = InMatrix;
+            return true;
+        }
+        
+        if (NumRows == 0 || NumCols == 0)
+        {
+            return true; // empty matrix, nothing to do
+        }
+
+        if (NumRowsSliced == 0)
+        {
+            NumRowsSliced = NumRows;
+        }
+
+        if (NumColsSliced == 0)
+        {
+			NumColsSliced = NumCols;
+        }
+        
+        OutSlicedMatrix.resize(NumRowsSliced, NumColsSliced);
+        for (int SlicedRowIndex = 0; SlicedRowIndex < NumRowsSliced; ++SlicedRowIndex)
+        {
+            for (int SlicedColIndex = 0; SlicedColIndex < NumColsSliced; ++SlicedColIndex)
+            {	
+				const int RowIndex = InRowsToSlice[SlicedRowIndex];
+				const int ColIndex = InColsToSlice[SlicedColIndex];
+				if (RowIndex < 0 || RowIndex >= NumRows || ColIndex < 0 || ColIndex >= NumCols)
+				{
+					return false;
+				}
+
+                OutSlicedMatrix(SlicedRowIndex, SlicedColIndex) = InMatrix(RowIndex, ColIndex);
+            }
+        }
+
+        return true;
+	}
+
+	template <typename ScalarType, int RowsAtCompileTime, int ColsAtCompileTime>
+	bool SliceDenseMatrix(const TDenseMatrix<ScalarType, RowsAtCompileTime, ColsAtCompileTime>& InMatrix,
+		const TArray<int>& InRowsToSlice,
+		TDenseMatrixX<ScalarType>& OutSlicedMatrix)
+	{
+		TArray<int> EmptyColArray;
+		SliceDenseMatrix(InMatrix, InRowsToSlice, EmptyColArray, OutSlicedMatrix);
+
+		return true;
+	}
+
+
+	// Explicit template instatiation
+	template void UE::Geometry::RandomDenseMatrix<double, -1, 1>(class Eigen::Matrix<double, -1, 1, 0, -1, 1>&, int, int, struct FRandomStream const&);
+	template bool UE::Geometry::SliceDenseMatrix<double, -1, 1>(class Eigen::Matrix<double, -1, 1, 0, -1, 1> const&, class TArray<int, class TSizedDefaultAllocator<32> > const&, class TArray<int, class TSizedDefaultAllocator<32> > const&, class Eigen::Matrix<double, -1, -1, 0, -1, -1>&); 
+	template bool UE::Geometry::SliceDenseMatrix<double, -1, 1>(class Eigen::Matrix<double, -1, 1, 0, -1, 1> const&, class TArray<int, class TSizedDefaultAllocator<32> > const&, class Eigen::Matrix<double, -1, -1, 0, -1, -1>&);
+}
+}
\ No newline at end of file
diff --git a/Engine/Plugins/Runtime/GeometryProcessing/Source/DynamicMesh/Private/Solvers/Internal/DenseMatrix.h b/Engine/Plugins/Runtime/GeometryProcessing/Source/DynamicMesh/Private/Solvers/Internal/DenseMatrix.h
--- a/Engine/Plugins/Runtime/GeometryProcessing/Source/DynamicMesh/Private/Solvers/Internal/DenseMatrix.h
+++ b/Engine/Plugins/Runtime/GeometryProcessing/Source/DynamicMesh/Private/Solvers/Internal/DenseMatrix.h
@@ -6,8 +6,6 @@
 #define EIGEN_MPL2_ONLY
 #endif
 
-#include "Math/RandomStream.h"
-
 #if defined(_MSC_VER) && USING_CODE_ANALYSIS
 #pragma warning(push)
 #pragma warning(disable : 6011)
@@ -15,15 +13,29 @@
 #pragma warning(disable : 6313)
 #pragma warning(disable : 6294)
 #endif
+
+#if defined(__clang__) && __clang_major__ > 9
+#pragma clang diagnostic push
+#pragma clang diagnostic ignored "-Wimplicit-int-float-conversion"
+#endif
+
 PRAGMA_DEFAULT_VISIBILITY_START
 THIRD_PARTY_INCLUDES_START
 #include <Eigen/Dense>
 THIRD_PARTY_INCLUDES_END
 PRAGMA_DEFAULT_VISIBILITY_END
+
+#if defined(__clang__) && __clang_major__ > 9
+#pragma clang diagnostic pop
+#endif
+
 #if defined(_MSC_VER) && USING_CODE_ANALYSIS
 #pragma warning(pop)
 #endif
 
+// Forward declarations
+struct FRandomStream;
+
 namespace UE
 {
 namespace Geometry
@@ -59,11 +71,37 @@ template <typename ScalarType, int RowsAtCompileTime, int ColsAtCompileTime>
 void RandomDenseMatrix(TDenseMatrix<ScalarType, RowsAtCompileTime, ColsAtCompileTime>& OutMatrix, 
 					   const int32 InRowsAtRunTime,
 					   const int32 InColsAtRunTime,
-					   const FRandomStream& InRandomStream) 
-{
-	auto RandNumFunc = [&InRandomStream]() { return InRandomStream.FRandRange(-1.0f, 1.0f); };
-	OutMatrix = TDenseMatrix<ScalarType, RowsAtCompileTime, ColsAtCompileTime>::NullaryExpr(InRowsAtRunTime, InColsAtRunTime, RandNumFunc);
-}
+					   const FRandomStream& InRandomStream);
+
+/** 
+ * Slice the matrix such that OutSlicedMatrix = InMatrix[InRowsToSlice, InColsToSlice].
+ * 
+ * @param InMatrix The matrix we are slicing.
+ * @param InRowsToSlice Rows slice mask. If empty all rows will be used. 
+ * @param InColsToSlice Columns slice mask. If empty all columns will be used.
+ * @param OutSlicedMatrix The sliced matrix.
+ * 
+ * @return true if slicing succeeds, false otherwise. Slicing can fail if one of the input parameters is invalid.
+ */
+template <typename ScalarType, int RowsAtCompileTime, int ColsAtCompileTime>
+bool SliceDenseMatrix(const TDenseMatrix<ScalarType, RowsAtCompileTime,ColsAtCompileTime>& InMatrix, 
+                      const TArray<int>& InRowsToSlice,
+                      const TArray<int>& InColsToSlice,
+                      TDenseMatrixX<ScalarType>& OutSlicedMatrix);
+
+/**
+ * Slice the matrix only along the rows such that OutSlicedMatrix = InMatrix[InRowsToSlice, :].
+ *
+ * @param InMatrix The matrix we are slicing.
+ * @param InRowsToSlice Rows slice mask. If empty all rows will be used.
+ * @param OutSlicedMatrix The sliced matrix.
+ *
+ * @return true if slicing succeeds, false otherwise. Slicing can fail if one of the input parameters is invalid.
+ */
+template <typename ScalarType, int RowsAtCompileTime, int ColsAtCompileTime>
+bool SliceDenseMatrix(const TDenseMatrix<ScalarType, RowsAtCompileTime, ColsAtCompileTime>& InMatrix,
+					  const TArray<int>& InRowsToSlice,
+					  TDenseMatrixX<ScalarType>& OutSlicedMatrix);
 
 } // end namespace UE::Geometry
 } // end namespace UE
\ No newline at end of file
