Fixed CIS error.

#rnx

diff --git a/Engine/Plugins/Runtime/GeometryProcessing/Source/DynamicMesh/Private/Solvers/Internal/DenseMatrix.h b/Engine/Plugins/Runtime/GeometryProcessing/Source/DynamicMesh/Private/Solvers/Internal/DenseMatrix.h
--- a/Engine/Plugins/Runtime/GeometryProcessing/Source/DynamicMesh/Private/Solvers/Internal/DenseMatrix.h
+++ b/Engine/Plugins/Runtime/GeometryProcessing/Source/DynamicMesh/Private/Solvers/Internal/DenseMatrix.h
@@ -2,6 +2,8 @@
 
 #pragma once
 
+#include "CoreTypes.h"
+
 #ifndef EIGEN_MPL2_ONLY
 #define EIGEN_MPL2_ONLY
 #endif
