Fixed CIS error

#rnx

diff --git a/Engine/Plugins/Runtime/GeometryProcessing/Source/DynamicMesh/Private/Solvers/Internal/DenseMatrix.h b/Engine/Plugins/Runtime/GeometryProcessing/Source/DynamicMesh/Private/Solvers/Internal/DenseMatrix.h
--- a/Engine/Plugins/Runtime/GeometryProcessing/Source/DynamicMesh/Private/Solvers/Internal/DenseMatrix.h
+++ b/Engine/Plugins/Runtime/GeometryProcessing/Source/DynamicMesh/Private/Solvers/Internal/DenseMatrix.h
@@ -3,6 +3,7 @@
 #pragma once
 
 #include "CoreTypes.h"
+#include "Containers/Array.h"
 
 #ifndef EIGEN_MPL2_ONLY
 #define EIGEN_MPL2_ONLY
