Added a quadratic programming class for solving quadratic optimization problems with fixed constraints.

#rnx

diff --git a/Engine/Plugins/Runtime/GeometryProcessing/Source/DynamicMesh/Private/Solvers/Internal/QuadraticProgramming.cpp b/Engine/Plugins/Runtime/GeometryProcessing/Source/DynamicMesh/Private/Solvers/Internal/QuadraticProgramming.cpp
new file mode 100644
--- /dev/null
+++ b/Engine/Plugins/Runtime/GeometryProcessing/Source/DynamicMesh/Private/Solvers/Internal/QuadraticProgramming.cpp
@@ -0,0 +1,198 @@
+// Copyright Epic Games, Inc. All Rights Reserved.
+
+#include "Solvers/Internal/QuadraticProgramming.h"
+#include "MatrixSolver.h"
+
+using namespace UE::Geometry;
+
+FQuadraticProgramming::FQuadraticProgramming(const FSparseMatrixD* InMatrixQ, const FColumnVectorD* InVectorF)
+:
+MatrixQ(InMatrixQ),
+VectorF(InVectorF)
+{
+    checkSlow(InMatrixQ);
+}
+
+FQuadraticProgramming::~FQuadraticProgramming()
+{
+
+}
+
+bool FQuadraticProgramming::SetFixedConstraints(const TArray<int>* InFixedRowIndices, const FSparseMatrixD* InFixedValues)
+{
+    if (InFixedRowIndices && InFixedValues && InFixedRowIndices->Num() == InFixedValues->rows())
+    {
+        if (InFixedRowIndices->Num() <= MatrixQ->rows() && InFixedValues->rows() <= MatrixQ->rows())
+        {
+            FixedRowIndices = InFixedRowIndices;
+            FixedValues = InFixedValues;
+            bFixedConstraintsSet = true; 
+            
+			return true;
+        }
+    }
+
+    return false;
+}
+
+bool FQuadraticProgramming::PreFactorize()
+{   
+    //
+    // Setup the solver depending on which types of constraints are set
+    //
+    if (bFixedConstraintsSet)
+    { 
+        /**
+         * With fixed constraints only, we can rewrite the original optimization problem 
+         *
+         *      trace(0.5*X^t*Q*X + X^t*f)
+         *      
+         * as a sum of sub-problems acting on each column x of X
+         * 
+         *                | Q_vv Q_vf |   |x_v|                    |f_v|
+         *  |x_v x_f|^t * | Q_fv Q_ff | * |x_f|   +  |x_v x_f|^t * |f_f|
+         * 
+         * where: 
+         *      x is a single column of X.
+         *      f is a set of all indices (rows) that are fixed to known values.
+         *      v is a set of the leftover variables we are optimizing for.
+         * 
+         * So Q_vf is a submatrix that is a result of slicing the rows of the original matrix Q by the indices in "v",
+         * and the columns by the indices in "f".
+         * 
+         * We can rewrite our constrained optimization problem as an unconstrained optimization 
+         * problem where we are optimizing for the x_v only:
+         * 
+         *      min (0.5*x_v^t*Q_vv*x_v  +  x_v^t * (f_v + Q_vf)) 
+         *  
+         * Taking the gradient with respect to x_v and setting the result to zero, we get:
+         * 
+         *      x_v = inverse(Q_vv) * -(f_v + Q_vf * x_f);
+         */  
+        const int32 NumParameters = static_cast<int32>(MatrixQ->rows()); // parameters are variable plus fixed values
+        const int32 NumFixed = FixedRowIndices->Num();
+        const int32 NumVariables = NumParameters - NumFixed;
+
+        
+        // Compute the row indices of the variable parameters based on the input fixed parameters
+        TBitArray<FDefaultBitArrayAllocator> IsParameterVariable(true, NumParameters);
+        for (const int FixedRowIdx : *FixedRowIndices)
+        {
+            IsParameterVariable[FixedRowIdx] = false;
+        }
+        
+		// The variable "v" set
+		VariableRowIndices.SetNumUninitialized(NumVariables);
+        for (int ParameterRowIdx = 0, VariableRowIdx = 0; ParameterRowIdx < NumParameters; ++ParameterRowIdx)
+        {
+            if (IsParameterVariable[ParameterRowIdx])
+            {
+				VariableRowIndices[VariableRowIdx++] = ParameterRowIdx;
+            }
+        }
+
+        // The Q_vv matrix
+        FSparseMatrixD VariablesQ;
+		SliceSparseMatrix(*MatrixQ, VariableRowIndices, VariableRowIndices, VariablesQ);
+
+        // Construct a linear solver for a symmetric positive (semi-)definite matrix
+        const EMatrixSolverType MatrixSolverType = EMatrixSolverType::FastestPSD;
+        const bool bIsSymmetric = true;
+        Solver = ContructMatrixSolver(MatrixSolverType);
+        Solver->SetUp(VariablesQ, bIsSymmetric);
+        if (!ensure(Solver->bSucceeded())) 
+        {
+            Solver = nullptr;
+            return false;
+        }
+
+		return true;
+    }
+
+    return false;
+}
+
+bool FQuadraticProgramming::Solve(FDenseMatrixD& Solution)
+{
+    if (bFixedConstraintsSet)
+    {   
+        if (!ensureMsgf(Solver, TEXT("Solver was not setup. Call the SetUp() method first.")))
+        {
+            return false;
+        }
+
+        // The Q_vf matrix
+        FSparseMatrixD VariablesFixedQ;
+		SliceSparseMatrix(*MatrixQ, VariableRowIndices, *FixedRowIndices, VariablesFixedQ);
+
+		// The f_v column vector 
+		FDenseMatrixD VariablesF;
+		if (VectorF)
+		{ 
+			SliceDenseMatrix(*VectorF, VariableRowIndices, VariablesF);
+		}
+		
+        Solution.resize(MatrixQ->rows(), FixedValues->cols());
+
+        // We can iterate over every column of X and solve for each column individually
+        for (int ColIdx = 0; ColIdx < Solution.cols(); ++ColIdx) //TODO: Parallelize
+        {
+            const FColumnVectorD FixedVector = FixedValues->col(ColIdx);
+
+			// The -(f_v + Q_vf * x_f) vector. If the linear coefficients were not specified, skip f_v.
+            FColumnVectorD VectorB;
+            if (VectorF) 
+            {
+				VectorB = -1 * (VariablesF + VariablesFixedQ * FixedVector);
+            } 
+            else 
+            {
+                VectorB = -1 * VariablesFixedQ * FixedVector;
+            }
+
+            FColumnVectorD VariablesVector;
+            Solver->Solve(VectorB, VariablesVector);
+            if (ensure(Solver->bSucceeded()))
+            {   
+                // Copy over the fixed and variable values to the solution matrix
+                for (int RowIdx = 0; RowIdx < FixedRowIndices->Num(); ++RowIdx)
+                {
+                    Solution((*FixedRowIndices)[RowIdx], ColIdx) = FixedVector(RowIdx);
+                }
+
+                for (int RowIdx = 0; RowIdx < VariableRowIndices.Num(); ++RowIdx)
+                {
+                    Solution(VariableRowIndices[RowIdx], ColIdx) = VariablesVector(RowIdx);
+                }
+            }
+			else 
+			{
+				return false;
+			}
+        }
+
+		return true;
+    }
+
+	return false;
+}
+
+bool FQuadraticProgramming::SolveWithFixedConstraints(const FSparseMatrixD& MatrixQ, 
+                                                      const FColumnVectorD& VectorF, 
+                                                      const TArray<int>& FixedRowIndices, 
+                                                      const FSparseMatrixD& FixedValues, 
+                                                      FDenseMatrixD& Solution)
+{
+    FQuadraticProgramming QP(&MatrixQ, &VectorF);
+    if (!QP.SetFixedConstraints(&FixedRowIndices, &FixedValues))
+    {
+        return false;
+    }
+    
+    if (!QP.PreFactorize())
+    {
+        return false;
+    }
+
+    return QP.Solve(Solution);
+}
\ No newline at end of file
diff --git a/Engine/Plugins/Runtime/GeometryProcessing/Source/DynamicMesh/Private/Solvers/Internal/QuadraticProgramming.h b/Engine/Plugins/Runtime/GeometryProcessing/Source/DynamicMesh/Private/Solvers/Internal/QuadraticProgramming.h
new file mode 100644
--- /dev/null
+++ b/Engine/Plugins/Runtime/GeometryProcessing/Source/DynamicMesh/Private/Solvers/Internal/QuadraticProgramming.h
@@ -0,0 +1,107 @@
+// Copyright Epic Games, Inc. All Rights Reserved.
+
+#pragma once
+
+#include "Containers/Array.h"
+#include "Templates/UniquePtr.h"
+#include "SparseMatrix.h"
+#include "DenseMatrix.h"
+
+// Forward declarations
+namespace UE::Geometry
+{
+class IMatrixSolverBase;
+}
+
+
+namespace UE
+{
+namespace Geometry
+{
+/**
+ * This class provides functionality to solve quadratic programming problems.
+ * 
+ * The goal is to minimize the function:
+ *
+ *      trace(0.5*X^t*Q*X + X^t*f)
+ *
+ *      where:
+ *          Q \in R^(n x n) is a symmetric positive (semi-)definite matrix of quadratic coefficients.
+ *          X \in R^(n x m) is a matrix of parameters we optimize for.
+ *          f \in R^(n x 1) is a column vector of linear coefficients.
+ *
+ *      subject to:
+ *
+ *          Fixed Constraints: 
+ *              X[FixedRowIndices, :] = FixedValues
+ *                  where FixedRowIndices \in Z+^(k x 1), FixedValues \in R^(k x m) and k is the number of fixed rows.
+ *
+ *          TODO:
+ *          Linear Equality Constraints:
+ *          Linear Inequality Constraints:
+ * 
+ * Example usage:
+ *      - The solver is initialized with the quadratic Q and linear coefficients f:
+ *          FSparseMatrixD Q = ...
+ *          FColumnVectorD f = ...
+ *          FQuadraticProgramming QP(&Q, &f);
+ * 
+ *      - Then constraints are set by calling Set[ConstraintType]Constraints methods (only the fixed constraints are supported at the moment).
+ *          TArray<int> FixedIndices = ...
+ *          FSparseMatrix FixedValues = ...
+ *          QP.SetFixedConstraints(&FixedIndices, &FixedValues);
+ * 
+ *      - Then PreFactorize() method needs to be called to pre-factorize the matrices and to set up the internal linear solver. 
+ *           QP.PreFactorize(Solution);
+ * 
+ *      - Finally call Solve():
+ *          FDenseMatrixD Solution;
+ *          QP.Solve(Solution);
+ * 
+ *       Calling the PreFactorize() method can be skipped if: 
+ *          Fixed constraints:
+ *              If only the values of the fixed constraints changed from the last time the PreFactorize() was called.
+ *          
+ *          TODO:
+ *          Linear Equality Constraints:
+ *          Linear Inequality Constraints:
+ */
+
+class DYNAMICMESH_API FQuadraticProgramming
+{
+public:
+
+    FQuadraticProgramming(const FSparseMatrixD* InMatrixQ, const FColumnVectorD* InVectorF = nullptr);
+
+	~FQuadraticProgramming();
+
+    bool SetFixedConstraints(const TArray<int>* InFixedRowIndices, const FSparseMatrixD* InFixedValues);
+
+    /**  Pre-factorizes the matrices and sets up the solver. */
+    bool PreFactorize();
+
+    //TODO: Allow to solve for sparse matrices with a passed float threshold value that determines if the value is close enough to 0 
+    bool Solve(FDenseMatrixD& Solution);
+
+
+    // 
+    // Helper "one-function call" methods for setting up and solving the QP problems
+    //
+    static bool SolveWithFixedConstraints(const FSparseMatrixD& MatrixQ, const FColumnVectorD& VectorF, const TArray<int>& FixedRowIndices, const FSparseMatrixD& FixedValues, FDenseMatrixD& Solution);
+
+protected:
+
+    const FSparseMatrixD* MatrixQ = nullptr; // Matrix of quadratic coefficient
+    const FColumnVectorD* VectorF = nullptr; // Column vector of linear coefficients
+
+    // Fixed constraints
+    bool bFixedConstraintsSet = false;  // set to true if SetFixedConstraints is called and successful
+    const TArray<int>* FixedRowIndices = nullptr;  // row indices of the fixed parameters, set by the user when calling SetFixedConstraints
+    const FSparseMatrixD* FixedValues = nullptr;   // matrix of fixed values, set by the user when calling SetFixedConstraints
+    TArray<int> VariableRowIndices;  // row indices of the variable parameters, computed internally by the call to SetUp()
+
+    TUniquePtr<IMatrixSolverBase> Solver = nullptr; // SetUp() method pre-factorizes the matrices and setups the solver
+};
+
+}
+}
