Extra comments for the ConstructFullCotangentLaplacian function.

#rnx

diff --git a/Engine/Plugins/Runtime/GeometryProcessing/Source/DynamicMesh/Public/Solvers/LaplacianMatrixAssembly.h b/Engine/Plugins/Runtime/GeometryProcessing/Source/DynamicMesh/Public/Solvers/LaplacianMatrixAssembly.h
--- a/Engine/Plugins/Runtime/GeometryProcessing/Source/DynamicMesh/Public/Solvers/LaplacianMatrixAssembly.h
+++ b/Engine/Plugins/Runtime/GeometryProcessing/Source/DynamicMesh/Public/Solvers/LaplacianMatrixAssembly.h
@@ -216,6 +216,19 @@ namespace UE
 		 * Construct sparse Cotangent Laplacian matrix.
 		 * This variant combines the N interior and M boundary vertices into a single (N+M) matrix and does not do any special treatment of boundaries,
 		 * they just get standard Cotan weights
+		 *
+		 * If you don't want to move the boundary vertices to the end, then initialize the VertexMap with bRemapBoundary 
+		 * set to false.
+		 * 
+		 * 
+		 * Example usage:
+		 * 
+		 * If you want to get the standard cotangent weight matrix for a mesh with no boundary reordering:
+		 * 
+		 * FDynamicMesh3 DynamicMesh;
+		 * const bool bRemapBoundary = false;
+		 * FVertexLinearization VertexMap(DynamicMesh, bRemapBoundary);
+		 * ConstructFullCotangentLaplacian(DynamicMesh, VertexMap, ECotangentWeightMode::Default, ECotangentAreaMode::NoArea);
 		 */
 		template<typename RealType>
 		void ConstructFullCotangentLaplacian(const FDynamicMesh3& DynamicMesh, const FVertexLinearization& VertexMap,
