Fixed typos. Added a missing include that caused Linux CIS errors. Removed redundant destructor.

#rnx

diff --git a/Engine/Plugins/Runtime/GeometryProcessing/Source/DynamicMesh/Private/Solvers/Internal/QuadraticProgramming.cpp b/Engine/Plugins/Runtime/GeometryProcessing/Source/DynamicMesh/Private/Solvers/Internal/QuadraticProgramming.cpp
--- a/Engine/Plugins/Runtime/GeometryProcessing/Source/DynamicMesh/Private/Solvers/Internal/QuadraticProgramming.cpp
+++ b/Engine/Plugins/Runtime/GeometryProcessing/Source/DynamicMesh/Private/Solvers/Internal/QuadraticProgramming.cpp
@@ -13,11 +13,6 @@ VectorF(InVectorF)
     checkSlow(InMatrixQ);
 }
 
-FQuadraticProgramming::~FQuadraticProgramming()
-{
-
-}
-
 bool FQuadraticProgramming::SetFixedConstraints(const TArray<int>* InFixedRowIndices, const FSparseMatrixD* InFixedValues)
 {
     if (InFixedRowIndices && InFixedValues && InFixedRowIndices->Num() == InFixedValues->rows())
diff --git a/Engine/Plugins/Runtime/GeometryProcessing/Source/DynamicMesh/Private/Solvers/Internal/QuadraticProgramming.h b/Engine/Plugins/Runtime/GeometryProcessing/Source/DynamicMesh/Private/Solvers/Internal/QuadraticProgramming.h
--- a/Engine/Plugins/Runtime/GeometryProcessing/Source/DynamicMesh/Private/Solvers/Internal/QuadraticProgramming.h
+++ b/Engine/Plugins/Runtime/GeometryProcessing/Source/DynamicMesh/Private/Solvers/Internal/QuadraticProgramming.h
@@ -6,13 +6,7 @@
 #include "Templates/UniquePtr.h"
 #include "SparseMatrix.h"
 #include "DenseMatrix.h"
-
-// Forward declarations
-namespace UE::Geometry
-{
-class IMatrixSolverBase;
-}
-
+#include "MatrixBase.h"
 
 namespace UE
 {
@@ -73,8 +67,6 @@ public:
 
     FQuadraticProgramming(const FSparseMatrixD* InMatrixQ, const FColumnVectorD* InVectorF = nullptr);
 
-	~FQuadraticProgramming();
-
     bool SetFixedConstraints(const TArray<int>* InFixedRowIndices, const FSparseMatrixD* InFixedValues);
 
     /**  Pre-factorizes the matrices and sets up the solver. */
@@ -98,9 +90,9 @@ protected:
     bool bFixedConstraintsSet = false;  // set to true if SetFixedConstraints is called and successful
     const TArray<int>* FixedRowIndices = nullptr;  // row indices of the fixed parameters, set by the user when calling SetFixedConstraints
     const FSparseMatrixD* FixedValues = nullptr;   // matrix of fixed values, set by the user when calling SetFixedConstraints
-    TArray<int> VariableRowIndices;  // row indices of the variable parameters, computed internally by the call to SetUp()
+    TArray<int> VariableRowIndices;  // row indices of the variable parameters, computed internally by the call to PreFactorize()
 
-    TUniquePtr<IMatrixSolverBase> Solver = nullptr; // SetUp() method pre-factorizes the matrices and setups the solver
+    TUniquePtr<IMatrixSolverBase> Solver = nullptr; // PreFactorize() method pre-factorizes the matrices and setups the solver
 };
 
 }
