Added a new transfer weight method to the FTransferBoneWeights operator.

#rnx

diff --git a/Engine/Plugins/Experimental/ChaosClothAssetEditor/Source/ChaosClothAssetDataflowNodes/Private/ChaosClothAsset/TransferSkinWeightsNode.cpp b/Engine/Plugins/Experimental/ChaosClothAssetEditor/Source/ChaosClothAssetDataflowNodes/Private/ChaosClothAsset/TransferSkinWeightsNode.cpp
--- a/Engine/Plugins/Experimental/ChaosClothAssetEditor/Source/ChaosClothAssetDataflowNodes/Private/ChaosClothAsset/TransferSkinWeightsNode.cpp
+++ b/Engine/Plugins/Experimental/ChaosClothAssetEditor/Source/ChaosClothAssetDataflowNodes/Private/ChaosClothAsset/TransferSkinWeightsNode.cpp
@@ -130,7 +130,7 @@ void FChaosClothAssetTransferSkinWeightsNode::Evaluate(Dataflow::FContext& Conte
 						const FVector3d PosD = FVector3d(SimPositions[VertexID]);
 
 						UE::AnimationCore::FBoneWeights BoneWeights;
-						TransferBoneWeights.Compute(PosD, FTransformSRT3d::Identity(), BoneWeights, &TargetBoneToIndex);
+						TransferBoneWeights.TransferWeightsToPoint(BoneWeights, PosD, &TargetBoneToIndex);
 
 						const int32 NumBones = BoneWeights.Num();
 
@@ -160,7 +160,7 @@ void FChaosClothAssetTransferSkinWeightsNode::Evaluate(Dataflow::FContext& Conte
 						const FVector3d PosD = FVector3d(RenderPositions[VertexID]);
 
 						UE::AnimationCore::FBoneWeights BoneWeights;
-						TransferBoneWeights.Compute(PosD, FTransformSRT3d::Identity(), BoneWeights, &TargetBoneToIndex);
+						TransferBoneWeights.TransferWeightsToPoint(BoneWeights, PosD, &TargetBoneToIndex);
 
 						const int32 NumBones = BoneWeights.Num();
 
diff --git a/Engine/Plugins/Experimental/ChaosClothAssetEditor/Source/ChaosClothAssetEditorTools/Private/ChaosClothAsset/ClothTransferSkinWeightsTool.cpp b/Engine/Plugins/Experimental/ChaosClothAssetEditor/Source/ChaosClothAssetEditorTools/Private/ChaosClothAsset/ClothTransferSkinWeightsTool.cpp
--- a/Engine/Plugins/Experimental/ChaosClothAssetEditor/Source/ChaosClothAssetEditorTools/Private/ChaosClothAsset/ClothTransferSkinWeightsTool.cpp
+++ b/Engine/Plugins/Experimental/ChaosClothAssetEditor/Source/ChaosClothAssetEditorTools/Private/ChaosClothAsset/ClothTransferSkinWeightsTool.cpp
@@ -77,10 +77,11 @@ namespace UE::Chaos::ClothAsset::Private
 
 			// Do the transfer
 			FTransferBoneWeights TransferBoneWeights(SourceMesh.Get(), FSkeletalMeshAttributes::DefaultSkinWeightProfileName);
+			TransferBoneWeights.TargetToWorld = FTransformSRT3d(SourceMeshTransform.Inverse());
 			bool bComputeOK = false;
 			if (TransferBoneWeights.Validate() == EOperationValidationResult::Ok)
 			{
-				bComputeOK = TransferBoneWeights.Compute(*ResultMesh, FTransformSRT3d(SourceMeshTransform.Inverse()), FSkeletalMeshAttributes::DefaultSkinWeightProfileName);
+				bComputeOK = TransferBoneWeights.TransferWeightsToMesh(*ResultMesh, FSkeletalMeshAttributes::DefaultSkinWeightProfileName);
 			}
 
 			FGeometryResult OpResult;
diff --git a/Engine/Plugins/Experimental/GeometryScripting/Source/GeometryScriptingCore/Private/MeshBoneWeightFunctions.cpp b/Engine/Plugins/Experimental/GeometryScripting/Source/GeometryScriptingCore/Private/MeshBoneWeightFunctions.cpp
--- a/Engine/Plugins/Experimental/GeometryScripting/Source/GeometryScriptingCore/Private/MeshBoneWeightFunctions.cpp
+++ b/Engine/Plugins/Experimental/GeometryScripting/Source/GeometryScriptingCore/Private/MeshBoneWeightFunctions.cpp
@@ -415,7 +415,7 @@ UDynamicMesh* UGeometryScriptLibrary_MeshBoneWeightFunctions::TransferBoneWeight
 				}
 			}
 			
-			if (!TransferBoneWeights.Compute(EditMesh, FTransformSRT3d::Identity(), Options.TargetProfile.GetProfileName()))
+			if (!TransferBoneWeights.TransferWeightsToMesh(EditMesh, Options.TargetProfile.GetProfileName()))
 			{
 				UE::Geometry::AppendError(Debug, EGeometryScriptErrorType::OperationFailed, LOCTEXT("TransferBoneWeightsFromMesh_TransferFailed", "TransferBoneWeightsFromMesh: Failed to transfer the weights"));
 			}
diff --git a/Engine/Plugins/Runtime/GeometryProcessing/Source/DynamicMesh/Private/Operations/TransferBoneWeights.cpp b/Engine/Plugins/Runtime/GeometryProcessing/Source/DynamicMesh/Private/Operations/TransferBoneWeights.cpp
--- a/Engine/Plugins/Runtime/GeometryProcessing/Source/DynamicMesh/Private/Operations/TransferBoneWeights.cpp
+++ b/Engine/Plugins/Runtime/GeometryProcessing/Source/DynamicMesh/Private/Operations/TransferBoneWeights.cpp
@@ -12,6 +12,8 @@
 #include "BoneWeights.h"
 #include "IndexTypes.h"
 #include "TransformTypes.h"
+#include "Solvers/Internal/QuadraticProgramming.h"
+#include "Solvers/LaplacianMatrixAssembly.h"
 
 using namespace UE::AnimationCore;
 using namespace UE::Geometry;
@@ -161,10 +163,15 @@ EOperationValidationResult FTransferBoneWeights::Validate()
 		return EOperationValidationResult::Failed_UnknownReason;
 	}
 
+	if (NormalThreshold >= 0 && !SourceMesh->Attributes()->PrimaryNormals())
+	{
+		return EOperationValidationResult::Failed_UnknownReason;
+	}
+
 	return EOperationValidationResult::Ok;
 }
 
-bool FTransferBoneWeights::Compute(FDynamicMesh3& InOutTargetMesh, const FTransformSRT3d& InToWorld, const FName& InTargetProfileName)
+bool FTransferBoneWeights::TransferWeightsToMesh(FDynamicMesh3& InOutTargetMesh, const FName& InTargetProfileName)
 {	
 	if (Validate() != EOperationValidationResult::Ok) 
 	{
@@ -180,6 +187,11 @@ bool FTransferBoneWeights::Compute(FDynamicMesh3& InOutTargetMesh, const FTransf
 	{
 		return false; // the target mesh must have bone attributes
 	}
+
+	if (NormalThreshold >= 0 && !InOutTargetMesh.Attributes()->PrimaryNormals())
+	{
+		return false; // the target mesh must have normal attributes if we are comparing source and target normals
+	}
 	
 	FDynamicMeshVertexSkinWeightsAttribute* TargetSkinWeights = TransferBoneWeightsLocals::GetOrCreateSkinWeightsAttribute(InOutTargetMesh, InTargetProfileName);
 	checkSlow(TargetSkinWeights);
@@ -212,28 +224,205 @@ bool FTransferBoneWeights::Compute(FDynamicMesh3& InOutTargetMesh, const FTransf
 
 	bool bFailed = false;
 	
-	ParallelFor(InOutTargetMesh.MaxVertexID(), [&](int32 VertexID)
+	MatchedVertices.Init(false, InOutTargetMesh.MaxVertexID());
+
+	if (TransferMethod == ETransferBoneWeightsMethod::ClosestPointOnSurface)
 	{
-		if (Cancelled() || bFailed) 
+		ParallelFor(InOutTargetMesh.MaxVertexID(), [&](int32 VertexID)
 		{
-			return;
+			if (Cancelled()) 
+			{
+				return;
+			}
+			
+			if (InOutTargetMesh.IsVertex(VertexID)) 
+			{
+				const FVector3d Point = InOutTargetMesh.GetVertex(VertexID);
+				const FVector3f Normal = NormalThreshold >= 0 ? InOutTargetMesh.Attributes()->PrimaryNormals()->GetElement(VertexID) : FVector3f::Zero();
+
+				FBoneWeights Weights;
+				if (TransferWeightsToPoint(Weights, Point, TargetBoneToIndex.Get(), Normal))
+				{
+					TargetSkinWeights->SetValue(VertexID, Weights);
+					MatchedVertices[VertexID] = true;
+				}
+			}
+
+		}, bUseParallel ? EParallelForFlags::None : EParallelForFlags::ForceSingleThread);
+
+		// If the caller requested to simply find the closest point for all vertices then the number of matched vertices
+		// must be equal to the target mesh vertex count
+		if (SearchRadius < 0 && NormalThreshold < 0)
+		{
+			int NumMatched = 0;
+			for (bool Flag : MatchedVertices)
+			{
+				if (Flag) 
+				{ 
+					NumMatched++;
+				}
+			}
+			bFailed = NumMatched != InOutTargetMesh.VertexCount();
 		}
+	} 
+	else if (TransferMethod == ETransferBoneWeightsMethod::InpaintWeights)
+	{
+		/**
+         *  Given two meshes, Mesh1 without weights and Mesh2 with weights, assume they are aligned in 3d space.
+         *  For every vertex on Mesh1 find the closest point on the surface of Mesh2 within a radius R. If the difference 
+         *  between the normals of the two points is below the threshold, then it's a match. Otherwise no match.
+         *  So now we have two sets of vertices on Mesh1. One with a match on the source mesh and one without a match.
+         *  For all the vertices with a match, copy weights over. For all the vertices without the match, do nothing.
+         *  Now, for all the vertices without a match, try to approximate the weights by smoothly interpolating between 
+         *  the weights at the known vertices via solving a quadratic problem. 
+         *  
+         *  The solver minimizes an energy
+         *      trace(W^t Q W)
+         *      W \in R^(nxm) is a matrix where n is the number of vertices and m is the number of bones. So (i,j) entry is 
+         *                    the influence (weight) of a vertex i by bone j
+         *      Q \in R^(nxn) is a matrix that combines both Dirichlet and Laplacian energies, Q = -0.5L + 0.5LM^-1L
+         *                    where L is a cotangent Laplacian and M is a mass matrix
+         *  
+         *  subject to constraints
+         *      All weights at a single vertex sum to 1: sum(W(i,:)) = 1
+         *      All weights must be non-negative: W(i,j) >=0 for any i, j
+         *      Any vertex for which we found a match must have fixed weights that can't be changed, 
+         *      i.e. W(i,j) = KnownWeights(i,j) where i is a vertex for which we found a match on the body.
+		 */
 		
-		if (InOutTargetMesh.IsVertex(VertexID)) 
+		// For every vertex on the target mesh try to find the match on the source mesh using the distance and normal checks
+		ParallelFor(InOutTargetMesh.MaxVertexID(), [&](int32 VertexID)
 		{
-			FVector3d Point = InOutTargetMesh.GetVertex(VertexID);
-		
-			FBoneWeights Weights;
-			if (Compute(Point, InToWorld, Weights, TargetBoneToIndex.Get()) == false)
+			if (Cancelled()) 
 			{
-				bFailed = true;
 				return;
 			}
 			
-			TargetSkinWeights->SetValue(VertexID, Weights);
+			if (InOutTargetMesh.IsVertex(VertexID)) 
+			{
+				const FVector3d Point = InOutTargetMesh.GetVertex(VertexID);
+				const FVector3f Normal = NormalThreshold >= 0 ? InOutTargetMesh.Attributes()->PrimaryNormals()->GetElement(VertexID) : FVector3f::Zero();
+
+				FBoneWeights Weights;
+				if (TransferWeightsToPoint(Weights, Point, TargetBoneToIndex.Get(), Normal))
+				{
+					TargetSkinWeights->SetValue(VertexID, Weights);
+					MatchedVertices[VertexID] = true;
+				}
+			}
+		}, bUseParallel ? EParallelForFlags::None : EParallelForFlags::ForceSingleThread);
+
+		if (Cancelled()) 
+		{
+			return false;
+		}
+
+		// Compute linearization so we can store constraints at linearized indices
+		FVertexLinearization VtxLinearization(InOutTargetMesh, false);
+		const TArray<int32>& ToMeshV = VtxLinearization.ToId();
+		const TArray<int32>& ToIndex = VtxLinearization.ToIndex();
+		
+		int32 NumMatched = 0;
+		for (bool Flag : MatchedVertices)
+		{
+			if (Flag) 
+			{ 
+				NumMatched++;
+			}
+		}
+
+		// If all vertices were matched then nothing else to do
+		if (NumMatched == InOutTargetMesh.VertexCount())
+		{
+			return true;
+		}
+
+		// Setup the sparse matrix FixedValues of known (matched) weight values and the array (FixedIndices) of the matched vertex IDs
+		const int TargetNumBones = InOutTargetMesh.Attributes()->GetBoneNames()->GetAttribValues().Num();
+		FSparseMatrixD FixedValues;
+		FixedValues.resize(NumMatched, TargetNumBones);
+		std::vector<Eigen::Triplet<FSparseMatrixD::Scalar>> FixedValuesTriplets;
+		FixedValuesTriplets.reserve(NumMatched);
+		
+		TArray<int> FixedIndices;
+		FixedIndices.Reserve(NumMatched);
+
+		for (int VertexID = 0; VertexID < InOutTargetMesh.MaxVertexID(); ++VertexID)
+		{
+			if (InOutTargetMesh.IsVertex(VertexID) && MatchedVertices[VertexID])
+			{
+				FBoneWeights Data;
+				TargetSkinWeights->GetValue(VertexID, Data);
+
+				const int NumBones = Data.Num();
+				checkSlow(NumBones > 0);
+				const int CurIdx = FixedIndices.Num();
+				for (int BoneID = 0; BoneID < NumBones; ++BoneID)
+				{
+					const int BoneIdx = Data[BoneID].GetBoneIndex();
+					const double BoneWeight = (double)Data[BoneID].GetWeight();
+					FixedValuesTriplets.emplace_back(CurIdx, BoneIdx, BoneWeight);
+				}
+
+				checkSlow(VertexID < ToIndex.Num());
+				FixedIndices.Add(ToIndex[VertexID]);
+			}
 		}
+		FixedValues.setFromTriplets(FixedValuesTriplets.begin(), FixedValuesTriplets.end());
+
+		// Setup the Laplacian matrix
+		const int32 NumVerts = VtxLinearization.NumVerts();
+		FEigenSparseMatrixAssembler LaplacianAssembler(NumVerts, NumVerts);
+		UE::MeshDeformation::ConstructFullCotangentLaplacian<double>(InOutTargetMesh, 
+																	 VtxLinearization, 
+																	 LaplacianAssembler,
+																	 UE::MeshDeformation::ECotangentWeightMode::Default,
+																	 UE::MeshDeformation::ECotangentAreaMode::NoArea);
+		FSparseMatrixD CotangentMatrix;
+		LaplacianAssembler.ExtractResult(CotangentMatrix);
+		FSparseMatrixD NegativeCotangentMatrix = -0.25 * CotangentMatrix; //TODO: Add LM^-1L term in the future
+
+		// Solve the QP problem with fixed constraints
+		bFailed = true;
+		FQuadraticProgramming QProgram(&NegativeCotangentMatrix);
+		if (ensure(QProgram.SetFixedConstraints(&FixedIndices, &FixedValues)))
+		{
+			if (ensure(QProgram.PreFactorize()))
+			{
+				FDenseMatrixD TargetWeights;
+				if (ensure(QProgram.Solve(TargetWeights)))
+				{
+					FBoneWeightsSettings BoneSettings;
+					BoneSettings.SetNormalizeType(EBoneWeightNormalizeType::None);
+					for (int IdxI = 0; IdxI < TargetWeights.rows(); ++IdxI)
+					{
+						FBoneWeights WeightArray;
+						for (int IdxJ = 0; IdxJ < TargetWeights.cols(); ++IdxJ)
+						{
+							const FBoneIndexType BoneId = static_cast<FBoneIndexType>(IdxJ);
+							const float Weight = (float)TargetWeights(IdxI, IdxJ);
+							if (Weight > KINDA_SMALL_NUMBER)
+							{
+								FBoneWeight Bweight(BoneId, Weight);
+								WeightArray.SetBoneWeight(Bweight, BoneSettings);
+							}
+						}
+
+						WeightArray.Renormalize(FBoneWeightsSettings());
+
+						checkSlow(IdxI < ToMeshV.Num());
+						TargetSkinWeights->SetValue(ToMeshV[IdxI], WeightArray);
+					}
 
-	}, bUseParallel ? EParallelForFlags::None : EParallelForFlags::ForceSingleThread);
+					bFailed = false;
+				}
+			}
+		}
+	}
+	else 
+	{
+		checkNoEntry(); // unsupported method
+	}
 
 	if (Cancelled() || bFailed) 
 	{
@@ -243,20 +432,71 @@ bool FTransferBoneWeights::Compute(FDynamicMesh3& InOutTargetMesh, const FTransf
 	return true;
 }
 
-bool FTransferBoneWeights::Compute(const FVector3d& InPoint, const FTransformSRT3d& InToWorld, FBoneWeights& OutWeights, const TMap<FName, uint16>* TargetBoneToIndex) 
-{
-	const FDynamicMeshVertexSkinWeightsAttribute* SourceSkinWeights = SourceMesh->Attributes()->GetSkinWeightsAttribute(SourceProfileName);
-	checkSlow(SourceSkinWeights);
+bool FTransferBoneWeights::TransferWeightsToPoint(UE::AnimationCore::FBoneWeights& OutWeights, 
+												  const FVector3d& InPoint, 
+												  const TMap<FName, uint16>* TargetBoneToIndex,
+												  const FVector3f& InNormal)
+{	
+	// Find the containing triangle and the barycentric coordinates of the closest point
+	int32 TriID; 
+	FVector3d Bary;
+	if (!FindClosestPointOnSourceSurface(InPoint, TargetToWorld, TriID, Bary))
+	{
+		return false;
+	}
+	FVector3f BaryF = FVector3f((float)Bary[0], (float)Bary[1], (float)Bary[2]);
+	const FIndex3i TriVertex = SourceMesh->GetTriangle(TriID);
 
+	const FDynamicMeshVertexSkinWeightsAttribute* SourceSkinWeights = SourceMesh->Attributes()->GetSkinWeightsAttribute(SourceProfileName);
 	const TArray<FName>* SourceBoneNames = nullptr;
 	if (!bIgnoreBoneAttributes) 
 	{
 		SourceBoneNames = &SourceMesh->Attributes()->GetBoneNames()->GetAttribValues();
 	}
 
+	if (SearchRadius < 0 && NormalThreshold < 0)
+	{
+		// If the radius and normals are ignored, simply interpolate the weights and return the result
+		TransferBoneWeightsLocals::InterpolateBoneWeights(OutWeights, TriVertex, BaryF, SourceSkinWeights, SourceBoneNames, TargetBoneToIndex);
+	}
+	else
+	{
+		bool bPassedRadiusCheck = true;
+		if (SearchRadius >= 0)
+		{
+			const FVector3d MatchedPoint = SourceMesh->GetTriBaryPoint(TriID, Bary[0], Bary[1], Bary[2]);
+			bPassedRadiusCheck = (InPoint - MatchedPoint).Length() <= SearchRadius;
+		}
+
+		bool bPassedNormalsCheck = true;
+		if (NormalThreshold >= 0)
+		{
+			FVector3f Normal0, Normal1, Normal2;
+			SourceMesh->Attributes()->PrimaryNormals()->GetTriElements(TriID, Normal0, Normal1, Normal2);
+
+			const FVector3f MatchedNormal = Normalized(BaryF[0]*Normal0 + BaryF[1]*Normal1 + BaryF[2]*Normal2);
+			const FVector3f InNormalNormalized = Normalized(InNormal);
+			const float NormalAngle = FMathf::ACos(InNormalNormalized.Dot(MatchedNormal));
+			bPassedNormalsCheck = (double)NormalAngle <= NormalThreshold;
+		}
+		
+		if (bPassedRadiusCheck && bPassedNormalsCheck)
+		{
+			TransferBoneWeightsLocals::InterpolateBoneWeights(OutWeights, TriVertex, BaryF, SourceSkinWeights, SourceBoneNames, TargetBoneToIndex);
+		}
+		else
+		{
+			return false;
+		}
+	}
+
+	return true;
+}
+
+bool FTransferBoneWeights::FindClosestPointOnSourceSurface(const FVector3d& InPoint, const FTransformSRT3d& InToWorld, int32& NearTriID, FVector3d& Bary)
+{
 	IMeshSpatial::FQueryOptions Options;
 	double NearestDistSqr;
-	int32 NearTriID;
 	
 	const FVector3d WorldPoint = InToWorld.TransformPosition(InPoint);
 	if (SourceBVH != nullptr) 
@@ -277,17 +517,9 @@ bool FTransferBoneWeights::Compute(const FVector3d& InPoint, const FTransformSRT
 	const FVector3d NearestPnt = Query.ClosestTrianglePoint;
 	const FIndex3i TriVertex = SourceMesh->GetTriangle(NearTriID);
 
-	const FVector3d Bary = VectorUtil::BarycentricCoords(NearestPnt,
-														 SourceMesh->GetVertexRef(TriVertex.A),
-														 SourceMesh->GetVertexRef(TriVertex.B),
-														 SourceMesh->GetVertexRef(TriVertex.C));
-
-	TransferBoneWeightsLocals::InterpolateBoneWeights(OutWeights,
-													  TriVertex,
-													  FVector3f((float)Bary[0], (float)Bary[1], (float)Bary[2]),
-													  SourceSkinWeights,
-													  SourceBoneNames,
-													  TargetBoneToIndex);
+	Bary = VectorUtil::BarycentricCoords(NearestPnt, SourceMesh->GetVertexRef(TriVertex.A),
+													 SourceMesh->GetVertexRef(TriVertex.B),
+													 SourceMesh->GetVertexRef(TriVertex.C));
 
 	return true;
-}
\ No newline at end of file
+}
diff --git a/Engine/Plugins/Runtime/GeometryProcessing/Source/DynamicMesh/Public/Operations/TransferBoneWeights.h b/Engine/Plugins/Runtime/GeometryProcessing/Source/DynamicMesh/Public/Operations/TransferBoneWeights.h
--- a/Engine/Plugins/Runtime/GeometryProcessing/Source/DynamicMesh/Public/Operations/TransferBoneWeights.h
+++ b/Engine/Plugins/Runtime/GeometryProcessing/Source/DynamicMesh/Public/Operations/TransferBoneWeights.h
@@ -4,7 +4,8 @@
 
 #include "GeometryTypes.h"
 #include "UObject/NameTypes.h"
-
+#include "Containers/BitArray.h"
+#include "TransformTypes.h"
 
 // Forward declarations
 class FProgressCancel;
@@ -13,9 +14,6 @@ namespace UE::AnimationCore { class FBoneWeights; }
 
 namespace UE::Geometry
 {
-	template<typename RealType> class TTransformSRT3;
-	typedef TTransformSRT3<double> FTransformSRT3d;
-
 	class FDynamicMesh3;
 	template<typename MeshType> class TMeshAABBTree3;
 	typedef TMeshAABBTree3<FDynamicMesh3> FDynamicMeshAABBTree3;
@@ -27,25 +25,25 @@ namespace UE
 namespace Geometry
 {
 /**
- * Transfer bone weights from one mesh (source) to another (target). Uses the dynamic mesh bone attributes to reindex 
- * the bone indices of the transferred weights from the source to the target skeletons. If both meshes have identical 
+ * Transfer bone weights from one mesh (source) to another (target). Uses the dynamic mesh bone attributes to reindex
+ * the bone indices of the transferred weights from the source to the target skeletons. If both meshes have identical
  * bone name attributes, then reindexing is skipped.
- * 
+ *
  * During the reindexing, if a weighted bone in the source skeleton is not present in the target skeleton, then the
  * weight is not transferred (skipped), and an error is printed to the console. For best results, the target skeleton
  * should be a superset of all the bones that are indexed by the transferred weights.
- * 
- * 
+ *
+ *
  * Example usage:
- * 
+ *
  * FDynamicMesh SourceMesh = ...; // Mesh we transferring weights from. Must have bone attributes.
  * FDynamicMesh TargetMesh = ...; // Mesh we are transferring weights to.
  *
  * FTransferBoneWeights TransferBoneWeights(&SourceMesh, FSkeletalMeshAttributes::DefaultSkinWeightProfileName);
- * 
+ *
  * // Optionally, transform the target mesh. This is useful when you want to align the two meshes in world space.
- * FTransformSRT3d InToWorld = ...; 
- * 
+ * TransferBoneWeights.TargetToWorld = ...;
+ *
  * // When transferring weights from a dynamic mesh with bone attributes to a dynamic mesh without bone attributes,
  * // first copy over the bone attributes from the source to the target.
  * if (!TargetMesh.HasAttributes() || !TargetMesh.Attributes()->HasBones())
@@ -53,22 +51,54 @@ namespace Geometry
  *     TargetMesh.EnableAttributes();
  *     TargetMesh.Attributes()->CopyBoneAttributes(*SourceMesh.Attributes());
  * }
- * 
- * if (TransferBoneWeights.Validate() == EOperationValidationResult::Ok) 
+ *
+ * // Set the transfer method.
+ * TransferBoneWeights.TransferMethod = ETransferBoneWeightsMethod::...;
+ *
+ * // if ETransferBoneWeightsMethod::ClosestPointOnSurface is used and you simply want to copy weights over from the
+ * // closest points then set the radius and normal threshold to -1 (default).
+ * TransferBoneWeights.SearchRadius = -1
+ * TransferBoneWeights.NormalThreshold = -1
+ *
+ * // if ETransferBoneWeightsMethod::InpaintWeights is used then additionally set the radius and normal parameters
+ * TransferBoneWeights.SearchRadius = ...    // Good estimate is to use a small value (0.05) of the bounding box radius
+ * TransferBoneWeights.NormalThreshold = ... // 30 degrees (0.52 rad) works well in practice
+ *
+ * if (TransferBoneWeights.Validate() == EOperationValidationResult::Ok)
+ * {
+ *      TransferBoneWeights.TransferWeightsToMesh(TargetMesh, FSkeletalMeshAttributes::DefaultSkinWeightProfileName);
+ * }
+ *
+ * // Alternatively if you don't want to use FDynamicMesh3 to represent your target mesh you can transfer weights to
+ * // to each point separately by calling
+ * if (TransferBoneWeights.Validate() == EOperationValidationResult::Ok)
  * {
- *      TransferBoneWeights.Compute(TargetMesh, InToWorld, FSkeletalMeshAttributes::DefaultSkinWeightProfileName);
+ *      for each Point:
+ *          FBoneWeights Weights;
+ *          TransferBoneWeights.TransferWeightsToPoint(Weights, Point);
  * }
+ *
+ * // After the transfer you can check which target mesh vertices had the weight transferred directly from the source mesh
+ * // via the FTransferBoneWeights:MatchedVertices array
+ *
  */
 
+
+
 class DYNAMICMESH_API FTransferBoneWeights
 {
 public:
 
 	enum class ETransferBoneWeightsMethod : uint8
 	{
-        // For every vertex on the target mesh, find the closest point on the surface of the source mesh.
-        // This is usually a point on a triangle where the bone weights are interpolated via barycentric coordinates.
-		ClosestPointOnSurface = 0
+        // For every vertex on the target mesh, find the closest point on the surface of the source mesh. If that point 
+        // is within the SearchRadius, and their normals differ by less than the NormalThreshold, then we directly copy  
+        // the weights from the source point to the target mesh vertex.
+		ClosestPointOnSurface = 0,
+
+        // Same as the ClosestPointOnSurface but for all the vertices we didn't copy the weights directly, automatically 
+		// compute the smooth weights.
+		InpaintWeights = 1
 	};
 
 	//
@@ -84,12 +114,31 @@ public:
 	/** The transfer method to compute the bone weights. */
 	ETransferBoneWeightsMethod TransferMethod = ETransferBoneWeightsMethod::ClosestPointOnSurface;
 
+	/** Transform applied to the input target mesh or target point before transfer. */
+	FTransformSRT3d TargetToWorld = FTransformSRT3d::Identity(); 
+
+	/**  Radius for searching the closest point. If negative, all points are considered. */
+	double SearchRadius = -1;
+
+	/** 
+	 * Maximum angle (in radians) difference between target and source point normals to be considred a match. 
+	 * If negative, normals are ignored.
+	 */
+	double NormalThreshold = -1;
+
 	/** 
      * Completely ignore the source and target mesh bone attributes when transferring weights from one dynamic mesh to another.
      * This skips re-indexing and simply copies skin weights over. Use with caution.
 	 */
 	bool bIgnoreBoneAttributes = false;
 
+	//
+	// Outputs
+	//
+
+	/** MatchedVertices[VertexID] is set to true for a tartet mesh vertex ID with a match found, false otherwise. */
+	TArray<bool> MatchedVertices;
+
 protected:
 		
 	/** Source mesh we are transfering weights from. */
@@ -132,35 +181,46 @@ public:
      * Transfer the bone weights from the source mesh to the given target mesh and store the result in the skin weight  
      * attribute with the given profile name.
 	 * 
-	 * @param InOutTargetMesh Target mesh we are transfering weights into
-	 * @param InToWorld Transform applied to the input target mesh
+	 * @param InOutTargetMesh	  Target mesh we are transfering weights into
      * @param InTargetProfileName Skin weight profile name we are writing into. If the profile with that name exists,  
-     *       					  then the data will be overwritten, otherwise a new attribute will be created.
+     *							  then the data will be overwritten, otherwise a new attribute will be created.
      * 
      * @return true if the algorithm succeeds, false if it failed or was canceled by the user.
 	 * 
 	 * @note Assumes that the InOutTargetMesh has bone attributes, use bIgnoreBoneAttributes flag to ignore the bone 
 	 * 		 attributes and skip re-indexing.
 	 */
-	virtual bool Compute(FDynamicMesh3& InOutTargetMesh, const FTransformSRT3d& InToWorld, const FName& InTargetProfileName);
+	virtual bool TransferWeightsToMesh(FDynamicMesh3& InOutTargetMesh, const FName& InTargetProfileName);
+
 
 	/**
-     * Compute a single bone weight for a given point.
-     *
-	 * @param InPoint Point for which we are computing a bone weight
-	 * @param InToWorld Transform applied to the point
-     * @param OutWeights Bone weight computed for the input transformed point
+	 * Compute the bone weights for a given point using the ETransferBoneWeightsMethod::ClosestPointOnSurface algorithm.
+	 *
+	 * @param OutWeights 		Bone weight computed for the input transformed point
+	 * @param InPoint 		    Point for which we are computing the bone weight
 	 * @param TargetBoneToIndex Optional map from the bone names to the bone indices of the target skeleton. 
 	 * 							If null, the bone indices of the skinning weights will not be re-indexed after the transfer.
+	 * @param InNormal			Normal at the input point. Should be set if NormalThreshold >= 0.
 	 * 
-     * @return true if the algorithm succeeds, false if it failed or was canceled by the user.
+	 * @return true if the algorithm succeeds, false if it failed to find the matching point or was canceled by the user.
 	 */
-	virtual bool Compute(const FVector3d& InPoint, const FTransformSRT3d& InToWorld, UE::AnimationCore::FBoneWeights& OutWeights, const TMap<FName, uint16>* TargetBoneToIndex = nullptr);
-
+	virtual bool TransferWeightsToPoint(UE::AnimationCore::FBoneWeights& OutWeights, 
+										const FVector3d& InPoint, 
+										const TMap<FName, uint16>* TargetBoneToIndex = nullptr,
+										const FVector3f& InNormal = FVector3f::Zero());
 protected:
 	
     /** @return if true, abort the computation. */
 	virtual bool Cancelled();
+	
+	/** 
+	 * Find the closest point on the surface of the source mesh and return the ID of the triangle containing it and its 
+	 * barycentric coordinates.
+	 * 
+	 * @return true if point is found, false otherwise
+	 */
+	bool FindClosestPointOnSourceSurface(const FVector3d& InPoint, const FTransformSRT3d& InToWorld, int32& OutTriID, FVector3d& OutBary);
+
 };
 
 } // end namespace UE::Geometry
