MeshModelingTools: Make FMeshFaceDualGraph support non-compact meshes.

diff --git a/Engine/Plugins/Runtime/GeometryProcessing/Source/DynamicMesh/Private/Polygroups/PolygroupsGenerator.cpp b/Engine/Plugins/Runtime/GeometryProcessing/Source/DynamicMesh/Private/Polygroups/PolygroupsGenerator.cpp
--- a/Engine/Plugins/Runtime/GeometryProcessing/Source/DynamicMesh/Private/Polygroups/PolygroupsGenerator.cpp
+++ b/Engine/Plugins/Runtime/GeometryProcessing/Source/DynamicMesh/Private/Polygroups/PolygroupsGenerator.cpp
@@ -679,6 +679,13 @@ public:
 		return vid;
 	}
 
+	void InsertVertex(int32 Vid, const FVector3d& Centroid, const FVector3d& Normal, double Area)
+	{
+		FDynamicGraph3d::InsertVertex(Vid, Centroid);
+		Normals.InsertAt({ {Normal.X, Normal.Y, Normal.Z} }, Vid);
+		Areas.InsertAt(Area, Vid);
+	}
+
 	FVector3d GetNormal(int32 vid) const
 	{
 		return Normals.AsVector3(vid);
@@ -690,17 +697,22 @@ public:
 		FMeshFaceDualGraph& FaceGraph,
 		TFunctionRef<bool(int,int)> TrisConnectedPredicate )
 	{
-		// if not true, code below needs updating
-		check(Mesh->IsCompactT());
-
+		bool bIsCompact = Mesh->IsCompactT();
 		for (int32 tid : Mesh->TriangleIndicesItr())
 		{
 			FVector3d Normal, Centroid;
 			double Area;
 			Mesh->GetTriInfo(tid, Normal, Area, Centroid);
 
-			int32 newid = FaceGraph.AppendVertex(Centroid, Normal, Area);
-			check(newid == tid);
+			if (bIsCompact)
+			{
+				int32 newid = FaceGraph.AppendVertex(Centroid, Normal, Area);
+				checkSlow(newid == tid);
+			}
+			else
+			{
+				FaceGraph.InsertVertex(tid, Centroid, Normal, Area);
+			}
 		}
 
 		for (int32 tid : Mesh->TriangleIndicesItr())
diff --git a/Engine/Source/Runtime/GeometryCore/Public/Curve/DynamicGraph.h b/Engine/Source/Runtime/GeometryCore/Public/Curve/DynamicGraph.h
--- a/Engine/Source/Runtime/GeometryCore/Public/Curve/DynamicGraph.h
+++ b/Engine/Source/Runtime/GeometryCore/Public/Curve/DynamicGraph.h
@@ -189,6 +189,17 @@ protected:
 		return VID;
 	}
 
+	bool insert_vertex_internal(int32 Vid)
+	{
+		if (!vertices_refcount.AllocateAt(Vid))
+		{
+			return false;
+		}
+		allocate_edges_list(Vid);
+		updateTimeStamp(true);
+		return true;
+	}
+
 private:
 	void allocate_edges_list(int VID)
 	{
diff --git a/Engine/Source/Runtime/GeometryCore/Public/Curve/DynamicGraph3.h b/Engine/Source/Runtime/GeometryCore/Public/Curve/DynamicGraph3.h
--- a/Engine/Source/Runtime/GeometryCore/Public/Curve/DynamicGraph3.h
+++ b/Engine/Source/Runtime/GeometryCore/Public/Curve/DynamicGraph3.h
@@ -83,6 +83,25 @@ public:
 		return vid;
 	}
 
+	/**
+	 * @return false If insertion failed (because a vertex with that ID already exists).
+	 */
+	bool InsertVertex(int32 Vid, TVector<T> V)
+	{
+		// For now we desided not to add the "bUnsafe" optimization machinery like the similar method in
+		// FDynamicMesh3. It is not needed if InsertVertex is used just for creating a graph with non compact 
+		// Vids, where we would presumably still iterate through the source indices in increasing order, and 
+		// so never have to worry about the performance cost of removing from the free list.
+
+		if (!insert_vertex_internal(Vid))
+		{
+			return false;
+		}
+		Vertices.InsertAt({ {V.X, V.Y, V.Z} }, Vid);
+
+		return true;
+	}
+
 	FRefCountVector::IndexEnumerable VertexIndicesItr() const
 	{
 		return vertices_refcount.Indices();
