Added a new overload for the TransferWeightsToPoint method that can take scalar bone data arrays.

#rnx

diff --git a/Engine/Plugins/Runtime/GeometryProcessing/Source/DynamicMesh/Private/Operations/TransferBoneWeights.cpp b/Engine/Plugins/Runtime/GeometryProcessing/Source/DynamicMesh/Private/Operations/TransferBoneWeights.cpp
--- a/Engine/Plugins/Runtime/GeometryProcessing/Source/DynamicMesh/Private/Operations/TransferBoneWeights.cpp
+++ b/Engine/Plugins/Runtime/GeometryProcessing/Source/DynamicMesh/Private/Operations/TransferBoneWeights.cpp
@@ -493,6 +493,33 @@ bool FTransferBoneWeights::TransferWeightsToPoint(UE::AnimationCore::FBoneWeight
 	return true;
 }
 
+template<typename BoneIndexType, typename BoneFloatWeightType, typename PosVectorType, typename NormalVectorType>
+bool FTransferBoneWeights::TransferWeightsToPoint(TArray<BoneIndexType>& OutBones, 
+												  TArray<BoneFloatWeightType>& OutWeights,
+												  const UE::Math::TVector<PosVectorType>& InPoint,
+												  const TMap<FName, uint16>* TargetBoneToIndex,
+												  const UE::Math::TVector<NormalVectorType>& InNormal)
+{
+	FBoneWeights BoneWeights;
+	if (!this->TransferWeightsToPoint(BoneWeights, FVector3d(InPoint), TargetBoneToIndex, FVector3f(InNormal)))
+	{
+		return false;
+	}
+
+	const int32 NumEntries = BoneWeights.Num();
+
+	OutBones.SetNum(NumEntries);
+	OutWeights.SetNum(NumEntries);
+
+	for (int32 BoneIdx = 0; BoneIdx < NumEntries; ++BoneIdx)
+	{
+		OutBones[BoneIdx] = static_cast<BoneIndexType>(BoneWeights[BoneIdx].GetBoneIndex());
+		OutWeights[BoneIdx] = static_cast<BoneFloatWeightType>(BoneWeights[BoneIdx].GetWeight());
+	}
+
+	return true;
+}
+
 bool FTransferBoneWeights::FindClosestPointOnSourceSurface(const FVector3d& InPoint, const FTransformSRT3d& InToWorld, int32& NearTriID, FVector3d& Bary)
 {
 	IMeshSpatial::FQueryOptions Options;
@@ -523,3 +550,6 @@ bool FTransferBoneWeights::FindClosestPointOnSourceSurface(const FVector3d& InPo
 
 	return true;
 }
+
+// template instantiation
+template DYNAMICMESH_API bool UE::Geometry::FTransferBoneWeights::TransferWeightsToPoint<int,float,float,float>(class TArray<int,class TSizedDefaultAllocator<32> > &,class TArray<float,class TSizedDefaultAllocator<32> > &,struct UE::Math::TVector<float> const &,class TMap<class FName,unsigned short,class FDefaultSetAllocator,struct TDefaultMapHashableKeyFuncs<class FName,unsigned short,0> > const *,struct UE::Math::TVector<float> const &);
\ No newline at end of file
diff --git a/Engine/Plugins/Runtime/GeometryProcessing/Source/DynamicMesh/Public/Operations/TransferBoneWeights.h b/Engine/Plugins/Runtime/GeometryProcessing/Source/DynamicMesh/Public/Operations/TransferBoneWeights.h
--- a/Engine/Plugins/Runtime/GeometryProcessing/Source/DynamicMesh/Public/Operations/TransferBoneWeights.h
+++ b/Engine/Plugins/Runtime/GeometryProcessing/Source/DynamicMesh/Public/Operations/TransferBoneWeights.h
@@ -4,7 +4,6 @@
 
 #include "GeometryTypes.h"
 #include "UObject/NameTypes.h"
-#include "Containers/BitArray.h"
 #include "TransformTypes.h"
 
 // Forward declarations
@@ -196,8 +195,8 @@ public:
 	/**
 	 * Compute the bone weights for a given point using the ETransferBoneWeightsMethod::ClosestPointOnSurface algorithm.
 	 *
-	 * @param OutWeights 		Bone weight computed for the input transformed point
-	 * @param InPoint 		    Point for which we are computing the bone weight
+	 * @param OutWeights 		Bone weight computed for the input transformed point.
+	 * @param InPoint			Point for which we are computing the bone weight.
 	 * @param TargetBoneToIndex Optional map from the bone names to the bone indices of the target skeleton. 
 	 * 							If null, the bone indices of the skinning weights will not be re-indexed after the transfer.
 	 * @param InNormal			Normal at the input point. Should be set if NormalThreshold >= 0.
@@ -208,6 +207,28 @@ public:
 										const FVector3d& InPoint, 
 										const TMap<FName, uint16>* TargetBoneToIndex = nullptr,
 										const FVector3f& InNormal = FVector3f::Zero());
+	
+	/**
+	 * Compute the bone weights for a given point using the ETransferBoneWeightsMethod::ClosestPointOnSurface algorithm.
+	 *
+	 * @param OutBones			Array of bone indices. Array size is equal to the number of bone influences.
+	 * @param OutWeights		Array of bone weights. Array size is equal to the number of bone influences.
+	 * @param InPoint			Point for which we are computing the bone weight.
+	 * @param TargetBoneToIndex Optional map from the bone names to the bone indices of the target skeleton. 
+	 * 							If null, the bone indices of the skinning weights will not be re-indexed after the transfer.
+	 * @param InNormal			Normal at the input point. Should be set if NormalThreshold >= 0.
+	 * 
+	 * @return true if the algorithm succeeds, false if it failed to find the matching point or was canceled by the user.
+	 * 
+	 * @note Add specialization in the source file for template types you want to use. 
+	 * 		 Currently supporting int for BoneIndexType and float for BoneFloatWeightType, PosVectorType, NormalVectorType
+	 */
+	template<typename BoneIndexType, typename BoneFloatWeightType, typename PosVectorType,  typename NormalVectorType = float>
+	bool TransferWeightsToPoint(TArray<BoneIndexType>& OutBones,
+								TArray<BoneFloatWeightType>& OutWeights,
+								const UE::Math::TVector<PosVectorType>& InPoint,
+								const TMap<FName, uint16>* TargetBoneToIndex = nullptr,
+								const UE::Math::TVector<NormalVectorType>& InNormal = UE::Math::TVector<NormalVectorType>::Zero());
 protected:
 	
     /** @return if true, abort the computation. */
