MeshModelingTools: Fix expand/contract selection not working correctly in some circumstances for edge/face selections due to considering border vertices solely based on neighboring vertices being in selection, rather than neighboring elements of the selection type.

diff --git a/Engine/Plugins/Runtime/GeometryProcessing/Source/DynamicMesh/Private/Selections/GeometrySelectionUtil.cpp b/Engine/Plugins/Runtime/GeometryProcessing/Source/DynamicMesh/Private/Selections/GeometrySelectionUtil.cpp
--- a/Engine/Plugins/Runtime/GeometryProcessing/Source/DynamicMesh/Private/Selections/GeometrySelectionUtil.cpp
+++ b/Engine/Plugins/Runtime/GeometryProcessing/Source/DynamicMesh/Private/Selections/GeometrySelectionUtil.cpp
@@ -15,7 +15,7 @@
 #include "Selections/MeshFaceSelection.h"
 
 
-namespace
+namespace GeometrySelectionUtilLocals
 {
 
 // Return an integer in the range [0,5] which can be used to look up a handling function based on the Selection type
@@ -27,8 +27,178 @@ int GetSelectionTypeAsIndex(const UE::Geometry::FGeometrySelection& Selection)
 	return Index;
 }
 
+// We don't currently have an overload of EnumerateSelectionTriangles that takes a FGroupTopology
+// instead of FPolygroupSet. If we build out the below version to support the other selection
+// types (vertex, edge), we might want to expose it.
+/**
+ * Given a selection with elements of type EGeometryElementType::Face, call TriangleFunc on
+ * each triangle of the selection.
+ * 
+ * @param GroupTopology must not be null if Selection is EGeometryTopologyType::Polygroup
+ */
+bool EnumerateFaceElementSelectionTriangles(
+	const UE::Geometry::FGeometrySelection& Selection,
+	const UE::Geometry::FDynamicMesh3& Mesh,
+	const UE::Geometry::FGroupTopology* GroupTopology,
+	TFunctionRef<void(int32)> TriangleFunc)
+{
+	using namespace UE::Geometry;
+
+	if (!ensure(Selection.ElementType == EGeometryElementType::Face))
+	{
+		return false;
+	}
+
+	if (Selection.TopologyType == EGeometryTopologyType::Polygroup)
+	{
+		if (!ensure(GroupTopology))
+		{
+			return false;
+		}
+
+		for (uint64 EncodedID : Selection.Selection)
+		{
+			FGeoSelectionID GroupTriID(EncodedID);
+			int32 SeedTriangleID = (int32)GroupTriID.GeometryID;
+			int32 GroupID = (int32)GroupTriID.TopologyID;
+			if (Mesh.IsTriangle(SeedTriangleID))
+			{
+				for (int32 Tid : GroupTopology->GetGroupTriangles(GroupID))
+				{
+					TriangleFunc(Tid);
+				}
+			}
+		}
+	}
+	else if (Selection.TopologyType == EGeometryTopologyType::Triangle)
+	{
+		for (uint64 TriangleID : Selection.Selection)
+		{
+			if (Mesh.IsTriangle((int32)TriangleID))
+			{
+				TriangleFunc((int32)TriangleID);
+			}
+		}
+	}
+	else
+	{
+		return ensure(false);
+	}
+
+	return true;
 }
 
+// We don't currently have an EnumerateSelectionEdges. If we build out the below to handle
+// other selection types (vertex, face), we might want to expose it.
+/**
+ * Given a selection with elements of type EGeometryElementType::Edge, call EdgeFunc on
+ * each mesh edge (with the Eid passed in) of the selection.
+ *
+ * @param GroupTopology must not be null if Selection is EGeometryTopologyType::Polygroup
+ */
+bool EnumerateEdgeElementSelectionEdges(
+	const UE::Geometry::FGeometrySelection& Selection,
+	const UE::Geometry::FDynamicMesh3& Mesh,
+	const UE::Geometry::FGroupTopology* GroupTopology,
+	TFunctionRef<void(uint32)> EdgeFunc)
+{
+	using namespace UE::Geometry;
+
+	if (!ensure(Selection.ElementType == EGeometryElementType::Edge))
+	{
+		return false;
+	}
+
+	if (Selection.TopologyType == EGeometryTopologyType::Polygroup)
+	{
+		if (!ensure(GroupTopology))
+		{
+			return false;
+		}
+
+		for (uint64 EncodedID : Selection.Selection)
+		{
+			FMeshTriEdgeID TriEdgeID(FGeoSelectionID(EncodedID).GeometryID);
+			int32 SeedEdgeID = Mesh.IsTriangle(TriEdgeID.TriangleID) ? Mesh.GetTriEdge(TriEdgeID.TriangleID, TriEdgeID.TriEdgeIndex) : IndexConstants::InvalidID;
+			if (Mesh.IsEdge(SeedEdgeID))
+			{
+				int32 GroupEdgeID = GroupTopology->FindGroupEdgeID(SeedEdgeID);
+				for (int32 Eid : GroupTopology->GetGroupEdgeEdges(GroupEdgeID))
+				{
+					EdgeFunc(Eid);
+				}
+			}
+		}
+	}
+	else if (Selection.TopologyType == EGeometryTopologyType::Triangle)
+	{
+		for (uint64 EncodedID : Selection.Selection)
+		{
+			FMeshTriEdgeID TriEdgeID(FGeoSelectionID(EncodedID).GeometryID);
+			int32 EdgeID = Mesh.IsTriangle(TriEdgeID.TriangleID) ? Mesh.GetTriEdge(TriEdgeID.TriangleID, TriEdgeID.TriEdgeIndex) : IndexConstants::InvalidID;
+			if (Mesh.IsEdge(EdgeID))
+			{
+				EdgeFunc(EdgeID);
+			}
+		}
+	}
+	else
+	{
+		return ensure(false);
+	}
+
+	return true;
+}
+
+bool EnumerateVertexElementSelectionVertices(
+	const UE::Geometry::FGeometrySelection& Selection,
+	const UE::Geometry::FDynamicMesh3& Mesh,
+	const UE::Geometry::FGroupTopology* GroupTopology,
+	TFunctionRef<void(uint32)> VertexFunc)
+{
+	using namespace UE::Geometry;
+
+	if (!ensure(Selection.ElementType == EGeometryElementType::Vertex))
+	{
+		return false;
+	}
+
+	if (Selection.TopologyType == EGeometryTopologyType::Polygroup)
+	{
+		if (!ensure(GroupTopology))
+		{
+			return false;
+		}
+
+		for (uint64 EncodedID : Selection.Selection)
+		{
+			int32 VertexID = (int32)FGeoSelectionID(EncodedID).GeometryID;
+			if (Mesh.IsVertex(VertexID))
+			{
+				VertexFunc(VertexID);
+			}
+		}
+	}
+	else if (Selection.TopologyType == EGeometryTopologyType::Triangle)
+	{
+		for (uint64 VertexID : Selection.Selection)
+		{
+			if (Mesh.IsVertex((int32)VertexID))
+			{
+				VertexFunc((int32)VertexID);
+			}
+		}
+	}
+	else
+	{
+		return ensure(false);
+	}
+
+	return true;
+}
+
+}//end GeometrySelectionUtilLocals
+
 bool UE::Geometry::AreSelectionsIdentical(
 	const FGeometrySelection& SelectionA, const FGeometrySelection& SelectionB)
 {
@@ -994,7 +1164,7 @@ bool UE::Geometry::ConvertSelection(
 	const FGeometrySelection& FromSelectionIn,
 	FGeometrySelection& ToSelectionOut)
 {
-
+	using namespace GeometrySelectionUtilLocals;
 
 	const auto NotImplemented = [](
 		const FDynamicMesh3& Mesh,
@@ -1453,42 +1623,306 @@ bool UE::Geometry::MakeSelectAllConnectedSelection(
 
 
 
+bool UE::Geometry::GetSelectionBoundaryVertices(
+	const FDynamicMesh3& Mesh,
+	const FGroupTopology* GroupTopology,
+	const FGeometrySelection& ReferenceSelection,
+	TSet<int32>& BorderVidsOut, TSet<int32>& CurVerticesOut)
+{
+	using namespace GeometrySelectionUtilLocals;
+
+	BorderVidsOut.Reset();
+	CurVerticesOut.Reset();
+	
+	switch (ReferenceSelection.ElementType)
+	{
+	case EGeometryElementType::Vertex:
+		EnumerateVertexElementSelectionVertices(ReferenceSelection, Mesh, GroupTopology, [&CurVerticesOut](uint32 Vid)
+		{
+			CurVerticesOut.Add(Vid);
+		});
+
+		// Border vertices are ones that have some adjacent vertices not in selection
+		for (int32 VertexID : CurVerticesOut)
+		{
+			// a boundary vertex is always on the selection boundary (for this and other selection types)
+			bool bIsBoundary = Mesh.IsBoundaryVertex(VertexID);	
+			if (!bIsBoundary)
+			{
+				Mesh.EnumerateVertexVertices(VertexID, [&](int32 NbrVertexID)
+				{
+					if (!CurVerticesOut.Contains(NbrVertexID))
+					{
+						bIsBoundary = true;
+					}
+				});
+			}
+
+			if (bIsBoundary)
+			{
+				BorderVidsOut.Add(VertexID);
+			}
+		}
+		break;
+	case EGeometryElementType::Edge:
+	{
+		// Border vertices are ones that have some adjacent edges that are not in selection, so
+		// determine edges in selection.
+		TSet<int32> EdgeIDsInSelection;
+		EnumerateEdgeElementSelectionEdges(ReferenceSelection, Mesh, GroupTopology, [&EdgeIDsInSelection, &CurVerticesOut, &Mesh](uint32 Eid)
+		{
+			EdgeIDsInSelection.Add(Eid);
+			FIndex2i EdgeV = Mesh.GetEdgeV(Eid);
+			CurVerticesOut.Add(EdgeV.A);
+			CurVerticesOut.Add(EdgeV.B);
+		});
 
+		for (int32 VertexID : CurVerticesOut)
+		{
+			bool bIsBoundary = Mesh.IsBoundaryVertex(VertexID);
+			if (!bIsBoundary)
+			{
+				Mesh.EnumerateVertexEdges(VertexID, [&EdgeIDsInSelection, &bIsBoundary](int32 EdgeID)
+				{
+					if (!EdgeIDsInSelection.Contains(EdgeID))
+					{
+						bIsBoundary = true;
+					}
+				});
+			}
 
+			if (bIsBoundary)
+			{
+				BorderVidsOut.Add(VertexID);
+			}
+		}
+	}
+		break;
+	case EGeometryElementType::Face:
+	{
+		// Border vertices are ones that have some adjacent triangles that are not in selection.
+		TSet<int32> TriangleIDsInSelection;
+		EnumerateFaceElementSelectionTriangles(ReferenceSelection, Mesh, GroupTopology, [&TriangleIDsInSelection, &CurVerticesOut, &Mesh](uint32 Tid)
+		{
+			TriangleIDsInSelection.Add(Tid);
+			FIndex3i Triangle = Mesh.GetTriangle(Tid);
+			CurVerticesOut.Add(Triangle.A);
+			CurVerticesOut.Add(Triangle.B);
+			CurVerticesOut.Add(Triangle.C);
+		});
 
-bool UE::Geometry::MakeBoundaryConnectedSelection(
-	const UE::Geometry::FDynamicMesh3& Mesh,
+		for (int32 VertexID : CurVerticesOut)
+		{
+			bool bIsBoundary = Mesh.IsBoundaryVertex(VertexID);
+			if (!bIsBoundary)
+			{
+				Mesh.EnumerateVertexTriangles(VertexID, [&TriangleIDsInSelection, &bIsBoundary](int32 TriangleID)
+				{
+					if (!TriangleIDsInSelection.Contains(TriangleID))
+					{
+						bIsBoundary = true;
+					}
+				});
+			}
+
+			if (bIsBoundary)
+			{
+				BorderVidsOut.Add(VertexID);
+			}
+		}
+	}
+		break;
+	default:
+		return ensure(false);
+	}
+
+	return true;
+}
+
+
+
+bool UE::Geometry::GetSelectionBoundaryCorners(
+	const FDynamicMesh3& Mesh,
 	const FGroupTopology* GroupTopology,
 	const FGeometrySelection& ReferenceSelection,
-	TFunctionRef<bool(FGeoSelectionID)> SelectionIDPredicate,
-	FGeometrySelection& BoundaryConnectedSelection)
+	TSet<int32>& BorderCornerIDsOut, TSet<int32>& CurCornerIDsOut)
 {
-	if (BoundaryConnectedSelection.TopologyType == EGeometryTopologyType::Triangle)
+	BorderCornerIDsOut.Reset();
+	CurCornerIDsOut.Reset();
+
+	if (!ensure(GroupTopology))
 	{
-		// convert to vertex selection
-		TSet<int32> CurVertices;
-		EnumerateTriangleSelectionVertices(ReferenceSelection, Mesh, FTransform::Identity, [&](uint64 ID, const FVector3d&)
+		return false;
+	}
+
+	if (!ensure(ReferenceSelection.TopologyType == EGeometryTopologyType::Polygroup))
+	{
+		// We don't currently support triangle selections here in part because it's not clear what to do. The
+		// proper thing is likely to convert to an equivalent polygroup selection and find border corners, but
+		// we haven't yet defined some of those conversions. Alternatively we could find the border vertices and 
+		// keep whichever ones happen to be corners, but that gives the unintuitive result of not giving any corners
+		// for selections that don't happen to line up with group boundaries.
+		// There's also the fact that we don't yet have a use case for supporting this here.
+		return false;
+	}
+
+	TArray<int32> NbrArray;
+
+	switch (ReferenceSelection.ElementType)
+	{
+	case EGeometryElementType::Vertex:
+	{
+		// Assemble included corners
+		for (uint64 ID : ReferenceSelection.Selection)
 		{
-			CurVertices.Add( FGeoSelectionID(ID).GeometryID );
-		});
-		// find 'border' vertices that have some adjacent vertices not in selection
-		TSet<int32> BorderVertices;
-		for (int32 VertexID : CurVertices)
+			CurCornerIDsOut.Add(FGeoSelectionID(ID).TopologyID);		// TODO: can we rely on TopologyID being stable here, or do we need to look up from VertexID?
+		}
+
+		// Border corners are ones that have a corner neighbor not in the selection
+		for (int32 CornerID : CurCornerIDsOut)
 		{
-			bool bIsBoundary = Mesh.IsBoundaryVertex(VertexID);		// a boundary vertex is always on the selection boundary
-			Mesh.EnumerateVertexVertices(VertexID, [&](int32 NbrVertexID)
+			// Boundary vertex corners are always considered to be on selection boundary (for this and other selection types)
+			bool bIsBoundary = Mesh.IsBoundaryVertex(GroupTopology->GetCornerVertexID(CornerID));
+			if (!bIsBoundary)
 			{
-				if (CurVertices.Contains(NbrVertexID) == false)
+				NbrArray.Reset();
+				GroupTopology->FindCornerNbrCorners(CornerID, NbrArray);
+				for (int32 NbrCornerID : NbrArray)
+				{
+					if (!CurCornerIDsOut.Contains(NbrCornerID))
+					{
+						bIsBoundary = true;
+						break;
+					}
+				}
+			}
+
+			if (bIsBoundary)
+			{
+				BorderCornerIDsOut.Add(CornerID);
+			}
+		}
+	}
+		break;
+	case EGeometryElementType::Edge:
+	{
+		// Assemble the current group edge selection and the included corners.
+		TSet<int32> GroupEdgeIDsInSelection;
+		for (uint64 ID : ReferenceSelection.Selection)
+		{
+			int32 GroupEdgeID = FGeoSelectionID(ID).TopologyID;
+			GroupEdgeIDsInSelection.Add(GroupEdgeID);
+
+			const FGroupTopology::FGroupEdge& Edge = GroupTopology->Edges[GroupEdgeID];
+			if (Edge.EndpointCorners.A != IndexConstants::InvalidID)
+			{
+				CurCornerIDsOut.Add(Edge.EndpointCorners.A);
+			}
+			if (Edge.EndpointCorners.B != IndexConstants::InvalidID)
+			{
+				CurCornerIDsOut.Add(Edge.EndpointCorners.B);
+			}
+		}
+
+		// Border corners are ones that have some attached group edges that are not in the current selection.
+		for (int32 CornerID : CurCornerIDsOut)
+		{
+			bool bIsBoundary = Mesh.IsBoundaryVertex(GroupTopology->GetCornerVertexID(CornerID));
+			if (!bIsBoundary)
+			{
+				NbrArray.Reset();
+				GroupTopology->FindCornerNbrEdges(CornerID, NbrArray);
+				for (int32 GroupEdgeID : NbrArray)
 				{
-					bIsBoundary = true;
+					if (!GroupEdgeIDsInSelection.Contains(GroupEdgeID))
+					{
+						bIsBoundary = true;
+						break;
+					}
+				}
+			}
+
+			if (bIsBoundary)
+			{
+				BorderCornerIDsOut.Add(CornerID);
+			}
+		}
+	}
+		break;
+	case EGeometryElementType::Face:
+	{
+		// Assemble current group selection and the included corners
+		TSet<int32> GroupsInSelection;
+		for (uint64 ID : ReferenceSelection.Selection)
+		{
+			int32 GroupID = FGeoSelectionID(ID).TopologyID;
+			GroupsInSelection.Add(GroupID);
+			GroupTopology->ForGroupEdges(GroupID, [&](const FGroupTopology::FGroupEdge& Edge, int)
+			{
+				if (Edge.EndpointCorners.A != IndexConstants::InvalidID)
+				{
+					CurCornerIDsOut.Add(Edge.EndpointCorners.A);
+				}
+				if (Edge.EndpointCorners.B != IndexConstants::InvalidID)
+				{
+					CurCornerIDsOut.Add(Edge.EndpointCorners.B);
 				}
 			});
+		}
+
+		// Boundary corners are ones that have an attached group not in selection
+		for (int32 CornerID : CurCornerIDsOut)
+		{
+			bool bIsBoundary = Mesh.IsBoundaryVertex(GroupTopology->GetCornerVertexID(CornerID));
+			if (!bIsBoundary)
+			{
+				NbrArray.Reset();
+				GroupTopology->FindCornerNbrGroups(CornerID, NbrArray);
+				for (int32 Group : NbrArray)
+				{
+					if (!GroupsInSelection.Contains(Group))
+					{
+						bIsBoundary = true;
+						break;
+					}
+				}
+			}
+
 			if (bIsBoundary)
 			{
-				BorderVertices.Add(VertexID);
+				BorderCornerIDsOut.Add(CornerID);
 			}
 		}
+	}
+		break;
+	default:
+		return ensure(false);
+	}
+
+	return true;
+}
+
+
+
+bool UE::Geometry::MakeBoundaryConnectedSelection(
+	const UE::Geometry::FDynamicMesh3& Mesh,
+	const FGroupTopology* GroupTopology,
+	const FGeometrySelection& ReferenceSelection,
+	TFunctionRef<bool(FGeoSelectionID)> SelectionIDPredicate,
+	FGeometrySelection& BoundaryConnectedSelection)
+{
+	using namespace GeometrySelectionUtilLocals;
+
+	if (BoundaryConnectedSelection.TopologyType == EGeometryTopologyType::Triangle)
+	{
+		TSet<int32> BorderVertices;
+		TSet<int32> CurVertices;
+		if (!GetSelectionBoundaryVertices(Mesh, GroupTopology, ReferenceSelection, BorderVertices, CurVertices))
+		{
+			return false;
+		}
 
+		// Now select elements connected to the border vertices.
 		if (BoundaryConnectedSelection.ElementType == EGeometryElementType::Vertex)
 		{
 			TSet<int32> AdjacentVertices = BorderVertices;
@@ -1561,76 +1995,15 @@ bool UE::Geometry::MakeBoundaryConnectedSelection(
 			return false;
 		}
 
-		// convert to Corner selection
 		TSet<int32> CurCornerIDs;
-		if (ReferenceSelection.ElementType == EGeometryElementType::Vertex)
-		{
-			for (uint64 ID : ReferenceSelection.Selection)
-			{
-				CurCornerIDs.Add( FGeoSelectionID(ID).TopologyID );		// TODO: can we rely on TopologyID being stable here, or do we need to look up from VertexID?
-			}
-		}
-		else if (ReferenceSelection.ElementType == EGeometryElementType::Edge)
-		{
-			for (uint64 ID : ReferenceSelection.Selection)
-			{
-				int32 GroupEdgeID = FGeoSelectionID(ID).TopologyID;
-				const FGroupTopology::FGroupEdge& Edge = GroupTopology->Edges[GroupEdgeID];
-				if ( Edge.EndpointCorners.A != IndexConstants::InvalidID )
-				{
-					CurCornerIDs.Add(Edge.EndpointCorners.A);
-				}
-				if ( Edge.EndpointCorners.B != IndexConstants::InvalidID )
-				{
-					CurCornerIDs.Add(Edge.EndpointCorners.B);
-				}
-			}
-		}
-		else if (ReferenceSelection.ElementType == EGeometryElementType::Face)
-		{
-			for (uint64 ID : ReferenceSelection.Selection)
-			{
-				int32 GroupID = FGeoSelectionID(ID).TopologyID;
-				const FGroupTopology::FGroup* Group = GroupTopology->FindGroupByID(GroupID);
-				GroupTopology->ForGroupEdges(GroupID, [&](const FGroupTopology::FGroupEdge& Edge, int)
-				{
-					if ( Edge.EndpointCorners.A != IndexConstants::InvalidID )
-					{
-						CurCornerIDs.Add(Edge.EndpointCorners.A);
-					}
-					if ( Edge.EndpointCorners.B != IndexConstants::InvalidID )
-					{
-						CurCornerIDs.Add(Edge.EndpointCorners.B);
-					}
-				});
-			}
-		}
-		else
+		TSet<int32> BorderCorners;
+		
+		if (!GetSelectionBoundaryCorners(Mesh, GroupTopology, ReferenceSelection, BorderCorners, CurCornerIDs))
 		{
 			return false;
 		}
 
-
-		// now find 'border' corners that have some adjacent corners not in selection
-		TSet<int32> BorderCorners;
 		TArray<int32> NbrArray;
-		for (int32 CornerID : CurCornerIDs)
-		{
-			bool bIsBoundary = Mesh.IsBoundaryVertex(GroupTopology->GetCornerVertexID(CornerID));		// a boundary vertex is always on the selection boundary
-			NbrArray.Reset();
-			GroupTopology->FindCornerNbrCorners(CornerID, NbrArray);
-			for ( int32 NbrCornerID : NbrArray )
-			{
-				if (CurCornerIDs.Contains(NbrCornerID) == false)
-				{
-					bIsBoundary = true;
-				}
-			}
-			if (bIsBoundary)
-			{
-				BorderCorners.Add(CornerID);
-			}
-		}
 
 		// now that we have boundary corners, iterate over them and select connected elements
 		if (BoundaryConnectedSelection.ElementType == EGeometryElementType::Vertex)
diff --git a/Engine/Plugins/Runtime/GeometryProcessing/Source/DynamicMesh/Public/Selections/GeometrySelectionUtil.h b/Engine/Plugins/Runtime/GeometryProcessing/Source/DynamicMesh/Public/Selections/GeometrySelectionUtil.h
--- a/Engine/Plugins/Runtime/GeometryProcessing/Source/DynamicMesh/Public/Selections/GeometrySelectionUtil.h
+++ b/Engine/Plugins/Runtime/GeometryProcessing/Source/DynamicMesh/Public/Selections/GeometrySelectionUtil.h
@@ -301,6 +301,54 @@ DYNAMICMESH_API bool MakeBoundaryConnectedSelection(
 	TFunctionRef<bool(FGeoSelectionID)> SelectionIDPredicate,
 	FGeometrySelection& BoundaryConnectedSelection);
 
+/**
+ * Given a selection, return the vertex IDs of the vertices on the boundary of this selection. A selected vertex 
+ * is considered to be on the boundary either if it is on the actual mesh boundary (for an open mesh) or it is
+ * connected to a triangle element that is not part of the selection (i.e. if the vertex has a neighbor vertex 
+ * not in selection for a selection of type EGeometryElementType::Vertex, or an adjacent edge not in selection
+ * for a selection of type EGeometryElementType::Edge, an adjacent triangle not in the selection for a selection 
+ * of type EGeometryElementType::Face).
+ * 
+ * For selections of type EGeometryTopologyType::Polygroup, the results are equivalent to first converting the
+ * selection to corresponding EGeometryTopologyType::Triangle selection and then finding the boundary vertices.
+ * This gives the intuitive result for face selections, but may or may not be what is desired for polygroup 
+ * vertex/edge selections, because vertices/edges that seem to be on the interior of the polygroup selection 
+ * may be considered border vertices if the tesselation is such that they are adjacent to unselected triangle
+ * vertices/edges.
+ * 
+ * @param GroupTopology Must not be null for selections of type EGeometryTopologyType::Polygroup
+ * @param BorderVidsOut Output vertex IDs of border vertices
+ * @param CurVerticesOut Output vertex IDs of all vertices in the current selection
+ * @return true if successful. For instance, could fail if GroupTopology was null for a EGeometryTopologyType::Polygroup selection
+ */
+DYNAMICMESH_API bool GetSelectionBoundaryVertices(
+	const UE::Geometry::FDynamicMesh3& Mesh,
+	const UE::Geometry::FGroupTopology* GroupTopology,
+	const UE::Geometry::FGeometrySelection& ReferenceSelection,
+	TSet<int32>& BorderVidsOut, TSet<int32>& CurVerticesOut);
+
+/**
+ * Given a EGeometryTopologyType::Polygroup selection, return the corner IDs of the polygroup corners on the
+ * boundary of the selection. A selected corner is considered to be on the boundary either if it is on the actual 
+ * mesh boundary (for an open mesh) or it is connected to an element that is not part of the selection (i.e. if
+ * there is a neighbor corner not in the selection for a selection of type EGeometryElementType::Vertex, or an
+ * adjoining edge not in the selection for a selection of type EGeometryElementType::Edge, or an adjoining group not
+ * in the selection for a selection of type EGeometryElementType::Face).
+ * 
+ * Selection must be of type EGeometryTopologyType::Polygroup, and GroupTopology must not be null.
+ *
+ * @param GroupTopology Must not be null
+ * @param BorderCornerIDsOut Output corner IDs of border corners
+ * @param CurCornerIDsOut Output corner IDs of all corners included in the current selection
+ * @return true if successful. For instance, could fail if GroupTopology was null
+ */
+DYNAMICMESH_API bool GetSelectionBoundaryCorners(
+	const UE::Geometry::FDynamicMesh3& Mesh,
+	const UE::Geometry::FGroupTopology* GroupTopology,
+	const UE::Geometry::FGeometrySelection& ReferenceSelection,
+	TSet<int32>& BorderCornerIDsOut, TSet<int32>& CurCornerIDsOut);
+
+
 
 enum class EGeometrySelectionCombineModes : uint8
 {
