MeshModelingTools: Allow SimplifyToMinimalPlanar to work on the boundary of meshes.

diff --git a/Engine/Plugins/Runtime/GeometryProcessing/Source/DynamicMesh/Private/MeshSimplification.cpp b/Engine/Plugins/Runtime/GeometryProcessing/Source/DynamicMesh/Private/MeshSimplification.cpp
--- a/Engine/Plugins/Runtime/GeometryProcessing/Source/DynamicMesh/Private/MeshSimplification.cpp
+++ b/Engine/Plugins/Runtime/GeometryProcessing/Source/DynamicMesh/Private/MeshSimplification.cpp
@@ -813,25 +813,69 @@ static bool IsDevelopableVertex(const FDynamicMesh3& Mesh, int32 VertexID, doubl
 }
 
 
-
-
-template<typename GetTriNormalFuncType>
-static bool IsCollapsableDevelopableEdge(const FDynamicMesh3& Mesh, int32 CollapseEdgeID, int32 RemoveV, int32 KeepV, double DotTolerance,
-	GetTriNormalFuncType GetTriNormalFunc)
+bool IsCollapsableDevelopableEdge(const FDynamicMesh3& Mesh, 
+	int32 CollapseEdgeID, int32 RemoveV, int32 KeepV, double DotTolerance,
+	const TArray<FVector3d>& TriNormals, const TArray<bool> IsBoundaryVtxCache)
 {
 	FIndex2i CollapseEdgeT = Mesh.GetEdgeT(CollapseEdgeID);
-	FVector3d Normal1 = GetTriNormalFunc(CollapseEdgeT.A);
-	FVector3d Normal2 = GetTriNormalFunc(CollapseEdgeT.B);
+	FVector3d Normal1 = TriNormals[CollapseEdgeT.A];
+
+	if (CollapseEdgeT.B == IndexConstants::InvalidID)
+	{
+		// If we're collapsing a boundary edge, the only way to avoid changing the shape is for RemoveV 
+		// to be flat and have exactly one other attached boundary edge that is colinear with this one. 
+		// Start by finding the other boundary edge and making sure that there is only one.
+		bool bFoundSecondBoundaryEdge = false;
+		for (int32 Eid : Mesh.VtxEdgesItr(RemoveV))
+		{
+			if (Eid != CollapseEdgeID && Mesh.IsBoundaryEdge(Eid))
+			{
+				if (bFoundSecondBoundaryEdge)
+				{
+					// Found more than one other boundary edge, so not collapsable
+					return false;
+				}
+				bFoundSecondBoundaryEdge = true;
+
+				// Verify that this second boundary edge is colinear with ours.
+				FVector3d KeepVert = Mesh.GetVertex(KeepV);
+				FVector3d RemoveVert = Mesh.GetVertex(RemoveV);
+				int32 OtherV = IndexUtil::FindEdgeOtherVertex(Mesh.GetEdgeV(Eid), RemoveV);
+				FVector3d OtherVert = Mesh.GetVertex(OtherV);
+				if (!(Normalized(RemoveVert - OtherVert).Dot(Normalized(KeepVert - RemoveVert)) > DotTolerance))
+				{
+					// Not colinear
+					return false;
+				}
+			}
+		}
+		if (!bFoundSecondBoundaryEdge)
+		{
+			// Seems impossible for a vertex to have exactly one attached boundary edge
+			return ensure(false);
+		}
+
+		// If we got to here, we found the other boundary edge, and we'll check for planarity further below.
+	}
+	else
+	{
+		// If this is not a boundary edge, then remove V must not be a boundary vertex, else we
+		// would deform the boundary on collapse.
+		if (IsBoundaryVtxCache[RemoveV])
+		{
+			return false;
+		}
+	}
 
-	// assuming is that RemoveV is developable vertex...should check?
+	FVector3d Normal2 = (CollapseEdgeT.B == IndexConstants::InvalidID) ? FVector3d::ZeroVector : TriNormals[CollapseEdgeT.B];
 
 	// planar case
-	if (Normal1.Dot(Normal2) > DotTolerance)
+	if (CollapseEdgeT.B == IndexConstants::InvalidID || Normal1.Dot(Normal2) > DotTolerance)
 	{
 		bool bIsFlat = true;
 		Mesh.EnumerateVertexTriangles(RemoveV, [&](int32 tid)
 		{
-			if (GetTriNormalFunc(tid).Dot(Normal1) < DotTolerance)
+			if (TriNormals[tid].Dot(Normal1) < DotTolerance)
 			{
 				bIsFlat = false;
 			}
@@ -851,10 +895,11 @@ static bool IsCollapsableDevelopableEdge(const FDynamicMesh3& Mesh, int32 Collap
 			FIndex2i EdgeT = Mesh.GetEdgeT(eid);
 			if (EdgeT.B == IndexConstants::InvalidID)
 			{
-				return false;		// abort if one of the edges of RemoveV is a boundary edge (?)
+				// We already handled the cases where RemoveV is a boundary vert, so this shouldn't happen.
+				return ensure(false);
 			}
-			FVector3d Normal3 = GetTriNormalFunc(EdgeT.A);
-			FVector3d Normal4 = GetTriNormalFunc(EdgeT.B);
+			FVector3d Normal3 = TriNormals[EdgeT.A];
+			FVector3d Normal4 = TriNormals[EdgeT.B];
 
 			FIndex2i OtherEdgeV = Mesh.GetEdgeV(eid);
 			int32 OtherV = IndexUtil::FindEdgeOtherVertex(OtherEdgeV, RemoveV);
@@ -893,8 +938,15 @@ void TMeshSimplification<QuadricErrorType>::SimplifyToMinimalPlanar(
 		return;
 	}
 
-	// we don't collapse on the boundary
-	bHaveBoundary = false;
+	// This function doesn't affect the shape of the boundary, so the value of bPreserveBoundaryShape
+	// shouldn't matter. Yet in practice, having bPreserveBoundaryShape be true is problematic because
+	// the related block in CollapseEdge() arbitrarily decides that just one of the verts of a boundary
+	// edge can be collapsed to.
+	// TODO: the above is probably a minor bug for other forms of simplification too, but fixing it requires
+	// going through the details of other simplification methods. For SimplifyToMinimalPlanar, the simplest
+	// solution is to just eliminate that factor as a concern, since it should be free to collapse along
+	// colinear boundaries unless they are explicitly constrained.
+	TGuardValue<bool> PreserveBoundaryShapeOverride(bPreserveBoundaryShape, false); // Sets to false, restores on exit
 
 	// keep triangle normals
 	TArray<FVector3d> TriNormals;
@@ -961,7 +1013,7 @@ void TMeshSimplification<QuadricErrorType>::SimplifyToMinimalPlanar(
 		FVector3d va = FVector3d::Zero(), vb = FVector3d::Zero();
 		for ( int32 eid : CollapseEdges )
 		{
-			if ( (Mesh->IsEdge(eid) == false) || Mesh->IsBoundaryEdge(eid) )
+			if (!Mesh->IsEdge(eid))
 			{
 				continue;
 			}
@@ -969,58 +1021,37 @@ void TMeshSimplification<QuadricErrorType>::SimplifyToMinimalPlanar(
 			COUNT_ITERATIONS++;
 
 			FIndex2i ev = Mesh->GetEdgeV(eid);
-			bool bDevelopableA = DevelopableVerts[ev.A];
-			bool bDevelopableB = DevelopableVerts[ev.B];
-			if (bDevelopableA || bDevelopableB)		// this may change during execution as edges are collapsed
+			FDynamicMesh3::FEdgeCollapseInfo CollapseInfo;
+			ESimplificationResult Result = ESimplificationResult::Failed_OpNotSuccessful;
+
+			// Try collapsing to vert B.
+			if (DevelopableVerts[ev.A]
+				&& IsCollapsableDevelopableEdge(*Mesh, eid, ev.A, ev.B, PlanarDotTol, TriNormals, IsBoundaryVtxCache))
 			{
-				if (! bDevelopableA)		// any other preference for verts?
-				{
-					Swap(ev.A, ev.B);
-				}
+				Result = CollapseEdge(eid, Mesh->GetVertex(ev.B), CollapseInfo, ev.B);
+			}
 
-				bool bIsCollapsible = IsCollapsableDevelopableEdge(*Mesh, eid, ev.A, ev.B, PlanarDotTol, [&](int32 tid) { return TriNormals[tid]; });
-				if (bIsCollapsible)
-				{
-					FDynamicMesh3::FEdgeCollapseInfo collapseInfo;
-					ESimplificationResult result;
-					result = CollapseEdge(eid, Mesh->GetVertex(ev.B), collapseInfo, ev.B);
-					if (result == ESimplificationResult::Ok_Collapsed)
-					{
-						++num_last_pass;
+			// If that didn't work, try collapsing to vert A
+			if (Result != ESimplificationResult::Ok_Collapsed && DevelopableVerts[ev.B]
+				&& IsCollapsableDevelopableEdge(*Mesh, eid, ev.B, ev.A, PlanarDotTol, TriNormals, IsBoundaryVtxCache))
+			{
+				Result = CollapseEdge(eid, Mesh->GetVertex(ev.A), CollapseInfo, ev.A);
+			}
 
-						int vKeptID = collapseInfo.KeptVertex;
-						Mesh->EnumerateVertexTriangles(vKeptID, [&](int32 tid)
-						{
-							TriNormals[tid] = Mesh->GetTriNormal(tid);
-						});
-						for (int32 vid : Mesh->VtxVerticesItr(vKeptID))
-						{
-							DevelopableVerts[vid] = IsDevelopableVertex(*Mesh, vid, PlanarDotTol, [&](int32 tid) { return TriNormals[tid]; });
-						}
-						DevelopableVerts[vKeptID] = IsDevelopableVertex(*Mesh, vKeptID, PlanarDotTol, [&](int32 tid) { return TriNormals[tid]; });
-					}
-					else if (bDevelopableA && bDevelopableB &&
-								IsCollapsableDevelopableEdge(*Mesh, eid, ev.B, ev.A, PlanarDotTol, [&](int32 tid) { return TriNormals[tid]; }) )
-					{
-						// we can try collapsing to A
-						result = CollapseEdge(eid, Mesh->GetVertex(ev.A), collapseInfo, ev.A);
-						if (result == ESimplificationResult::Ok_Collapsed)
-						{
-							++num_last_pass;
+			if (Result == ESimplificationResult::Ok_Collapsed)
+			{
+				++num_last_pass;
 
-							int vKeptID = collapseInfo.KeptVertex;
-							Mesh->EnumerateVertexTriangles(vKeptID, [&](int32 tid)
-							{
-								TriNormals[tid] = Mesh->GetTriNormal(tid);
-							});
-							for (int32 vid : Mesh->VtxVerticesItr(vKeptID))
-							{
-								DevelopableVerts[vid] = IsDevelopableVertex(*Mesh, vid, PlanarDotTol, [&](int32 tid) { return TriNormals[tid]; });
-							}
-							DevelopableVerts[vKeptID] = IsDevelopableVertex(*Mesh, vKeptID, PlanarDotTol, [&](int32 tid) { return TriNormals[tid]; });
-						}
-					}
+				int vKeptID = CollapseInfo.KeptVertex;
+				Mesh->EnumerateVertexTriangles(vKeptID, [&](int32 tid)
+				{
+					TriNormals[tid] = Mesh->GetTriNormal(tid);
+				});
+				for (int32 vid : Mesh->VtxVerticesItr(vKeptID))
+				{
+					DevelopableVerts[vid] = IsDevelopableVertex(*Mesh, vid, PlanarDotTol, [&](int32 tid) { return TriNormals[tid]; });
 				}
+				DevelopableVerts[vKeptID] = IsDevelopableVertex(*Mesh, vKeptID, PlanarDotTol, [&](int32 tid) { return TriNormals[tid]; });
 			}
 		}
 
