-Added SmoothBoneWeights geometry operator that will contain a collection of algorithms for smoothing skin/bone weights.
-Refactored the SkinWeightsPaintTool to use the new operator.

diff --git a/Engine/Plugins/Experimental/MeshModelingToolsetExp/Source/MeshModelingToolsEditorOnlyExp/Private/SkeletalMesh/SkinWeightsPaintTool.cpp b/Engine/Plugins/Experimental/MeshModelingToolsetExp/Source/MeshModelingToolsEditorOnlyExp/Private/SkeletalMesh/SkinWeightsPaintTool.cpp
--- a/Engine/Plugins/Experimental/MeshModelingToolsetExp/Source/MeshModelingToolsEditorOnlyExp/Private/SkeletalMesh/SkinWeightsPaintTool.cpp
+++ b/Engine/Plugins/Experimental/MeshModelingToolsetExp/Source/MeshModelingToolsEditorOnlyExp/Private/SkeletalMesh/SkinWeightsPaintTool.cpp
@@ -23,6 +23,7 @@
 #include "Selection/PolygonSelectionMechanic.h"
 #include "Spatial/PointSetHashTable.h"
 #include "Util/ColorConstants.h"
+#include "Operations/SmoothBoneWeights.h"
 
 #include UE_INLINE_GENERATED_CPP_BY_NAME(SkinWeightsPaintTool)
 
@@ -36,6 +37,44 @@ static EAsyncExecution SkinPaintToolAsyncExecTarget = EAsyncExecution::ThreadPoo
 // any weight below this value is ignored, since it won't be representable in unsigned 16-bit precision
 constexpr float MinimumWeightThreshold = 1.0f / 65535.0f;
 
+class FPaintToolWeightsDataSource : public UE::Geometry::TBoneWeightsDataSource<int32, float>
+{
+public:
+
+	FPaintToolWeightsDataSource(const SkinPaintTool::FSkinToolWeights* InWeights)
+	:
+	Weights(InWeights) 
+	{
+		checkSlow(Weights);
+	}
+
+	virtual ~FPaintToolWeightsDataSource() = default;
+
+	virtual int32 GetBoneNum(const int32 VertexID) override
+	{
+		return Weights->PreChangeWeights[VertexID].Num();
+	}
+
+	virtual int32 GetBoneIndex(const int32 VertexID, const int32 Index) override
+	{
+		return Weights->PreChangeWeights[VertexID][Index].BoneIndex;
+	}
+
+	virtual float GetBoneWeight(const int32 VertexID, const int32 Index) override
+	{
+		return Weights->PreChangeWeights[VertexID][Index].Weight;
+	}
+
+	virtual float GetWeightOfBoneOnVertex(const int32 VertexID, const int32 BoneIndex) override
+	{
+		return Weights->GetWeightOfBoneOnVertex(BoneIndex, VertexID, Weights->PreChangeWeights);
+	}
+
+protected:
+
+	const SkinPaintTool::FSkinToolWeights* Weights = nullptr;
+};
+
 USkinWeightsPaintToolProperties::USkinWeightsPaintToolProperties()
 {
 	ColorRamp.Add(FLinearColor::Blue);
@@ -673,6 +712,10 @@ void USkinWeightsPaintTool::Setup()
 		[this]() { return MeshSpatial.Get(); }
 	);
 	
+	SmoothWeightsDataSource = MakeUnique<FPaintToolWeightsDataSource>(&Weights);
+	SmoothWeightsOp = MakeUnique<UE::Geometry::TSmoothBoneWeights<int32, float>>(PreviewMesh->GetMesh(), SmoothWeightsDataSource.Get());
+	SmoothWeightsOp->MinimumWeightThreshold = MinimumWeightThreshold;
+	
 	// inform user of tool keys
 	// TODO talk with UX team about viewport overlay to show hotkeys
 	GetToolManager()->DisplayMessage(
@@ -1183,91 +1226,30 @@ void USkinWeightsPaintTool::RelaxWeightOnVertices(
 	const float UseStrength,
 	FMultiBoneWeightEdits& InOutWeightEdits)
 {
-	const FDynamicMesh3* CurrentMesh = PreviewMesh->GetMesh();
-	const UE::Geometry::FNonManifoldMappingSupport NonManifoldMappingSupport(*CurrentMesh);
-
-	auto NormalizeWeights = [](TMap<BoneIndex, float>& InOutWeights)
+	if (!ensure(SmoothWeightsOp))
 	{
-		float TotalWeight = 0.f;
-		for (const TTuple<BoneIndex, float>& Weight : InOutWeights)
-		{
-			TotalWeight += Weight.Value;
-		}
-		for (TTuple<BoneIndex, float>& Weight : InOutWeights)
-		{
-			Weight.Value /= TotalWeight;
-		}
-	};
+		return;
+	}
 
-	// for each vertex in the stamp...
-	constexpr int32 AvgNumNeighbors = 8;
-	using VertexNeighborWeights = TArray<float, TInlineAllocator<AvgNumNeighbors>>;
-	TArray<int32> AllNeighborVertices;
-	TMap<BoneIndex, VertexNeighborWeights> WeightsOnAllNeighbors;
-	TMap<BoneIndex, float> FinalWeights;
 	for (int32 Index = 0; Index < VerticesToEdit.Num(); ++Index)
 	{
 		const int32 VertexID = VerticesToEdit[Index];
-		const int32 SrcVertexID = NonManifoldMappingSupport.GetOriginalNonManifoldVertexID(VertexID);
-
-		// get list of all neighboring vertices, AND this vertex
-		AllNeighborVertices.Reset();
-		AllNeighborVertices.Add(VertexID);
-		for (const int32 NeighborVertexID : CurrentMesh->VtxVerticesItr(SrcVertexID))
-		{
-			AllNeighborVertices.Add(NeighborVertexID);
-		}
+		const float UseFalloff = VertexFalloffs.IsValidIndex(Index) ? VertexFalloffs[Index] * UseStrength : UseStrength;
 
-		// get all weights above a given threshold across ALL neighbors (including self)
-		WeightsOnAllNeighbors.Reset();
-		for (const int32 NeighborVertexID : AllNeighborVertices)
-		{
-			for (const FVertexBoneWeight& BoneWeight : Weights.PreChangeWeights[NeighborVertexID])
-			{
-				if (BoneWeight.Weight > MinimumWeightThreshold)
-				{
-					VertexNeighborWeights& BoneWeights = WeightsOnAllNeighbors.FindOrAdd(BoneWeight.BoneIndex);
-					BoneWeights.Add(BoneWeight.Weight);
-				}
-			}
-		}
+		TMap<int32, float> FinalWeights;
+		const bool bSmoothSuccess = SmoothWeightsOp->SmoothWeightsAtVertex(VertexID, UseFalloff, FinalWeights);
 
-		// calculate single average weight of each bone on all the neighbors
-		FinalWeights.Reset();
-		for (const TTuple<BoneIndex, VertexNeighborWeights>& NeighborWeights : WeightsOnAllNeighbors)
+		if (ensure(bSmoothSuccess))
 		{
-			float TotalWeightOnThisBone = 0.f;
-			for (const float& Value : NeighborWeights.Value)
+			// apply weight edits
+			for (const TTuple<BoneIndex, float>& FinalWeight : FinalWeights)
 			{
-				TotalWeightOnThisBone += Value;
+				// record an edit for this vertex, for this bone
+				const int32 BoneIndex = FinalWeight.Key;
+				const float NewWeight = FinalWeight.Value;
+				const float OldWeight = Weights.GetWeightOfBoneOnVertex(BoneIndex, VertexID, Weights.PreChangeWeights);
+				InOutWeightEdits.MergeSingleEdit(BoneIndex, VertexID, OldWeight, NewWeight);
 			}
-			FinalWeights.Add(NeighborWeights.Key, TotalWeightOnThisBone / NeighborWeights.Value.Num());
-		}
-
-		// normalize the weights
-		NormalizeWeights(FinalWeights);
-
-		// lerp weights from previous values, to fully relaxed values by brush strength scaled by falloff
-		const float UseFalloff = VertexFalloffs.IsValidIndex(Index) ? VertexFalloffs[Index] * UseStrength : UseStrength;
-		for (TTuple<BoneIndex, float>& FinalWeight : FinalWeights)
-		{
-			const int32 BoneIndex = FinalWeight.Key;
-			float NewWeight = FinalWeight.Value;
-			float OldWeight = Weights.GetWeightOfBoneOnVertex(BoneIndex, VertexID, Weights.PreChangeWeights);
-			FinalWeight.Value = FMath::Lerp(OldWeight, NewWeight, UseFalloff);
-		}
-
-		// normalize again
-		NormalizeWeights(FinalWeights);
-
-		// apply weight edits
-		for (const TTuple<BoneIndex, float>& FinalWeight : FinalWeights)
-		{
-			// record an edit for this vertex, for this bone
-			const int32 BoneIndex = FinalWeight.Key;
-			const float NewWeight = FinalWeight.Value;
-			const float OldWeight = Weights.GetWeightOfBoneOnVertex(BoneIndex, VertexID, Weights.PreChangeWeights);
-			InOutWeightEdits.MergeSingleEdit(BoneIndex, VertexID, OldWeight, NewWeight);
 		}
 	}
 }
diff --git a/Engine/Plugins/Experimental/MeshModelingToolsetExp/Source/MeshModelingToolsEditorOnlyExp/Public/SkeletalMesh/SkinWeightsPaintTool.h b/Engine/Plugins/Experimental/MeshModelingToolsetExp/Source/MeshModelingToolsEditorOnlyExp/Public/SkeletalMesh/SkinWeightsPaintTool.h
--- a/Engine/Plugins/Experimental/MeshModelingToolsetExp/Source/MeshModelingToolsEditorOnlyExp/Public/SkeletalMesh/SkinWeightsPaintTool.h
+++ b/Engine/Plugins/Experimental/MeshModelingToolsetExp/Source/MeshModelingToolsEditorOnlyExp/Public/SkeletalMesh/SkinWeightsPaintTool.h
@@ -21,6 +21,12 @@ struct FMeshDescription;
 class USkinWeightsPaintTool;
 class UPolygonSelectionMechanic;
 
+namespace UE::Geometry 
+{
+	template <typename BoneIndexType, typename BoneWeightType> class TBoneWeightsDataSource;
+	template <typename BoneIndexType, typename BoneWeightType> class TSmoothBoneWeights;
+}
+
 using BoneIndex = int32;
 using VertexIndex = int32;
 
@@ -370,6 +376,10 @@ protected:
 	// cached mirror data
 	SkinPaintTool::FSkinMirrorData MirrorData;
 
+	// Smooth weights data source and operator
+	TUniquePtr<UE::Geometry::TBoneWeightsDataSource<int32, float>> SmoothWeightsDataSource;
+	TUniquePtr<UE::Geometry::TSmoothBoneWeights<int32, float>> SmoothWeightsOp;
+
 	// vertex colors updated when switching current bone or editing weights
 	void UpdateCurrentBoneVertexColors();
 	FVector4f WeightToColor(float Value) const;
diff --git a/Engine/Plugins/Runtime/GeometryProcessing/Source/DynamicMesh/Private/Operations/SmoothBoneWeights.cpp b/Engine/Plugins/Runtime/GeometryProcessing/Source/DynamicMesh/Private/Operations/SmoothBoneWeights.cpp
new file mode 100644
--- /dev/null
+++ b/Engine/Plugins/Runtime/GeometryProcessing/Source/DynamicMesh/Private/Operations/SmoothBoneWeights.cpp
@@ -0,0 +1,267 @@
+// Copyright Epic Games, Inc. All Rights Reserved.
+
+#include "Operations/SmoothBoneWeights.h"
+#include "DynamicMesh/DynamicMesh3.h"
+#include "DynamicMesh/DynamicVertexSkinWeightsAttribute.h"
+#include "DynamicMesh/DynamicMeshAttributeSet.h"
+#include "DynamicMesh/NonManifoldMappingSupport.h"
+#include "Util/ProgressCancel.h"
+#include "BoneWeights.h"
+#include "Math/UnrealMathUtility.h"
+
+using namespace UE::AnimationCore;
+using namespace UE::Geometry;
+
+namespace SmoothBoneWeightsLocals 
+{	
+	template <typename BoneIndexType, typename BoneWeightType>
+	void NormalizeWeights(TMap<BoneIndexType, BoneWeightType>& InOutWeights)
+	{
+		BoneWeightType TotalWeight = (BoneWeightType) 0.f;
+		for (const TTuple<BoneIndexType, BoneWeightType>& Weight : InOutWeights)
+		{
+			TotalWeight += Weight.Value;
+		}
+
+		if (!FMath::IsNearlyZero(TotalWeight))
+		{
+			for (TTuple<BoneIndexType, BoneWeightType>& Weight : InOutWeights)
+			{
+				Weight.Value /= TotalWeight;
+			}
+		}
+	}
+	
+	/** Data source implementation for bone weights data stored in the FDynamicMeshVertexSkinWeightsAttribute. */
+	class FSkinWeightsAttributeDataSource : public TBoneWeightsDataSource<FBoneIndexType, float>
+	{
+	public:
+		FSkinWeightsAttributeDataSource(const FDynamicMeshVertexSkinWeightsAttribute* InAttribute)
+		:
+		Attribute(InAttribute) 
+		{
+			checkSlow(Attribute);
+		}
+
+		virtual ~FSkinWeightsAttributeDataSource() = default; 
+
+		virtual int32 GetBoneNum(const int32 VertexID) override
+		{
+			FBoneWeights Weights;
+			Attribute->GetValue(VertexID, Weights);
+			return Weights.Num();
+		}
+
+		virtual FBoneIndexType GetBoneIndex(const int32 VertexID, const int32 Index) override
+		{
+			FBoneWeights Weights;
+			Attribute->GetValue(VertexID, Weights);
+			return Weights[Index].GetBoneIndex();
+		}
+
+		virtual float GetBoneWeight(const int32 VertexID, const int32 Index) override
+		{
+			FBoneWeights Weights;
+			Attribute->GetValue(VertexID, Weights);
+			return Weights[Index].GetWeight();
+		}
+
+		virtual float GetWeightOfBoneOnVertex(const int32 VertexID, const FBoneIndexType BoneIndex) override
+		{
+			FBoneWeights Weights;
+			Attribute->GetValue(VertexID, Weights);
+			for (const FBoneWeight& BoneWeight : Weights)
+			{
+				if (BoneWeight.GetBoneIndex() == BoneIndex)
+				{
+					return BoneWeight.GetWeight();;
+				}
+			}
+
+			return 0.f;
+		}
+
+	protected:
+		const FDynamicMeshVertexSkinWeightsAttribute* Attribute = nullptr;
+	};
+}
+
+
+//
+// TSmoothBoneWeights
+//
+
+template <typename BoneIndexType, typename BoneWeightType>
+TSmoothBoneWeights<BoneIndexType, BoneWeightType>::TSmoothBoneWeights(const FDynamicMesh3* InSourceMesh,
+																	  TBoneWeightsDataSource<BoneIndexType, BoneWeightType>* InDataSource)
+:
+SourceMesh(InSourceMesh),
+DataSource(InDataSource)
+{
+	checkSlow(InSourceMesh);
+}
+
+template <typename BoneIndexType, typename BoneWeightType>
+bool TSmoothBoneWeights<BoneIndexType, BoneWeightType>::Cancelled()
+{
+	return (Progress == nullptr) ? false : Progress->Cancelled();
+}
+
+template <typename BoneIndexType, typename BoneWeightType>
+EOperationValidationResult TSmoothBoneWeights<BoneIndexType, BoneWeightType>::Validate()
+{	
+	if (SourceMesh == nullptr) 
+	{
+		return EOperationValidationResult::Failed_UnknownReason;
+	}
+
+	if (DataSource == nullptr)
+	{
+		return EOperationValidationResult::Failed_UnknownReason;
+	}
+
+	return EOperationValidationResult::Ok;
+}
+
+template <typename BoneIndexType, typename BoneWeightType>
+bool TSmoothBoneWeights<BoneIndexType, BoneWeightType>::SmoothWeightsAtVertex(const int32  VertexID,
+																			  const BoneWeightType VertexFalloff,
+																			  TMap<BoneIndexType, BoneWeightType>& OutFinalWeights)
+{
+	using namespace SmoothBoneWeightsLocals;
+	
+	const UE::Geometry::FNonManifoldMappingSupport NonManifoldMappingSupport(*SourceMesh);
+
+	// for each vertex in the stamp...
+	constexpr int32 AvgNumNeighbors = 8;
+	using VertexNeighborWeights = TArray<BoneWeightType, TInlineAllocator<AvgNumNeighbors>>;
+	TArray<int32> AllNeighborVertices;
+	TMap<BoneIndexType, VertexNeighborWeights> WeightsOnAllNeighbors;
+	const int32 SrcVertexID = NonManifoldMappingSupport.GetOriginalNonManifoldVertexID(VertexID);
+
+	// get list of all neighboring vertices, AND this vertex
+	AllNeighborVertices.Add(VertexID);
+	for (const int32 NeighborVertexID : SourceMesh->VtxVerticesItr(SrcVertexID))
+	{
+		AllNeighborVertices.Add(NeighborVertexID);
+	}
+
+	// get all weights above a given threshold across ALL neighbors (including self)
+	for (const int32 NeighborVertexID : AllNeighborVertices)
+	{
+		for (int32 Index = 0; Index < DataSource->GetBoneNum(NeighborVertexID); ++Index)
+		{	
+			const BoneWeightType Weight = DataSource->GetBoneWeight(NeighborVertexID, Index);
+			if (Weight > MinimumWeightThreshold)
+			{
+				VertexNeighborWeights& BoneWeights = WeightsOnAllNeighbors.FindOrAdd(DataSource->GetBoneIndex(NeighborVertexID, Index));
+				BoneWeights.Add(Weight);
+			}
+		}
+	}
+
+	// calculate single average weight of each bone on all the neighbors
+	OutFinalWeights.Reset();
+	for (const TTuple<BoneIndexType, VertexNeighborWeights>& NeighborWeights : WeightsOnAllNeighbors)
+	{
+		BoneWeightType TotalWeightOnThisBone = (BoneWeightType) 0.f;
+		for (const BoneWeightType& Value : NeighborWeights.Value)
+		{
+			TotalWeightOnThisBone += Value;
+		}
+		OutFinalWeights.Add(NeighborWeights.Key, TotalWeightOnThisBone / (BoneWeightType)NeighborWeights.Value.Num());
+	}
+
+	// normalize the weights
+	NormalizeWeights(OutFinalWeights);
+
+	// lerp weights from previous values, to fully relaxed values by brush strength scaled by falloff
+	for (TTuple<BoneIndexType, BoneWeightType>& FinalWeight : OutFinalWeights)
+	{
+		const BoneIndexType BoneIndex = FinalWeight.Key;
+		BoneWeightType NewWeight = FinalWeight.Value;
+		BoneWeightType OldWeight = DataSource->GetWeightOfBoneOnVertex(VertexID, BoneIndex);
+		FinalWeight.Value = FMath::Lerp<BoneWeightType>(OldWeight, NewWeight, VertexFalloff);
+	}
+
+	// normalize again
+	NormalizeWeights(OutFinalWeights);
+
+	if (Cancelled()) 
+	{
+		return false;
+	}
+		
+	return true;
+}
+
+
+//
+// FSmoothDynamicMeshVertexSkinWeights
+//
+
+FSmoothDynamicMeshVertexSkinWeights::FSmoothDynamicMeshVertexSkinWeights(const FDynamicMesh3* InSourceMesh, const FName InProfile)
+:
+TSmoothBoneWeights(InSourceMesh, nullptr)
+{
+	if (SourceMesh && SourceMesh->Attributes())
+	{
+		Attribute = SourceMesh->Attributes()->GetSkinWeightsAttribute(InProfile);
+		if (Attribute)
+		{
+			SkinWeightsAttributeDataSource = MakeUnique<SmoothBoneWeightsLocals::FSkinWeightsAttributeDataSource>(Attribute);
+			DataSource = SkinWeightsAttributeDataSource.Get();
+		}
+	}
+}
+
+FSmoothDynamicMeshVertexSkinWeights::FSmoothDynamicMeshVertexSkinWeights(const FDynamicMesh3* InSourceMesh, FDynamicMeshVertexSkinWeightsAttribute* InAttribute)
+:
+TSmoothBoneWeights(InSourceMesh, nullptr),
+Attribute(InAttribute)
+{
+	if (InAttribute)
+	{
+		SkinWeightsAttributeDataSource = MakeUnique<SmoothBoneWeightsLocals::FSkinWeightsAttributeDataSource>(InAttribute);
+		DataSource = SkinWeightsAttributeDataSource.Get();
+	}
+}
+
+
+EOperationValidationResult FSmoothDynamicMeshVertexSkinWeights::Validate()
+{	
+	if (Attribute == nullptr) 
+	{
+		return EOperationValidationResult::Failed_UnknownReason;
+	}
+
+	return TSmoothBoneWeights::Validate();
+}
+
+bool FSmoothDynamicMeshVertexSkinWeights::SmoothWeightsAtVertex(const int32 VertexID, const float VertexFalloff)
+{	
+	TMap<FBoneIndexType, float> FinalWeights;
+	if (TSmoothBoneWeights<FBoneIndexType, float>::SmoothWeightsAtVertex(VertexID, VertexFalloff, FinalWeights))
+	{
+		FBoneWeightsSettings BoneSettings;
+		BoneSettings.SetNormalizeType(EBoneWeightNormalizeType::None);
+		FBoneWeights WeightArray;
+		
+		for (const TTuple<FBoneIndexType, float>& FinalWeight : FinalWeights)
+		{
+			WeightArray.SetBoneWeight(FBoneWeight(FinalWeight.Key, FinalWeight.Value), BoneSettings);
+		}
+		
+		Attribute->SetValue(VertexID, WeightArray);
+
+		return true;
+	}
+
+	return false;
+}
+
+
+
+// template instantiation
+template class DYNAMICMESH_API UE::Geometry::TSmoothBoneWeights<FBoneIndexType, float>;
+template class DYNAMICMESH_API UE::Geometry::TSmoothBoneWeights<int, float>;
\ No newline at end of file
diff --git a/Engine/Plugins/Runtime/GeometryProcessing/Source/DynamicMesh/Public/Operations/SmoothBoneWeights.h b/Engine/Plugins/Runtime/GeometryProcessing/Source/DynamicMesh/Public/Operations/SmoothBoneWeights.h
new file mode 100644
--- /dev/null
+++ b/Engine/Plugins/Runtime/GeometryProcessing/Source/DynamicMesh/Public/Operations/SmoothBoneWeights.h
@@ -0,0 +1,208 @@
+// Copyright Epic Games, Inc. All Rights Reserved.
+
+#pragma once
+
+#include "GeometryTypes.h"
+#include "UObject/NameTypes.h"
+#include "Containers/Map.h"
+#include "BoneIndices.h"
+
+// Forward declarations
+class FProgressCancel;
+namespace UE::Geometry 
+{ 
+	class FDynamicMesh3; 
+	template<typename ParentType> class TDynamicVertexSkinWeightsAttribute;
+	using FDynamicMeshVertexSkinWeightsAttribute = TDynamicVertexSkinWeightsAttribute<FDynamicMesh3>;
+}
+
+
+namespace UE
+{
+namespace Geometry
+{
+
+
+/** 
+ * Interface for getting the bone weights data. Should be subclassed and implemented if the bone weights data
+ * is not stored using the FDynamicMeshVertexSkinWeightsAttribute.
+ */
+template <typename BoneIndexType, typename BoneWeightType>
+class TBoneWeightsDataSource
+{
+public:
+	virtual ~TBoneWeightsDataSource() = default;
+
+	/** Number of bones at a vertex. */
+	virtual int32 GetBoneNum(const int32 InVertexIDs) = 0;
+
+	/** Bone ID of the bone stored at an index Index, where Index is between 0 and GetBoneNum(InVertexIDs). */
+	virtual BoneIndexType GetBoneIndex(const int32 InVertexIDs, const int32 Index) = 0;
+
+	/** Bone weight of the bone stored at an index Index, where Index is between 0 and GetBoneNum(InVertexIDs). */
+	virtual BoneWeightType GetBoneWeight(const int32 InVertexIDs, const int32 Index) = 0;
+
+	/** Bone weight of the bone with Bone ID at a vertex. Must return 0 if bone has no influence. */
+	virtual BoneWeightType GetWeightOfBoneOnVertex(const int32 InVertexIDs, const BoneIndexType BoneIndex) = 0;
+};
+
+
+
+
+/**
+ * Collection of algorithms for smoothing bone weights. Bone weight data is fetched using the TBoneWeightsDataSource
+ * which is implemented by the caller to match their custom data storage of bone weights.
+ *
+ *
+ * Example usage:
+ *
+ * // Mesh whose bone weights we are smoothing
+ * FDynamicMesh SourceMesh = ... ;
+ * 
+ * // The caller should create a subclass of TBoneWeightsDataSource that implements its interface
+ * TUniquePtr<TBoneWeightsDataSource> DataSource = ... ;
+ * 
+ * // Setup the smoothing operator
+ * TSmoothBoneWeights SmoothBoneWeights(&SourceMesh, DataSource.Get());
+ *
+ * // Array of vertex IDs to apply smoothing to
+ * TArray<int32> VerticesToSmooth = ... ;
+ * 
+ * // Run the operator
+ * const float SmoothStrength = 0.2;
+ * if (SmoothBoneWeights.Validate() == EOperationValidationResult::Ok)
+ * {
+ * 		for (const int32 VertexID : VerticesToSmooth)
+ * 		{
+ * 			TMap<BoneIndexType, BoneWeightType> FinalWeights;
+ *      	SmoothBoneWeights.SmoothWeightsAtVertex(VertexID, SmoothStrength, FinalWeights);
+ * 
+ * 			// Use FinalWeights to modify your weights
+ * 		}
+ * }
+ */
+template <typename BoneIndexType, typename BoneWeightType>
+class DYNAMICMESH_API TSmoothBoneWeights
+{
+public:
+
+	//
+	// Optional Inputs
+	//
+	
+    /** Set this to be able to cancel the running operation. */
+	FProgressCancel* Progress = nullptr;
+
+	/** Only consider the weights above the threshold. */
+	BoneWeightType MinimumWeightThreshold = (BoneWeightType) 0.0f; 
+
+protected:
+		
+	const FDynamicMesh3* SourceMesh;
+	
+	/** Data source supplied by the caller. */
+	TBoneWeightsDataSource<BoneIndexType, BoneWeightType>* DataSource = nullptr;
+	
+public:
+
+	/**
+	 * @param InSourceMesh The mesh whose vertices contain bone weights data.
+	 * @param InDataSource The bone weights data.
+	 */
+	TSmoothBoneWeights(const FDynamicMesh3* InSourceMesh, TBoneWeightsDataSource<BoneIndexType, BoneWeightType>* InDataSource = nullptr);
+
+	virtual ~TSmoothBoneWeights() = default;
+
+	/** @return EOperationValidationResult::Ok if we can apply operation, or error code if we cannot. */
+	virtual EOperationValidationResult Validate();
+
+	/**
+     * Compute smoothed weights at a vertex and return the result.
+	 * 
+	 * @param VertexID		  The vertex whose weights we are smoothing.
+     * @param VertexFalloff   The lerp value between the old and the new weight.
+     * @param OutFinalWeights The smoothed weights.
+	 * 
+     * @return true if the algorithm succeeds, false if it failed or was canceled by the user.
+	 */
+	virtual bool SmoothWeightsAtVertex(const int32 VertexID,
+									   const BoneWeightType VertexFalloff, 
+									   TMap<BoneIndexType, BoneWeightType>& OutFinalWeights);
+
+protected:
+	
+    /** @return if true, abort the computation. */
+	virtual bool Cancelled();
+};
+
+
+
+
+/**
+ * Subclass of TSmoothBoneWeights for dealing with smoothing of bone weights stored in the 
+ * FDynamicMeshVertexSkinWeightsAttribute. 
+ *
+ *
+ * Example usage:
+ *
+ * // Mesh whose bone weights we are smoothing
+ * FDynamicMesh SourceMesh = ... ;
+ * 
+ * // Setup the smoothing operator
+ * FSmoothDynamicMeshVertexSkinWeights SmoothBoneWeights(&SourceMesh, FSkeletalMeshAttributes::DefaultSkinWeightProfileName);
+ *
+ * // Array of vertex IDs to apply smoothing to
+ * TArray<int32> VerticesToSmooth = ... ;
+ * 
+ * // Run the operator
+ * const float SmoothStrength = 0.2;
+ * if (SmoothBoneWeights.Validate() == EOperationValidationResult::Ok)
+ * {
+ * 		for (const int32 VertexID : VerticesToSmooth)
+ * 		{
+ *      	SmoothBoneWeights.SmoothWeightsAtVertex(VertexID, SmoothStrength);
+ * 		}
+ * }
+ */
+class DYNAMICMESH_API FSmoothDynamicMeshVertexSkinWeights : public TSmoothBoneWeights<FBoneIndexType, float>
+{
+public:
+
+	using TSmoothBoneWeights<FBoneIndexType, float>::SmoothWeightsAtVertex; 
+	
+	/**
+	 * @param InSourceMesh The mesh whose vertices contain bone weights data.
+	 * @param InProfile    The name of the skin weights profile containing the bone weights data.
+	 */
+	FSmoothDynamicMeshVertexSkinWeights(const FDynamicMesh3* InSourceMesh, const FName InProfile);
+
+	/**
+	 * @param InSourceMesh The mesh whose vertices contain bone weights data.
+	 * @param InAttribute  The skin weight attribute containing the bone data.
+	 */
+	FSmoothDynamicMeshVertexSkinWeights(const FDynamicMesh3* InSourceMesh, FDynamicMeshVertexSkinWeightsAttribute* InAttribute);
+
+	virtual ~FSmoothDynamicMeshVertexSkinWeights() = default;
+	
+	/** @return EOperationValidationResult::Ok if we can apply operation, or error code if we cannot. */
+	virtual EOperationValidationResult Validate() override;
+	
+	/**
+	 * Compute smoothed weights at a vertex and write the result into FDynamicMeshVertexSkinWeightsAttribute passed 
+	 * upon construction.
+	 *
+	 * @param VertexID		  The vertex whose weights we are smoothing.
+	 * @param VertexFalloff   The lerp value between the old and the new weight.
+	 *
+	 * @return true if the algorithm succeeds, false if it failed or was canceled by the user.
+	 */
+	bool SmoothWeightsAtVertex(const int32 VertexID, const float VertexFalloff);
+
+protected:
+	TUniquePtr<TBoneWeightsDataSource<FBoneIndexType, float>> SkinWeightsAttributeDataSource = nullptr;
+	FDynamicMeshVertexSkinWeightsAttribute* Attribute = nullptr;
+
+};
+
+} // end namespace UE::Geometry
+} // end namespace UE
