-Fixed a ConstructMatrixSolver typo.
-Runt the solve for each column in parallel using ParallelFor.
-Allow to set fixed values contraints via a  Dense matrix.
-When calling the Solve allow to output a sparse matrix by prunning small numbers.

#rnx

diff --git a/Engine/Plugins/Runtime/GeometryProcessing/Source/DynamicMesh/Private/Solvers/Internal/ConstrainedPoissonSolver.h b/Engine/Plugins/Runtime/GeometryProcessing/Source/DynamicMesh/Private/Solvers/Internal/ConstrainedPoissonSolver.h
--- a/Engine/Plugins/Runtime/GeometryProcessing/Source/DynamicMesh/Private/Solvers/Internal/ConstrainedPoissonSolver.h
+++ b/Engine/Plugins/Runtime/GeometryProcessing/Source/DynamicMesh/Private/Solvers/Internal/ConstrainedPoissonSolver.h
@@ -63,7 +63,7 @@ public:
 		SymmetricMatrixPtr.Reset(SymmatrixMatrixOperator.Release());
 
 		bMatrixSolverDirty = true;
-		MatrixSolver = ContructMatrixSolver(MatrixSolverType);
+		MatrixSolver = ConstructMatrixSolver(MatrixSolverType);
 	}
 
 
diff --git a/Engine/Plugins/Runtime/GeometryProcessing/Source/DynamicMesh/Private/Solvers/Internal/MatrixSolver.cpp b/Engine/Plugins/Runtime/GeometryProcessing/Source/DynamicMesh/Private/Solvers/Internal/MatrixSolver.cpp
--- a/Engine/Plugins/Runtime/GeometryProcessing/Source/DynamicMesh/Private/Solvers/Internal/MatrixSolver.cpp
+++ b/Engine/Plugins/Runtime/GeometryProcessing/Source/DynamicMesh/Private/Solvers/Internal/MatrixSolver.cpp
@@ -10,7 +10,7 @@ namespace Geometry
 IMatrixSolverBase::~IMatrixSolverBase() {};
 IIterativeMatrixSolverBase::~IIterativeMatrixSolverBase() {};
 
-TUniquePtr<IMatrixSolverBase> ContructMatrixSolver(const EMatrixSolverType& MatrixSolverType)
+TUniquePtr<IMatrixSolverBase> ConstructMatrixSolver(const EMatrixSolverType& MatrixSolverType)
 {
 	TUniquePtr<IMatrixSolverBase> ResultPtr;
 
diff --git a/Engine/Plugins/Runtime/GeometryProcessing/Source/DynamicMesh/Private/Solvers/Internal/MatrixSolver.h b/Engine/Plugins/Runtime/GeometryProcessing/Source/DynamicMesh/Private/Solvers/Internal/MatrixSolver.h
--- a/Engine/Plugins/Runtime/GeometryProcessing/Source/DynamicMesh/Private/Solvers/Internal/MatrixSolver.h
+++ b/Engine/Plugins/Runtime/GeometryProcessing/Source/DynamicMesh/Private/Solvers/Internal/MatrixSolver.h
@@ -38,7 +38,7 @@ namespace UE
 namespace Geometry
 {
 // Matrix Solver Factory
-TUniquePtr<IMatrixSolverBase> ContructMatrixSolver(const EMatrixSolverType& MatrixSolverType);
+TUniquePtr<IMatrixSolverBase> ConstructMatrixSolver(const EMatrixSolverType& MatrixSolverType);
 
 template <typename DirectSolverType>
 class TMatrixSolver : public IMatrixSolverBase
diff --git a/Engine/Plugins/Runtime/GeometryProcessing/Source/DynamicMesh/Private/Solvers/Internal/MeshDiffusionIntegrator.cpp b/Engine/Plugins/Runtime/GeometryProcessing/Source/DynamicMesh/Private/Solvers/Internal/MeshDiffusionIntegrator.cpp
--- a/Engine/Plugins/Runtime/GeometryProcessing/Source/DynamicMesh/Private/Solvers/Internal/MeshDiffusionIntegrator.cpp
+++ b/Engine/Plugins/Runtime/GeometryProcessing/Source/DynamicMesh/Private/Solvers/Internal/MeshDiffusionIntegrator.cpp
@@ -130,7 +130,7 @@ void FMeshDiffusionIntegrator::Integrate_BackwardEuler(const EMatrixSolverType M
 
 	SparseMatrix.makeCompressed();
 
-	TUniquePtr<IMatrixSolverBase> MatrixSolver = ContructMatrixSolver(MatrixSolverType);
+	TUniquePtr<IMatrixSolverBase> MatrixSolver = ConstructMatrixSolver(MatrixSolverType);
 
 	MatrixSolver->SetUp(SparseMatrix, bIsSymmetric);
 
diff --git a/Engine/Plugins/Runtime/GeometryProcessing/Source/DynamicMesh/Private/Solvers/Internal/MeshUVSolver.cpp b/Engine/Plugins/Runtime/GeometryProcessing/Source/DynamicMesh/Private/Solvers/Internal/MeshUVSolver.cpp
--- a/Engine/Plugins/Runtime/GeometryProcessing/Source/DynamicMesh/Private/Solvers/Internal/MeshUVSolver.cpp
+++ b/Engine/Plugins/Runtime/GeometryProcessing/Source/DynamicMesh/Private/Solvers/Internal/MeshUVSolver.cpp
@@ -345,7 +345,7 @@ bool FLeastSquaresConformalMeshUVSolver::SolveParameterization(const FSparseMatr
 	const int32 NumVerts = VtxLinearization.NumVerts();
 	
 	// create a suitable matrix solver
-	TUniquePtr<IMatrixSolverBase> MatrixSolver = ContructMatrixSolver(InSolverType);
+	TUniquePtr<IMatrixSolverBase> MatrixSolver = ConstructMatrixSolver(InSolverType);
 	MatrixSolver->SetUp(InSystemMatrix, false);
 	if (MatrixSolver->bSucceeded() == false)  // check that the factorization succeeded 
 	{
diff --git a/Engine/Plugins/Runtime/GeometryProcessing/Source/DynamicMesh/Private/Solvers/Internal/PowerMethodSolver.cpp b/Engine/Plugins/Runtime/GeometryProcessing/Source/DynamicMesh/Private/Solvers/Internal/PowerMethodSolver.cpp
--- a/Engine/Plugins/Runtime/GeometryProcessing/Source/DynamicMesh/Private/Solvers/Internal/PowerMethodSolver.cpp
+++ b/Engine/Plugins/Runtime/GeometryProcessing/Source/DynamicMesh/Private/Solvers/Internal/PowerMethodSolver.cpp
@@ -240,7 +240,7 @@ void FSparsePowerMethod::CreateSparseMatrixOperator(const FSparseMatrixD& InMatr
 
     OutMatrixOp.Factorize = [&InMatrix, MatrixSolverType, bIsSymmetric]() 
     {	
-        TUniquePtr<IMatrixSolverBase> Solver = ContructMatrixSolver(MatrixSolverType);
+        TUniquePtr<IMatrixSolverBase> Solver = ConstructMatrixSolver(MatrixSolverType);
         Solver->SetUp(InMatrix, bIsSymmetric);
 
         if (Solver->bSucceeded() == false) 
diff --git a/Engine/Plugins/Runtime/GeometryProcessing/Source/DynamicMesh/Private/Solvers/Internal/QuadraticProgramming.cpp b/Engine/Plugins/Runtime/GeometryProcessing/Source/DynamicMesh/Private/Solvers/Internal/QuadraticProgramming.cpp
--- a/Engine/Plugins/Runtime/GeometryProcessing/Source/DynamicMesh/Private/Solvers/Internal/QuadraticProgramming.cpp
+++ b/Engine/Plugins/Runtime/GeometryProcessing/Source/DynamicMesh/Private/Solvers/Internal/QuadraticProgramming.cpp
@@ -2,6 +2,8 @@
 
 #include "Solvers/Internal/QuadraticProgramming.h"
 #include "MatrixSolver.h"
+#include "Async/ParallelFor.h"
+#include <vector>
 
 using namespace UE::Geometry;
 
@@ -13,6 +15,11 @@ VectorF(InVectorF)
     checkSlow(InMatrixQ);
 }
 
+void FQuadraticProgramming::SetEnableParallelization(const bool bEnable)
+{
+    bUseParallel = bEnable;
+}
+
 bool FQuadraticProgramming::SetFixedConstraints(const TArray<int>* InFixedRowIndices, const FSparseMatrixD* InFixedValues)
 {
     if (InFixedRowIndices && InFixedValues && InFixedRowIndices->Num() == InFixedValues->rows())
@@ -20,7 +27,26 @@ bool FQuadraticProgramming::SetFixedConstraints(const TArray<int>* InFixedRowInd
         if (InFixedRowIndices->Num() <= MatrixQ->rows() && InFixedValues->rows() <= MatrixQ->rows())
         {
             FixedRowIndices = InFixedRowIndices;
-            FixedValues = InFixedValues;
+            FixedValuesSparse = InFixedValues; // dense and sparse constraints are mutually exclusive
+            FixedValuesDense = nullptr;
+            bFixedConstraintsSet = true; 
+            
+			return true;
+        }
+    }
+
+    return false;
+}
+
+bool FQuadraticProgramming::SetFixedConstraints(const TArray<int>* InFixedRowIndices, const FDenseMatrixD* InFixedValues)
+{
+    if (InFixedRowIndices && InFixedValues && InFixedRowIndices->Num() == InFixedValues->rows())
+    {
+        if (InFixedRowIndices->Num() <= MatrixQ->rows() && InFixedValues->rows() <= MatrixQ->rows())
+        {
+            FixedRowIndices = InFixedRowIndices;
+            FixedValuesDense = InFixedValues; // dense and sparse constraints are mutually exclusive
+            FixedValuesSparse = nullptr;
             bFixedConstraintsSet = true; 
             
 			return true;
@@ -93,7 +119,7 @@ bool FQuadraticProgramming::PreFactorize()
         // Construct a linear solver for a symmetric positive (semi-)definite matrix
         const EMatrixSolverType MatrixSolverType = EMatrixSolverType::FastestPSD;
         const bool bIsSymmetric = true;
-        Solver = ContructMatrixSolver(MatrixSolverType);
+        Solver = ConstructMatrixSolver(MatrixSolverType);
         Solver->SetUp(VariablesQ, bIsSymmetric);
         if (!ensure(Solver->bSucceeded())) 
         {
@@ -107,8 +133,9 @@ bool FQuadraticProgramming::PreFactorize()
     return false;
 }
 
-bool FQuadraticProgramming::Solve(FDenseMatrixD& Solution)
+bool FQuadraticProgramming::Solve(FDenseMatrixD& Solution, const bool bVariablesOnly)
 {
+    FDateTime StartTime = FDateTime::UtcNow();
     if (bFixedConstraintsSet)
     {   
         if (!ensureMsgf(Solver, TEXT("Solver was not setup. Call the SetUp() method first.")))
@@ -116,6 +143,15 @@ bool FQuadraticProgramming::Solve(FDenseMatrixD& Solution)
             return false;
         }
 
+        if (!ensureMsgf((FixedValuesSparse == nullptr && FixedValuesDense != nullptr) || 
+                        (FixedValuesSparse != nullptr && FixedValuesDense == nullptr),
+                        TEXT("Either the fixed constraints were not setup or both sparse and dense constraints are setup.")))
+        {
+            return false;
+        }
+
+        const int32 FixedValuesCols = FixedValuesSparse ? (int32)FixedValuesSparse->cols() : (int32)FixedValuesDense->cols();
+
         // The Q_vf matrix
         FSparseMatrixD VariablesFixedQ;
 		SliceSparseMatrix(*MatrixQ, VariableRowIndices, *FixedRowIndices, VariablesFixedQ);
@@ -127,12 +163,42 @@ bool FQuadraticProgramming::Solve(FDenseMatrixD& Solution)
 			SliceDenseMatrix(*VectorF, VariableRowIndices, VariablesF);
 		}
 		
-        Solution.resize(MatrixQ->rows(), FixedValues->cols());
+        // Matrix that will hold the final solution values. Optionally, omit fixed rows and only store variables.
+        const FDenseMatrixD::Index SolutionRows = bVariablesOnly ? static_cast<FDenseMatrixD::Index>(VariableRowIndices.Num()) : MatrixQ->rows();
+        Solution.resize(SolutionRows, FixedValuesCols);
+        Solution.setZero();
+
+		// If FixedValues is a sparse matrix then we can check which columns are non-zero and 
+		// skip the solve for zero columns when VectorF is null
+        TArray<int32> NonZeroFixedValuesCols;
+        if (FixedValuesSparse)
+        {
+            for (int32 ColdIdx = 0; ColdIdx < FixedValuesSparse->outerSize(); ++ColdIdx)
+            {
+                if (!VectorF && FixedValuesSparse->innerVector(ColdIdx).nonZeros() > 0)
+                {
+                    NonZeroFixedValuesCols.Add(ColdIdx);
+                }
+            }
+        }
 
-        // We can iterate over every column of X and solve for each column individually
-        for (int ColIdx = 0; ColIdx < Solution.cols(); ++ColIdx) //TODO: Parallelize
+		// Solve for each non-zero column in parallel (if bUseParallel == true)
+        TArray<bool> SolveSucceeded; // track for which columns we failed to solve
+        SolveSucceeded.Init(true, static_cast<int32>(Solution.cols()));
+        const int32 NumIterations = FixedValuesSparse ? NonZeroFixedValuesCols.Num() : FixedValuesCols;
+        ParallelFor(NumIterations, [&](int32 IterIdx)
         {
-            const FColumnVectorD FixedVector = FixedValues->col(ColIdx);
+            const int32 ColIdx = FixedValuesSparse ? NonZeroFixedValuesCols[IterIdx] : IterIdx;
+			FColumnVectorD FixedVector;
+			if (FixedValuesSparse)
+			{
+				FixedVector = FixedValuesSparse->col(ColIdx);
+			}
+			else
+			{
+				FixedVector = FixedValuesDense->col(ColIdx);
+			}
+
 
 			// The -(f_v + Q_vf * x_f) vector. If the linear coefficients were not specified, skip f_v.
             FColumnVectorD VectorB;
@@ -149,36 +215,114 @@ bool FQuadraticProgramming::Solve(FDenseMatrixD& Solution)
             Solver->Solve(VectorB, VariablesVector);
             if (ensure(Solver->bSucceeded()))
             {   
-                // Copy over the fixed and variable values to the solution matrix
-                for (int RowIdx = 0; RowIdx < FixedRowIndices->Num(); ++RowIdx)
+                if (bVariablesOnly)
                 {
-                    Solution((*FixedRowIndices)[RowIdx], ColIdx) = FixedVector(RowIdx);
+                    // Copy over only the variable values to the solution matrix
+                    Solution.col(ColIdx) = VariablesVector;
                 }
-
-                for (int RowIdx = 0; RowIdx < VariableRowIndices.Num(); ++RowIdx)
+                else
                 {
-                    Solution(VariableRowIndices[RowIdx], ColIdx) = VariablesVector(RowIdx);
+                    // Copy over the fixed and variable values to the solution matrix
+                    for (int RowIdx = 0; RowIdx < FixedRowIndices->Num(); ++RowIdx)
+                    {
+                        Solution((*FixedRowIndices)[RowIdx], ColIdx) = FixedVector(RowIdx);
+                    }
+
+                    for (int RowIdx = 0; RowIdx < VariableRowIndices.Num(); ++RowIdx)
+                    {
+                        Solution(VariableRowIndices[RowIdx], ColIdx) = VariablesVector(RowIdx);
+                    }
                 }
             }
 			else 
 			{
-				return false;
+				SolveSucceeded[ColIdx] = false;
 			}
+        }, bUseParallel ? EParallelForFlags::None : EParallelForFlags::ForceSingleThread);
+
+        // Check if all solves were successful
+        for (const bool bSuccess : SolveSucceeded)
+        {
+            if (!bSuccess)
+            {
+                return false;
+            }
         }
+    }
 
-		return true;
+    SolveTimeElapsedInSec = (FDateTime::UtcNow() - StartTime).GetTotalSeconds();
+
+	return true;
+}
+
+bool FQuadraticProgramming::Solve(FSparseMatrixD& Solution, const bool bVariablesOnly, const double Tolerance)
+{
+    // Eigen library has a sparseView() function but it prunes all values below a Tolerance, even negative values. 
+    // We want to keep the values whose absolute value is above the Tolerance. 
+
+    FDenseMatrixD SolutionDense;
+    if (!Solve(SolutionDense, bVariablesOnly))
+    {
+        return false;
+    }
+
+    std::vector<Eigen::Triplet<FSparseMatrixD::Scalar>> Triplets;
+    for (int32 ColIdx = 0; ColIdx < SolutionDense.cols(); ++ColIdx)
+    {
+        for (int32 RowIdx = 0; RowIdx < SolutionDense.rows(); ++RowIdx)
+        {
+            FDenseMatrixD::Scalar Value = SolutionDense(RowIdx, ColIdx);
+            if (FMath::Abs<FDenseMatrixD::Scalar>(Value) > (FDenseMatrixD::Scalar)Tolerance)
+            {
+                Triplets.emplace_back(RowIdx, ColIdx, Value);
+            }
+        }
     }
 
-	return false;
+    Solution.resize(SolutionDense.rows(), SolutionDense.cols());
+    Solution.setFromTriplets(Triplets.begin(), Triplets.end());
+
+    return true;
 }
 
+
+// 
+// Helper "one-function call" methods for setting up and solving the QP problems
+//
+
 bool FQuadraticProgramming::SolveWithFixedConstraints(const FSparseMatrixD& MatrixQ, 
-                                                      const FColumnVectorD& VectorF, 
+                                                      const FColumnVectorD* VectorF, 
                                                       const TArray<int>& FixedRowIndices, 
                                                       const FSparseMatrixD& FixedValues, 
-                                                      FDenseMatrixD& Solution)
+                                                      FDenseMatrixD& Solution,
+                                                      const bool bVariablesOnly,
+                                                      TArray<int>* VariableRowIndices)
+{
+    FQuadraticProgramming QP(&MatrixQ, VectorF);
+    if (!QP.SetFixedConstraints(&FixedRowIndices, &FixedValues))
+    {
+        return false;
+    }
+    
+    if (!QP.PreFactorize())
+    {
+        return false;
+    }
+
+    *VariableRowIndices = QP.GetVariableRowIndices();
+    return QP.Solve(Solution, bVariablesOnly);
+}
+
+bool FQuadraticProgramming::SolveWithFixedConstraints(const FSparseMatrixD& MatrixQ,
+                                                      const FColumnVectorD* VectorF,
+                                                      const TArray<int>& FixedRowIndices,
+                                                      const FSparseMatrixD& FixedValues,
+                                                      FSparseMatrixD& Solution,
+                                                      const bool bVariablesOnly,
+                                                      const double Tolerance,
+                                                      TArray<int>* VariableRowIndices)
 {
-    FQuadraticProgramming QP(&MatrixQ, &VectorF);
+    FQuadraticProgramming QP(&MatrixQ, VectorF);
     if (!QP.SetFixedConstraints(&FixedRowIndices, &FixedValues))
     {
         return false;
@@ -189,5 +333,22 @@ bool FQuadraticProgramming::SolveWithFixedConstraints(const FSparseMatrixD& Matr
         return false;
     }
 
-    return QP.Solve(Solution);
+    *VariableRowIndices = QP.GetVariableRowIndices();
+
+    return QP.Solve(Solution, bVariablesOnly, Tolerance);
+}
+
+
+//
+// Solver stats
+//
+
+double FQuadraticProgramming::GetSolveTimeElapsedInSec() const 
+{ 
+    return SolveTimeElapsedInSec; 
+}
+
+TArray<int> FQuadraticProgramming::GetVariableRowIndices() const
+{
+    return VariableRowIndices;
 }
\ No newline at end of file
diff --git a/Engine/Plugins/Runtime/GeometryProcessing/Source/DynamicMesh/Private/Solvers/Internal/QuadraticProgramming.h b/Engine/Plugins/Runtime/GeometryProcessing/Source/DynamicMesh/Private/Solvers/Internal/QuadraticProgramming.h
--- a/Engine/Plugins/Runtime/GeometryProcessing/Source/DynamicMesh/Private/Solvers/Internal/QuadraticProgramming.h
+++ b/Engine/Plugins/Runtime/GeometryProcessing/Source/DynamicMesh/Private/Solvers/Internal/QuadraticProgramming.h
@@ -67,19 +67,47 @@ public:
 
     FQuadraticProgramming(const FSparseMatrixD* InMatrixQ, const FColumnVectorD* InVectorF = nullptr);
 
+    /** If true, will solve for each column x \in X in parallel. By default, parallelization is enabled. */
+    void SetEnableParallelization(const bool bEnable);
+
+    /** Set the fixed constraints either via a sparse or a dense matrix. Calls are mutually exclusive. */
     bool SetFixedConstraints(const TArray<int>* InFixedRowIndices, const FSparseMatrixD* InFixedValues);
+    bool SetFixedConstraints(const TArray<int>* InFixedRowIndices, const FDenseMatrixD* InFixedValues);
 
-    /**  Pre-factorizes the matrices and sets up the solver. */
+    /** Pre-factorizes the matrices and sets up the solver. */
     bool PreFactorize();
+    
+    /**
+     * @param Solution The result of the solve as a dense matrix.
+	 * @param bVariablesOnly If true, the Solution matrix only contains the rows for non-fixed (variable) rows (set via SetFixedConstraints function).
+     */
+    bool Solve(FDenseMatrixD& Solution, const bool bVariablesOnly = false);
 
-    //TODO: Allow to solve for sparse matrices with a passed float threshold value that determines if the value is close enough to 0 
-    bool Solve(FDenseMatrixD& Solution);
+    /**
+     * @param Solution The result of the solve as a sparse matrix.
+	 * @param bVariablesOnly If true, the Solution matrix only contains the rows for non-fixed (variable) rows (set via SetFixedConstraints function).
+     * @param Tolerance Convert dense solution to sparse by pruning any values with their absolute value below the Tolerance.
+     */
+    bool Solve(FSparseMatrixD& Solution, const bool bVariablesOnly = false, const double Tolerance = SMALL_NUMBER);
 
 
     // 
     // Helper "one-function call" methods for setting up and solving the QP problems
     //
-    static bool SolveWithFixedConstraints(const FSparseMatrixD& MatrixQ, const FColumnVectorD& VectorF, const TArray<int>& FixedRowIndices, const FSparseMatrixD& FixedValues, FDenseMatrixD& Solution);
+
+    static bool SolveWithFixedConstraints(const FSparseMatrixD& MatrixQ, const FColumnVectorD* VectorF, const TArray<int>& FixedRowIndices, const FSparseMatrixD& FixedValues, FDenseMatrixD& Solution, const bool bVariablesOnly = false, TArray<int>* VariableRowIndices = nullptr);
+    static bool SolveWithFixedConstraints(const FSparseMatrixD& MatrixQ, const FColumnVectorD* VectorF, const TArray<int>& FixedRowIndices, const FSparseMatrixD& FixedValues, FSparseMatrixD& Solution, const bool bVariablesOnly = false, const double Tolerance = SMALL_NUMBER, TArray<int>* VariableRowIndices = nullptr);
+
+
+    //
+    // Solver stats
+    // 
+
+    /** How long did the last call to Solve() take in seconds. */
+    double GetSolveTimeElapsedInSec() const;
+
+    /** Array of row indices into MatrixQ representing variables (non-fixed). */
+    TArray<int> GetVariableRowIndices() const;
 
 protected:
 
@@ -87,12 +115,23 @@ protected:
     const FColumnVectorD* VectorF = nullptr; // Column vector of linear coefficients
 
     // Fixed constraints
-    bool bFixedConstraintsSet = false;  // set to true if SetFixedConstraints is called and successful
-    const TArray<int>* FixedRowIndices = nullptr;  // row indices of the fixed parameters, set by the user when calling SetFixedConstraints
-    const FSparseMatrixD* FixedValues = nullptr;   // matrix of fixed values, set by the user when calling SetFixedConstraints
+    bool bFixedConstraintsSet = false;                  // set to true if SetFixedConstraints is called and successful
+    const TArray<int>* FixedRowIndices = nullptr;       // row indices of the fixed parameters, set by the user when calling SetFixedConstraints
+    const FSparseMatrixD* FixedValuesSparse = nullptr;  // matrix of fixed values, set by the user when calling SetFixedConstraints
+    const FDenseMatrixD* FixedValuesDense = nullptr;    // matrix of fixed values, set by the user when calling SetFixedConstraints
+
     TArray<int> VariableRowIndices;  // row indices of the variable parameters, computed internally by the call to PreFactorize()
 
     TUniquePtr<IMatrixSolverBase> Solver = nullptr; // PreFactorize() method pre-factorizes the matrices and setups the solver
+
+    bool bUseParallel = true; // if true when solving for X, will solve for each column x \in X in parallel
+
+
+    //
+    // Debug solver stats
+    //
+    
+    double SolveTimeElapsedInSec = 0.0; // Set every time Solve() is called
 };
 
 }
