-Added a funciton to MeshLocalParam to return an array of points with set UVs.
-Added a function to smooth weights for an array of vertices and all vertices within a geodesic distance away from them.

#rnx

diff --git a/Engine/Plugins/Runtime/GeometryProcessing/Source/DynamicMesh/Private/Operations/SmoothBoneWeights.cpp b/Engine/Plugins/Runtime/GeometryProcessing/Source/DynamicMesh/Private/Operations/SmoothBoneWeights.cpp
--- a/Engine/Plugins/Runtime/GeometryProcessing/Source/DynamicMesh/Private/Operations/SmoothBoneWeights.cpp
+++ b/Engine/Plugins/Runtime/GeometryProcessing/Source/DynamicMesh/Private/Operations/SmoothBoneWeights.cpp
@@ -8,6 +8,8 @@
 #include "Util/ProgressCancel.h"
 #include "BoneWeights.h"
 #include "Math/UnrealMathUtility.h"
+#include "Parameterization/MeshLocalParam.h"
+#include "DynamicMesh/MeshNormals.h"
 
 using namespace UE::AnimationCore;
 using namespace UE::Geometry;
@@ -260,6 +262,66 @@ bool FSmoothDynamicMeshVertexSkinWeights::SmoothWeightsAtVertex(const int32 Vert
 	return false;
 }
 
+bool FSmoothDynamicMeshVertexSkinWeights::SmoothWeightsAtVerticesWithinDistance(const TArray<int32>& Vertices, 
+																  				const float Strength, 
+																  				const double FloodFillUpToDistance)
+{
+	TSet<int32> VerticesToSmooth;
+	VerticesToSmooth.Append(Vertices);
+
+	const double FloodFillUpToDistanceSquared = FloodFillUpToDistance * FloodFillUpToDistance;
+
+	// We want to add vertices to the VerticesToSmooth within FloodFillUpToDistance away from each vertex in the 
+	// Vertices array
+	if (FloodFillUpToDistance > 0)
+	{
+		for (const int32 SeedVID : Vertices)
+		{	
+			if (Cancelled()) 
+			{
+				return false;
+			}
+
+			// If at least one neighboring vertex is not part of the set of vertices to smooth then we need to flood
+			bool bNeedToFlood = false;
+			for (const int32 NeighborVertexID : SourceMesh->VtxVerticesItr(SeedVID))
+			{
+				if (!VerticesToSmooth.Contains(NeighborVertexID) && 
+					DistanceSquared(SourceMesh->GetVertex(SeedVID), SourceMesh->GetVertex(NeighborVertexID)) < FloodFillUpToDistanceSquared)
+				{
+					bNeedToFlood = true;
+					break;
+				}
+			}
+
+			if (bNeedToFlood)
+			{
+				FVector3d Normal = FMeshNormals::ComputeVertexNormal(*SourceMesh, SeedVID);
+				FFrame3d SeedFrame = SourceMesh->GetVertexFrame(SeedVID, false, &Normal);
+
+				TMeshLocalParam<FDynamicMesh3> Param(SourceMesh);
+				Param.ParamMode = ELocalParamTypes::ExponentialMapUpwindAvg;
+				Param.ComputeToMaxDistance(SeedVID, SeedFrame, FloodFillUpToDistance);
+
+				// Only points within FloodFillUpToDistance should have UVs set
+				TArray<int32> PointsWithinDistance;
+				Param.GetPointsWithUV(PointsWithinDistance);
+
+				VerticesToSmooth.Append(PointsWithinDistance);
+			}
+		}
+	}
+
+	for (const int32 VertexID : VerticesToSmooth)
+	{
+		if (!FSmoothDynamicMeshVertexSkinWeights::SmoothWeightsAtVertex(VertexID, Strength))
+		{
+			return false;
+		}
+	}
+
+	return true;
+}
 
 
 // template instantiation
diff --git a/Engine/Plugins/Runtime/GeometryProcessing/Source/DynamicMesh/Public/Operations/SmoothBoneWeights.h b/Engine/Plugins/Runtime/GeometryProcessing/Source/DynamicMesh/Public/Operations/SmoothBoneWeights.h
--- a/Engine/Plugins/Runtime/GeometryProcessing/Source/DynamicMesh/Public/Operations/SmoothBoneWeights.h
+++ b/Engine/Plugins/Runtime/GeometryProcessing/Source/DynamicMesh/Public/Operations/SmoothBoneWeights.h
@@ -198,6 +198,19 @@ public:
 	 */
 	bool SmoothWeightsAtVertex(const int32 VertexID, const float VertexFalloff);
 
+	/**
+	 * Compute smoothed weights at vertices and neighbors within distance and write the result into 
+	 * FDynamicMeshVertexSkinWeightsAttribute passed upon construction.
+	 *
+	 * @param Vertices		  		Array of vertex IDs we are smoothing weights for.
+	 * @param Strength 				The lerp value between the old and the new weight.
+	 * @param FloodFillUpToDistance	For every vertex in Vertices, find vertices that are within FloodFillUpToDistance 
+	 * 								geodesic distance away and smooth them as well.
+	 *
+	 * @return true if the algorithm succeeds, false if it failed or was canceled by the user.
+	 */
+	bool SmoothWeightsAtVerticesWithinDistance(const TArray<int32>& Vertices, const float Strength, const double FloodFillUpToDistance);
+
 protected:
 	TUniquePtr<TBoneWeightsDataSource<FBoneIndexType, float>> SkinWeightsAttributeDataSource = nullptr;
 	FDynamicMeshVertexSkinWeightsAttribute* Attribute = nullptr;
diff --git a/Engine/Source/Runtime/GeometryCore/Public/Parameterization/MeshLocalParam.h b/Engine/Source/Runtime/GeometryCore/Public/Parameterization/MeshLocalParam.h
--- a/Engine/Source/Runtime/GeometryCore/Public/Parameterization/MeshLocalParam.h
+++ b/Engine/Source/Runtime/GeometryCore/Public/Parameterization/MeshLocalParam.h
@@ -164,6 +164,20 @@ public:
 		return (Node != nullptr && Node->bFrozen);
 	}
 
+	/**
+	 * @return Array of points with set UVs. 
+	 */
+	void GetPointsWithUV(TArray<int32>& Points) const 
+	{
+		Points.Empty(AllocatedNodes.Num());
+		for (const FGraphNode& Node : AllocatedNodes)
+		{
+			if (Node.bFrozen)
+			{
+				Points.Add(Node.PointID);
+			}
+		}
+	}
 
 	/**
 	 * @return the UV calculated for index PointID
