-Add optional smoothing of weights after transfer.
-Add support for multulayerd meshes.
-Add an option to compute the laplacian on intrisic mesh.
-Add a mask to force computation of inpainted weights.
-Add the laplacian energy term.
-Improved peformance by getting the sparse matrix from the solve and then iterating over all (most zero) values.

#rnx

diff --git a/Engine/Plugins/Runtime/GeometryProcessing/Source/DynamicMesh/Private/Operations/SmoothBoneWeights.cpp b/Engine/Plugins/Runtime/GeometryProcessing/Source/DynamicMesh/Private/Operations/SmoothBoneWeights.cpp
--- a/Engine/Plugins/Runtime/GeometryProcessing/Source/DynamicMesh/Private/Operations/SmoothBoneWeights.cpp
+++ b/Engine/Plugins/Runtime/GeometryProcessing/Source/DynamicMesh/Private/Operations/SmoothBoneWeights.cpp
@@ -264,7 +264,8 @@ bool FSmoothDynamicMeshVertexSkinWeights::SmoothWeightsAtVertex(const int32 Vert
 
 bool FSmoothDynamicMeshVertexSkinWeights::SmoothWeightsAtVerticesWithinDistance(const TArray<int32>& Vertices, 
 																  				const float Strength, 
-																  				const double FloodFillUpToDistance)
+																  				const double FloodFillUpToDistance,
+																				const int32 NumIterations)
 {
 	TSet<int32> VerticesToSmooth;
 	VerticesToSmooth.Append(Vertices);
@@ -312,11 +313,14 @@ bool FSmoothDynamicMeshVertexSkinWeights::SmoothWeightsAtVerticesWithinDistance(
 		}
 	}
 
-	for (const int32 VertexID : VerticesToSmooth)
+	for (int32 Itr = 0; Itr < NumIterations; ++Itr)
 	{
-		if (!FSmoothDynamicMeshVertexSkinWeights::SmoothWeightsAtVertex(VertexID, Strength))
+		for (const int32 VertexID : VerticesToSmooth)
 		{
-			return false;
+			if (!FSmoothDynamicMeshVertexSkinWeights::SmoothWeightsAtVertex(VertexID, Strength))
+			{
+				return false;
+			}
 		}
 	}
 
diff --git a/Engine/Plugins/Runtime/GeometryProcessing/Source/DynamicMesh/Private/Operations/TransferBoneWeights.cpp b/Engine/Plugins/Runtime/GeometryProcessing/Source/DynamicMesh/Private/Operations/TransferBoneWeights.cpp
--- a/Engine/Plugins/Runtime/GeometryProcessing/Source/DynamicMesh/Private/Operations/TransferBoneWeights.cpp
+++ b/Engine/Plugins/Runtime/GeometryProcessing/Source/DynamicMesh/Private/Operations/TransferBoneWeights.cpp
@@ -14,6 +14,7 @@
 #include "TransformTypes.h"
 #include "Solvers/Internal/QuadraticProgramming.h"
 #include "Solvers/LaplacianMatrixAssembly.h"
+#include "Operations/SmoothBoneWeights.h"
 
 using namespace UE::AnimationCore;
 using namespace UE::Geometry;
@@ -57,7 +58,7 @@ namespace TransferBoneWeightsLocals
 
 			FBoneWeights MappedWeights;
 
-			for (int WeightIdx = 0; WeightIdx < OutWeights.Num(); ++WeightIdx)
+			for (int32 WeightIdx = 0; WeightIdx < OutWeights.Num(); ++WeightIdx)
 			{
 				const FBoneWeight& BoneWeight = OutWeights[WeightIdx];
 
@@ -209,7 +210,7 @@ bool FTransferBoneWeights::TransferWeightsToMesh(FDynamicMesh3& InOutTargetMesh,
 			TargetBoneToIndex = MakeUnique<TMap<FName, uint16>>();
 			TargetBoneToIndex->Reserve(TargetBoneNames.Num());
 
-			for (int BoneID = 0; BoneID < TargetBoneNames.Num(); ++BoneID)
+			for (int32 BoneID = 0; BoneID < TargetBoneNames.Num(); ++BoneID)
 			{
 				const FName& BoneName = TargetBoneNames[BoneID];
 				if (TargetBoneToIndex->Contains(BoneName))
@@ -254,7 +255,7 @@ bool FTransferBoneWeights::TransferWeightsToMesh(FDynamicMesh3& InOutTargetMesh,
 		// must be equal to the target mesh vertex count
 		if (SearchRadius < 0 && NormalThreshold < 0)
 		{
-			int NumMatched = 0;
+			int32 NumMatched = 0;
 			for (bool Flag : MatchedVertices)
 			{
 				if (Flag) 
@@ -280,7 +281,7 @@ bool FTransferBoneWeights::TransferWeightsToMesh(FDynamicMesh3& InOutTargetMesh,
          *      trace(W^t Q W)
          *      W \in R^(nxm) is a matrix where n is the number of vertices and m is the number of bones. So (i,j) entry is 
          *                    the influence (weight) of a vertex i by bone j
-         *      Q \in R^(nxn) is a matrix that combines both Dirichlet and Laplacian energies, Q = -0.5L + 0.5LM^-1L
+         *      Q \in R^(nxn) is a matrix that combines both Dirichlet and Laplacian energies, Q = -L + L*M^(-1)*L
          *                    where L is a cotangent Laplacian and M is a mass matrix
          *  
          *  subject to constraints
@@ -300,6 +301,12 @@ bool FTransferBoneWeights::TransferWeightsToMesh(FDynamicMesh3& InOutTargetMesh,
 			
 			if (InOutTargetMesh.IsVertex(VertexID)) 
 			{
+				// check if we need to force the vertex to not have a match
+				if (ForceInpaint.Num() == InOutTargetMesh.MaxVertexID() && ForceInpaint[VertexID] > 0)
+				{
+					return;
+				}
+
 				const FVector3d Point = InOutTargetMesh.GetVertex(VertexID);
 				const FVector3f Normal = NormalThreshold >= 0 ? InOutTargetMesh.Attributes()->PrimaryNormals()->GetElement(VertexID) : FVector3f::Zero();
 
@@ -338,7 +345,7 @@ bool FTransferBoneWeights::TransferWeightsToMesh(FDynamicMesh3& InOutTargetMesh,
 		}
 
 		// Setup the sparse matrix FixedValues of known (matched) weight values and the array (FixedIndices) of the matched vertex IDs
-		const int TargetNumBones = InOutTargetMesh.Attributes()->GetBoneNames()->GetAttribValues().Num();
+		const int32 TargetNumBones = InOutTargetMesh.Attributes()->GetBoneNames()->GetAttribValues().Num();
 		FSparseMatrixD FixedValues;
 		FixedValues.resize(NumMatched, TargetNumBones);
 		std::vector<Eigen::Triplet<FSparseMatrixD::Scalar>> FixedValuesTriplets;
@@ -347,20 +354,20 @@ bool FTransferBoneWeights::TransferWeightsToMesh(FDynamicMesh3& InOutTargetMesh,
 		TArray<int> FixedIndices;
 		FixedIndices.Reserve(NumMatched);
 
-		for (int VertexID = 0; VertexID < InOutTargetMesh.MaxVertexID(); ++VertexID)
+		for (int32 VertexID = 0; VertexID < InOutTargetMesh.MaxVertexID(); ++VertexID)
 		{
 			if (InOutTargetMesh.IsVertex(VertexID) && MatchedVertices[VertexID])
 			{
 				FBoneWeights Data;
 				TargetSkinWeights->GetValue(VertexID, Data);
 
-				const int NumBones = Data.Num();
+				const int32 NumBones = Data.Num();
 				checkSlow(NumBones > 0);
-				const int CurIdx = FixedIndices.Num();
-				for (int BoneID = 0; BoneID < NumBones; ++BoneID)
+				const int32 CurIdx = FixedIndices.Num();
+				for (int32 BoneID = 0; BoneID < NumBones; ++BoneID)
 				{
 					const int BoneIdx = Data[BoneID].GetBoneIndex();
-					const double BoneWeight = (double)Data[BoneID].GetWeight();
+					const double BoneWeight = Data[BoneID].GetWeight();
 					FixedValuesTriplets.emplace_back(CurIdx, BoneIdx, BoneWeight);
 				}
 
@@ -370,51 +377,97 @@ bool FTransferBoneWeights::TransferWeightsToMesh(FDynamicMesh3& InOutTargetMesh,
 		}
 		FixedValues.setFromTriplets(FixedValuesTriplets.begin(), FixedValuesTriplets.end());
 
-		// Setup the Laplacian matrix
 		const int32 NumVerts = VtxLinearization.NumVerts();
+		FEigenSparseMatrixAssembler CotangentAssembler(NumVerts, NumVerts);
 		FEigenSparseMatrixAssembler LaplacianAssembler(NumVerts, NumVerts);
-		UE::MeshDeformation::ConstructFullCotangentLaplacian<double>(InOutTargetMesh, 
-																	 VtxLinearization, 
-																	 LaplacianAssembler,
-																	 UE::MeshDeformation::ECotangentWeightMode::Default,
-																	 UE::MeshDeformation::ECotangentAreaMode::NoArea);
-		FSparseMatrixD CotangentMatrix;
-		LaplacianAssembler.ExtractResult(CotangentMatrix);
-		FSparseMatrixD NegativeCotangentMatrix = -0.25 * CotangentMatrix; //TODO: Add LM^-1L term in the future
+
+		if (bUseIntrinsicLaplacian)
+		{
+			// Construct the Cotangent weights matrix
+			UE::MeshDeformation::ConstructFullIDTCotangentLaplacian<double>(InOutTargetMesh, VtxLinearization, CotangentAssembler,
+																			UE::MeshDeformation::ECotangentWeightMode::Default, 
+																			UE::MeshDeformation::ECotangentAreaMode::NoArea);
+
+			// Construct the Laplacian with cotangent weights scaled by the voronoi area (i.e. M^(-1)*L matrix where M is the mass/stiffness matrix)
+			UE::MeshDeformation::ConstructFullIDTCotangentLaplacian<double>(InOutTargetMesh, VtxLinearization, LaplacianAssembler,
+																			UE::MeshDeformation::ECotangentWeightMode::Default,
+																			UE::MeshDeformation::ECotangentAreaMode::VoronoiArea);
+		}
+		else 
+		{
+			UE::MeshDeformation::ConstructFullCotangentLaplacian<double>(InOutTargetMesh, VtxLinearization, CotangentAssembler,
+																		 UE::MeshDeformation::ECotangentWeightMode::Default,
+																		 UE::MeshDeformation::ECotangentAreaMode::NoArea);
+
+			UE::MeshDeformation::ConstructFullCotangentLaplacian<double>(InOutTargetMesh, VtxLinearization, LaplacianAssembler,
+																		 UE::MeshDeformation::ECotangentWeightMode::Default,
+																		 UE::MeshDeformation::ECotangentAreaMode::VoronoiArea);
+		}
+
+		FSparseMatrixD CotangentMatrix, MassCotangentMatrix;
+		CotangentAssembler.ExtractResult(CotangentMatrix);
+		LaplacianAssembler.ExtractResult(MassCotangentMatrix);
+
+		// -L * L* M^(-1)*L energy
+		FSparseMatrixD Energy = -1*CotangentMatrix + CotangentMatrix*MassCotangentMatrix;
 
 		// Solve the QP problem with fixed constraints
-		bFailed = true;
-		FQuadraticProgramming QProgram(&NegativeCotangentMatrix);
-		if (ensure(QProgram.SetFixedConstraints(&FixedIndices, &FixedValues)))
+		FSparseMatrixD TargetWeights;
+		TArray<int> VaribleRows;
+
+		// We want the solution TargetWeights matrix to only contain the rows representing the variable (non-fixed) rows
+		constexpr bool bVariablesOnly = true;
+ 		bFailed = !FQuadraticProgramming::SolveWithFixedConstraints(Energy, nullptr, FixedIndices, FixedValues, TargetWeights, bVariablesOnly, KINDA_SMALL_NUMBER, &VaribleRows);
+		checkSlow((VaribleRows.Num() + FixedIndices.Num()) == Energy.rows());
+
+		if (!bFailed)
 		{
-			if (ensure(QProgram.PreFactorize()))
+			// Transpose so we can efficiently iterate over the col-major matrix. Each column now contains per-vertex weights.
+			// Otherwise, we are iterating over rows of a col-major matrix which is slow.
+			FSparseMatrixD TargetWeightsTransposed = TargetWeights.transpose(); 
+
+			FBoneWeightsSettings BoneSettings;
+			BoneSettings.SetNormalizeType(EBoneWeightNormalizeType::None);
+			
+			// Iterate over every column containing all bone weights for the vertex
+			for (int32 ColIdx = 0; ColIdx < TargetWeightsTransposed.outerSize(); ++ColIdx)
 			{
-				FDenseMatrixD TargetWeights;
-				if (ensure(QProgram.Solve(TargetWeights)))
+				FBoneWeights WeightArray;
+
+				// Iterate over only non-zero rows (i.e. bone indices with non-zero weights)
+				for (FSparseMatrixD::InnerIterator Itr(TargetWeightsTransposed, ColIdx); Itr; ++Itr)
 				{
-					FBoneWeightsSettings BoneSettings;
-					BoneSettings.SetNormalizeType(EBoneWeightNormalizeType::None);
-					for (int IdxI = 0; IdxI < TargetWeights.rows(); ++IdxI)
+					const FBoneIndexType BoneIdx = static_cast<FBoneIndexType>(Itr.row());
+					const float Weight = static_cast<float>(Itr.value());
+					FBoneWeight Bweight(BoneIdx, Weight);
+					WeightArray.SetBoneWeight(Bweight, BoneSettings);
+				}
+
+				WeightArray.Renormalize(FBoneWeightsSettings());
+
+				const int32 VertexIDLinearalized = bVariablesOnly ? static_cast<int32>(VaribleRows[ColIdx]) : ColIdx; // linearized vertex ID (matrix row) of the variable in the Energy matrix
+				const int32 VertexID = static_cast<int32>(ToMeshV[VertexIDLinearalized]);
+				TargetSkinWeights->SetValue(VertexID, WeightArray);
+			}
+
+			// Optional post-processing smoothing of the weights at the vertices without a match
+			if (NumSmoothingIterations > 0 && SmoothingStrength > 0)
+			{
+				TArray<int32> VerticesToSmooth;
+				const int32 NumNotMatched = InOutTargetMesh.VertexCount() - NumMatched;
+				VerticesToSmooth.Reserve(NumNotMatched);
+				for (int32 VertexID = 0; VertexID < InOutTargetMesh.MaxVertexID(); ++VertexID)
+				{
+					if (InOutTargetMesh.IsVertex(VertexID) && !MatchedVertices[VertexID])
 					{
-						FBoneWeights WeightArray;
-						for (int IdxJ = 0; IdxJ < TargetWeights.cols(); ++IdxJ)
-						{
-							const FBoneIndexType BoneId = static_cast<FBoneIndexType>(IdxJ);
-							const float Weight = (float)TargetWeights(IdxI, IdxJ);
-							if (Weight > KINDA_SMALL_NUMBER)
-							{
-								FBoneWeight Bweight(BoneId, Weight);
-								WeightArray.SetBoneWeight(Bweight, BoneSettings);
-							}
-						}
-
-						WeightArray.Renormalize(FBoneWeightsSettings());
-
-						checkSlow(IdxI < ToMeshV.Num());
-						TargetSkinWeights->SetValue(ToMeshV[IdxI], WeightArray);
+						VerticesToSmooth.Add(VertexID);
 					}
+				}
 
-					bFailed = false;
+				FSmoothDynamicMeshVertexSkinWeights SmoothWeights(&InOutTargetMesh, InTargetProfileName);
+				if (ensure(SmoothWeights.Validate() == EOperationValidationResult::Ok))
+				{
+					ensure(SmoothWeights.SmoothWeightsAtVerticesWithinDistance(VerticesToSmooth, SmoothingStrength, SearchRadius, NumSmoothingIterations));
 				}
 			}
 		}
@@ -478,6 +531,12 @@ bool FTransferBoneWeights::TransferWeightsToPoint(UE::AnimationCore::FBoneWeight
 			const FVector3f InNormalNormalized = Normalized(InNormal);
 			const float NormalAngle = FMathf::ACos(InNormalNormalized.Dot(MatchedNormal));
 			bPassedNormalsCheck = (double)NormalAngle <= NormalThreshold;
+
+			if (!bPassedNormalsCheck && LayeredMeshSupport)
+			{
+				// try again with a flipped normal
+				bPassedNormalsCheck = (double)(TMathUtil<float>::Pi - NormalAngle) <= NormalThreshold;
+			}
 		}
 		
 		if (bPassedRadiusCheck && bPassedNormalsCheck)
diff --git a/Engine/Plugins/Runtime/GeometryProcessing/Source/DynamicMesh/Public/Operations/SmoothBoneWeights.h b/Engine/Plugins/Runtime/GeometryProcessing/Source/DynamicMesh/Public/Operations/SmoothBoneWeights.h
--- a/Engine/Plugins/Runtime/GeometryProcessing/Source/DynamicMesh/Public/Operations/SmoothBoneWeights.h
+++ b/Engine/Plugins/Runtime/GeometryProcessing/Source/DynamicMesh/Public/Operations/SmoothBoneWeights.h
@@ -206,10 +206,11 @@ public:
 	 * @param Strength 				The lerp value between the old and the new weight.
 	 * @param FloodFillUpToDistance	For every vertex in Vertices, find vertices that are within FloodFillUpToDistance 
 	 * 								geodesic distance away and smooth them as well.
+	 * @param NumIterations			Number of times the smoothing will be run
 	 *
 	 * @return true if the algorithm succeeds, false if it failed or was canceled by the user.
 	 */
-	bool SmoothWeightsAtVerticesWithinDistance(const TArray<int32>& Vertices, const float Strength, const double FloodFillUpToDistance);
+	bool SmoothWeightsAtVerticesWithinDistance(const TArray<int32>& Vertices, const float Strength, const double FloodFillUpToDistance, const int32 NumIterations);
 
 protected:
 	TUniquePtr<TBoneWeightsDataSource<FBoneIndexType, float>> SkinWeightsAttributeDataSource = nullptr;
diff --git a/Engine/Plugins/Runtime/GeometryProcessing/Source/DynamicMesh/Public/Operations/TransferBoneWeights.h b/Engine/Plugins/Runtime/GeometryProcessing/Source/DynamicMesh/Public/Operations/TransferBoneWeights.h
--- a/Engine/Plugins/Runtime/GeometryProcessing/Source/DynamicMesh/Public/Operations/TransferBoneWeights.h
+++ b/Engine/Plugins/Runtime/GeometryProcessing/Source/DynamicMesh/Public/Operations/TransferBoneWeights.h
@@ -116,6 +116,16 @@ public:
 	/** Transform applied to the input target mesh or target point before transfer. */
 	FTransformSRT3d TargetToWorld = FTransformSRT3d::Identity(); 
 
+	/** 
+     * Completely ignore the source and target mesh bone attributes when transferring weights from one dynamic mesh to another.
+     * This skips re-indexing and simply copies skin weights over. Use with caution.
+	 */
+	bool bIgnoreBoneAttributes = false;
+
+	//
+	// Optional Inputs for ETransferBoneWeightsMethod::InpaintWeights method
+	//
+
 	/**  Radius for searching the closest point. If negative, all points are considered. */
 	double SearchRadius = -1;
 
@@ -124,18 +134,38 @@ public:
 	 * If negative, normals are ignored.
 	 */
 	double NormalThreshold = -1;
+	
+	/** 
+	 * If true, when the closest point doesn't pass the normal threshold test, will try again with a flipped normal. 
+	 * This helps with layered meshes where the "inner" and "outer" layers are close to each other but whose normals 
+	 * are pointing in the opposite directions. 
+	 */
+	bool LayeredMeshSupport = false;
+
+	/** The number of optional post-processing smoothing iterations applied to the vertices without the match. */
+	int32 NumSmoothingIterations = 0; 
+
+	/** The strength of each post-processing smoothing iteration. */
+	float SmoothingStrength = 0.0f;
+
+	/** If true, will use the intrinsic Delaunay mesh to construct sparse Cotangent Laplacian matrix. */
+	bool bUseIntrinsicLaplacian = false;
 
 	/** 
-     * Completely ignore the source and target mesh bone attributes when transferring weights from one dynamic mesh to another.
-     * This skips re-indexing and simply copies skin weights over. Use with caution.
+	 * Optional mask where if ForceInpaint[VertexID] > 0 we want to force the weights for the vertex to be computed  
+	 * automatically.
+	 * 
+	 * @note Only used when TransferMethod == ETransferBoneWeightsMethod::InpaintWeights.
+	 * 		 The size must be equal to the InSourceMesh.MaxVertexID(), otherwise the mask is ignored.
 	 */
-	bool bIgnoreBoneAttributes = false;
+	TArray<float> ForceInpaint;
+
 
 	//
 	// Outputs
 	//
 
-	/** MatchedVertices[VertexID] is set to true for a tartet mesh vertex ID with a match found, false otherwise. */
+	/** MatchedVertices[VertexID] is set to true for a target mesh vertex ID with a match found, false otherwise. */
 	TArray<bool> MatchedVertices;
 
 protected:
