Added an operation to compute per-vertex averaged value of all overlay elements.

diff --git a/Engine/Plugins/Runtime/GeometryProcessing/Source/DynamicMesh/Private/Operations/AverageOverlayToVertices.cpp b/Engine/Plugins/Runtime/GeometryProcessing/Source/DynamicMesh/Private/Operations/AverageOverlayToVertices.cpp
new file mode 100644
--- /dev/null
+++ b/Engine/Plugins/Runtime/GeometryProcessing/Source/DynamicMesh/Private/Operations/AverageOverlayToVertices.cpp
@@ -0,0 +1,216 @@
+// Copyright Epic Games, Inc. All Rights Reserved.
+
+#include "Operations/AverageOverlayToVertices.h"
+#include "DynamicMesh/DynamicMesh3.h"
+#include "DynamicMesh/DynamicMeshOverlay.h"
+#include "Async/ParallelFor.h"
+#include "MeshQueries.h"
+#include "Math/UnrealMathUtility.h"
+
+using namespace UE::Geometry;
+
+namespace FAverageOverlayToVerticesLocals
+{
+    /** 
+     * Compute the weighted sum of elements values for all the mesh vertices.
+     * 
+     * @param Overlay overlay containing the elements to be averaged.
+     * @param bWeightByArea if true, include weighting by the area of the triangle.
+     * @param bWeightByAngle if true, include weighting by the interior angles of the triangle.
+     * @param SummedElementValues Mesh->MaxVertexID()*ElementSize sized array of the weighted sum of the element values 
+     *                            at a vertex.
+     * @param HasElements Mesh->MaxVertexID() sized array where HasElements[VertexID] is true if VertexID has at 
+     *                    least one element (i.e, at least one triangle sharing the vertex is set in the overlay).
+     * @param VertexTotalWeights sum of all element weights at a vertex.
+     */
+    template<typename RealType, int ElementSize>
+    void AverageFullMesh(const FDynamicMesh3* Mesh,
+                         const TDynamicMeshOverlay<RealType, ElementSize>* Overlay,
+                         bool bWeightByArea, 
+                         bool bWeightByAngle,
+                         TArray<RealType>& SummedElementValues,
+                         TArray<bool>& HasElements,
+                         TArray<double>& VertexTotalWeights)
+    {
+        SummedElementValues.Init((RealType)0.0, Mesh->MaxVertexID()*ElementSize);
+        HasElements.Init(false, Mesh->MaxVertexID());
+        VertexTotalWeights.Init(0.0, Mesh->MaxVertexID());
+
+        // reuse buffer
+        RealType ElementData[ElementSize];
+
+        // Iterate over every triangle and accumulate the overlay element values for each triangle vertex
+        for (const int32 TID : Mesh->TriangleIndicesItr())
+        {
+            if (Overlay->IsSetTriangle(TID))
+            {
+                const FIndex3i TriVert = Mesh->GetTriangle(TID);
+                const FIndex3i TriElem = Overlay->GetTriangle(TID);
+                
+                // per-corner weights
+                const FVector3d Weights = TMeshQueries<FDynamicMesh3>::GetVertexWeightsOnTriangle(*Mesh, TID, Mesh->GetTriArea(TID), bWeightByArea, bWeightByAngle);
+
+                // iterate over every corner of the triangle
+                for (int32 CornerIdx = 0; CornerIdx < 3; ++CornerIdx)
+                {	
+                    const int32 CornerVID = TriVert[CornerIdx];
+                    Overlay->GetElement(TriElem[CornerIdx], ElementData);
+                
+                    const int32 StartIdx = CornerVID * ElementSize;      // starting offset into the output averaged values array
+                    VertexTotalWeights[CornerVID] += Weights[CornerIdx]; // accumulate weights for the vertex to be used later for the normalization
+
+                    // add the weighted contribution of the element for the vertex
+                    for (int32 EIdx = 0; EIdx < ElementSize; ++EIdx)
+                    {
+                        SummedElementValues[StartIdx + EIdx] += RealType(Weights[CornerIdx] * (double)ElementData[EIdx]);
+                    }
+
+                    HasElements[CornerVID] = true;
+                }
+            }
+        }
+    }
+
+    /**
+     * Compute the weighed sum of element values for the selection of the mesh vertices.
+     * 
+     * @param Overlay overlay containing the elements to be averaged.
+     * @param SelectedVertices the array of vertex IDs representing the selection.
+     * @param bWeightByArea if true, include weighting by the area of the triangle.
+     * @param bWeightByAngle if true, include weighting by the interior angles of the triangle.
+     * @param SummedElementValues Mesh->MaxVertexID()*ElementSize sized array of the weighted sum of the element values 
+     *                            at a vertex.
+     * @param HasElements Mesh->MaxVertexID() sized array where HasElements[VertexID] is true if VertexID has at 
+     *                    least one element (i.e, at least one triangle sharing the vertex is set in the overlay).
+     * @param VertexTotalWeights sum of all element weights at a vertex.
+     */
+    template<typename RealType, int ElementSize>
+    void AverageSelection(const FDynamicMesh3* Mesh,
+                          const TDynamicMeshOverlay<RealType, ElementSize>* Overlay,
+                          const TArray<int32>& SelectedVertices,
+                          bool bWeightByArea, 
+                          bool bWeightByAngle,
+                          TArray<RealType>& SummedElementValues,
+                          TArray<bool>& HasElements,
+                          TArray<double>& VertexTotalWeights)
+    {
+        // Map mesh vertex ID to selection idx
+        TMap<int32, int32> VIDToSelected;
+        int32 Count = 0;
+        VIDToSelected.Reserve(SelectedVertices.Num());
+        for (const int32 VID : SelectedVertices)
+        {
+            VIDToSelected.Add(VID, Count++);
+        }
+
+        // Get selected triangles
+        TArray<int32> SelectedTriangles;
+		SelectedTriangles = TMeshQueries<FDynamicMesh3>::GetVertexSelectedTriangles(*Mesh, SelectedVertices);
+        
+        SummedElementValues.Init((RealType)0.0, SelectedVertices.Num()*ElementSize);
+        HasElements.Init(false, SelectedVertices.Num());
+        VertexTotalWeights.Init(0.0, SelectedVertices.Num());
+
+        // reuse buffer
+        RealType ElementData[ElementSize];
+
+        // Iterate over every triangle and accumulate the overlay element values for each triangle vertex
+        for (const int32 TID : SelectedTriangles)
+        {
+            if (Mesh->IsTriangle(TID) && Overlay->IsSetTriangle(TID))
+            {
+                const FIndex3i TriVert = Mesh->GetTriangle(TID);
+                const FIndex3i TriElem = Overlay->GetTriangle(TID);
+                
+                // per-corner weights
+                const FVector3d Weights = TMeshQueries<FDynamicMesh3>::GetVertexWeightsOnTriangle(*Mesh, TID, Mesh->GetTriArea(TID), bWeightByArea, bWeightByAngle);
+
+                // iterate over every corner of the triangle
+                for (int32 CornerIdx = 0; CornerIdx < 3; ++CornerIdx)
+                {	
+                    const int32 CornerVID = TriVert[CornerIdx];
+                    if (VIDToSelected.Contains(CornerVID)) // check if corner is one of the selected vertices
+                    {
+                        Overlay->GetElement(TriElem[CornerIdx], ElementData);
+                    
+                        const int32 SelectedIdx = VIDToSelected[CornerVID];    // remap vertex ID to the index into output averaged values array
+                        const int32 StartIdx = SelectedIdx * ElementSize;      // starting offset into the output averaged values array
+                        VertexTotalWeights[SelectedIdx] += Weights[CornerIdx]; // accumulate weights for the vertex to be used later for the normalization
+
+                        // add the weighted contribution of the element for the vertex
+                        for (int32 EIdx = 0; EIdx < ElementSize; ++EIdx)
+                        {
+                            SummedElementValues[StartIdx + EIdx] += RealType(Weights[CornerIdx] * (double)ElementData[EIdx]);
+                        }
+
+                        HasElements[SelectedIdx] = true;
+                    }
+                }
+            }
+        }
+    }
+}
+
+FAverageOverlayToVertices::FAverageOverlayToVertices(const FDynamicMesh3& InMesh)
+:
+Mesh(&InMesh)
+{
+}
+
+template<typename RealType, int ElementSize>
+bool FAverageOverlayToVertices::AverageOverlay(const TDynamicMeshOverlay<RealType, ElementSize>* Overlay,
+                                               TArray<RealType>& VertexValues,
+                                               TArray<bool>& HasElements)
+{
+	using namespace FAverageOverlayToVerticesLocals;
+
+    if (Overlay == nullptr)
+    {
+        return false;
+    }
+
+    const bool bWeightByArea = TriangleWeight == ETriangleWeight::Area || TriangleWeight == ETriangleWeight::AngleArea;
+    const bool bWeightByAngle = TriangleWeight == ETriangleWeight::Angle || TriangleWeight == ETriangleWeight::AngleArea;
+    
+    TArray<double> VertexTotalWeights;
+    if (Selection.Num())
+    {
+		AverageSelection(Mesh, Overlay, Selection, bWeightByArea, bWeightByAngle, VertexValues, HasElements, VertexTotalWeights);
+    }
+    else 
+    {
+		AverageFullMesh(Mesh, Overlay, bWeightByArea, bWeightByAngle, VertexValues, HasElements, VertexTotalWeights);
+    }
+
+    //
+    // Normalize the accumulated values
+    //
+    ParallelFor(HasElements.Num(), [&VertexValues, &HasElements, &VertexTotalWeights](int32 SelectedIdx)
+    {
+        if (HasElements[SelectedIdx])
+        {
+            double InvTotalWeight = 1.0;
+            if (!FMath::IsNearlyZero(VertexTotalWeights[SelectedIdx]))
+            {
+                InvTotalWeight = 1.0/VertexTotalWeights[SelectedIdx];
+            }
+
+            const int32 StartIdx = SelectedIdx*ElementSize;
+            for (int32 EIdx = 0; EIdx < ElementSize; ++EIdx)
+            {
+                VertexValues[StartIdx + EIdx] = RealType(InvTotalWeight * (double)VertexValues[StartIdx + EIdx]);
+            }
+        }
+    }, bUseParallel ? EParallelForFlags::None : EParallelForFlags::ForceSingleThread);
+
+    
+    return true;
+}
+
+
+// template instantiation
+template DYNAMICMESH_API bool UE::Geometry::FAverageOverlayToVertices::AverageOverlay<float,4>(class  UE::Geometry::TDynamicMeshOverlay<float,4>  const *,class TArray<float, class TSizedDefaultAllocator<32> > &,class TArray<bool,class TSizedDefaultAllocator<32> > &);
+template DYNAMICMESH_API bool UE::Geometry::FAverageOverlayToVertices::AverageOverlay<float,3>(class  UE::Geometry::TDynamicMeshOverlay<float,3>  const *,class TArray<float, class TSizedDefaultAllocator<32> > &,class TArray<bool,class TSizedDefaultAllocator<32> > &);
+template DYNAMICMESH_API bool UE::Geometry::FAverageOverlayToVertices::AverageOverlay<float,2>(class  UE::Geometry::TDynamicMeshOverlay<float,2>  const *,class TArray<float, class TSizedDefaultAllocator<32> > &,class TArray<bool,class TSizedDefaultAllocator<32> > &);
+template DYNAMICMESH_API bool UE::Geometry::FAverageOverlayToVertices::AverageOverlay<float,1>(class  UE::Geometry::TDynamicMeshOverlay<float,1>  const *,class TArray<float, class TSizedDefaultAllocator<32> > &,class TArray<bool,class TSizedDefaultAllocator<32> > &);
+template DYNAMICMESH_API bool UE::Geometry::FAverageOverlayToVertices::AverageOverlay<double,3>(class UE::Geometry::TDynamicMeshOverlay<double,3> const *,class TArray<double,class TSizedDefaultAllocator<32> > &,class TArray<bool,class TSizedDefaultAllocator<32> > &);
\ No newline at end of file
diff --git a/Engine/Plugins/Runtime/GeometryProcessing/Source/DynamicMesh/Public/Operations/AverageOverlayToVertices.h b/Engine/Plugins/Runtime/GeometryProcessing/Source/DynamicMesh/Public/Operations/AverageOverlayToVertices.h
new file mode 100644
--- /dev/null
+++ b/Engine/Plugins/Runtime/GeometryProcessing/Source/DynamicMesh/Public/Operations/AverageOverlayToVertices.h
@@ -0,0 +1,76 @@
+// Copyright Epic Games, Inc. All Rights Reserved.
+
+#pragma once
+
+#include "CoreTypes.h"
+#include "Containers/Array.h"
+
+// Forward declarations
+namespace UE::Geometry 
+{ 
+	class FDynamicMesh3; 
+    template<typename RealType, int ElementSize>
+    class TDynamicMeshOverlay; 
+}
+
+namespace UE
+{
+namespace Geometry
+{
+
+/** 
+ * Compute the per-vertex values which are the weighted average of all vertex element values stored in the overlay. 
+ */
+class DYNAMICMESH_API FAverageOverlayToVertices
+{
+public:
+
+	//
+	// Optional Inputs
+	//
+	
+    /** Enable/disable multi-threading. */
+    bool bUseParallel = true;
+
+    /** Array of vertex IDs whose vertices will be affected. */
+    TArray<int32> Selection;
+        
+    /** The type of weight used for each element. */
+    enum class ETriangleWeight : uint8
+    {
+        Uniform,      // all elements are weighted equally
+        Angle = 0,    // use the angle between the edges that share the element
+        Area = 1,     // use the area of the triangle containing the element
+        AngleArea = 2 // use the angle weighted by the area
+    };
+    ETriangleWeight TriangleWeight = ETriangleWeight::AngleArea;
+
+protected:
+		
+	const FDynamicMesh3* Mesh;
+
+public:
+
+	FAverageOverlayToVertices(const FDynamicMesh3& InMesh);
+
+	virtual ~FAverageOverlayToVertices() = default;
+
+    /** 
+     * Compute the per-vertex values which are the weighted average of all vertex element values stored in the overlay. 
+     * 
+     * @param Overlay      Overlay containing the element values to be averaged.
+     * @param VertexValues Array of per-vertex averaged element values. The size is either equal to the 
+     *                     ElementSize*Mesh.MaxVertexID() or ElementSize*Selection.Num() if Selection is not empty. 
+     * @param HasElements  Array of per-vertex flags, indicating if the vertex had at least one element set in the  
+     *                     overlay and contains correct averaged value.
+     * 
+     * @return true on success, false otherwise
+     */
+    template<typename RealType, int ElementSize>
+	bool AverageOverlay(const TDynamicMeshOverlay<RealType, ElementSize>* Overlay,
+                        TArray<RealType>& VertexValues,
+                        TArray<bool>& HasElements);
+};
+
+}
+}
\ No newline at end of file
diff --git a/Engine/Source/Runtime/GeometryCore/Private/DynamicMesh/MeshNormals.cpp b/Engine/Source/Runtime/GeometryCore/Private/DynamicMesh/MeshNormals.cpp
--- a/Engine/Source/Runtime/GeometryCore/Private/DynamicMesh/MeshNormals.cpp
+++ b/Engine/Source/Runtime/GeometryCore/Private/DynamicMesh/MeshNormals.cpp
@@ -5,6 +5,7 @@
 #include "DynamicMesh/MeshNormals.h"
 #include "Async/ParallelFor.h"
 #include "DynamicMesh/MeshIndexUtil.h"
+#include "MeshQueries.h"
 
 using namespace UE::Geometry;
 
@@ -620,4 +621,7 @@ void FMeshNormals::InitializeOverlayRegionToPerVertexNormals(FDynamicMeshNormalO
 	}
 }
 
-
+FVector3d FMeshNormals::GetVertexWeightsOnTriangle(const FDynamicMesh3* Mesh, int TriID, double TriArea, bool bWeightByArea, bool bWeightByAngle)
+{
+	return TMeshQueries<FDynamicMesh3>::GetVertexWeightsOnTriangle(*Mesh, TriID, TriArea, bWeightByArea, bWeightByAngle);
+}
diff --git a/Engine/Source/Runtime/GeometryCore/Public/DynamicMesh/MeshNormals.h b/Engine/Source/Runtime/GeometryCore/Public/DynamicMesh/MeshNormals.h
--- a/Engine/Source/Runtime/GeometryCore/Public/DynamicMesh/MeshNormals.h
+++ b/Engine/Source/Runtime/GeometryCore/Public/DynamicMesh/MeshNormals.h
@@ -213,20 +213,7 @@ public:
 	 * @param bWeightByArea if true, include weighting by the area of the triangle
 	 * @param bWeightByAngle if true, include weighting by the interior angles of the triangle
 	 */
-	static FVector3d GetVertexWeightsOnTriangle(const FDynamicMesh3* Mesh, int TriID, double TriArea, bool bWeightByArea, bool bWeightByAngle)
-	{
-		FVector3d TriNormalWeights = FVector3d::One();
-		if (bWeightByAngle)
-		{
-			TriNormalWeights = Mesh->GetTriInternalAnglesR(TriID); // component-wise multiply by per-vertex internal angles
-		}
-		if (bWeightByArea)
-		{
-			TriNormalWeights *= TriArea;
-		}
-		return TriNormalWeights;
-	}
-
+	static FVector3d GetVertexWeightsOnTriangle(const FDynamicMesh3* Mesh, int TriID, double TriArea, bool bWeightByArea, bool bWeightByAngle);
 
 protected:
 	/** Compute per-vertex normals using area-weighted averaging of one-ring triangle normals */
diff --git a/Engine/Source/Runtime/GeometryCore/Public/MeshQueries.h b/Engine/Source/Runtime/GeometryCore/Public/MeshQueries.h
--- a/Engine/Source/Runtime/GeometryCore/Public/MeshQueries.h
+++ b/Engine/Source/Runtime/GeometryCore/Public/MeshQueries.h
@@ -462,6 +462,23 @@ public:
 		return Query;
 	}
 
+	/** 
+	 * Compute the length for each edge and return the result as an array of Mesh.MaxEdgeID() size.
+	 * When mesh is not compact, the length of the missing edges is set to -1.0
+	 */
+	static void GetAllEdgeLengths(const TriangleMeshType& Mesh, TArray<double>& Lengths, double& TotalLength)
+	{
+		Lengths.Init(-1.0, Mesh.MaxEdgeID());
+		TotalLength = 0.0;
+		for (const int32 EdgeID : Mesh.EdgeIndicesItr())
+		{
+			FVector3d vA, vB;
+			Mesh.GetEdgeV(EdgeID, vA, vB);
+			Lengths[EdgeID] = (vA - vB).Length();
+			TotalLength += Lengths[EdgeID];
+		}
+	}
+
 	/// Compute the mean edge length for the given mesh
 	static double AverageEdgeLength(const TriangleMeshType& Mesh)
 	{
@@ -739,6 +756,84 @@ public:
 		RootMeanSqrDeviation = FMathd::Sqrt(RootMeanSqrDeviation / NumDistances);
 	}
 
+	/**
+     * Retrieve the area and/or angle weights for each vertex of a triangle.
+	 * 
+     * @param Mesh the mesh to query
+     * @param TriID the triangle index of the mesh to query
+     * @param TriArea the area of the triangle
+     * @param bWeightByArea if true, include weighting by the area of the triangle
+     * @param bWeightByAngle if true, include weighting by the interior angles of the triangle
+     */
+    static FVector3d GetVertexWeightsOnTriangle(const TriangleMeshType& Mesh, int TriID, double TriArea, bool bWeightByArea, bool bWeightByAngle)
+    {
+        FVector3d TriNormalWeights = FVector3d::One();
+        if (bWeightByAngle)
+        {
+            TriNormalWeights = Mesh.GetTriInternalAnglesR(TriID); // component-wise multiply by per-vertex internal angles
+        }
+        if (bWeightByArea)
+        {
+            TriNormalWeights *= TriArea;
+        }
+        return TriNormalWeights;
+    }
+
+	/**
+	 * Get triangles that contain at least on vertex in the Vertices array.
+	 * 
+	 * @param Vertices Array of mesh vertex IDs.
+	 * @param Triangles Triangle IDs containing at least one vertex ID in Vertices.
+	 */
+	static TArray<int32> GetVertexSelectedTriangles(const TriangleMeshType& Mesh, const TArray<int32>& Vertices)
+	{
+		TSet<int32> TriangleSet; // All triangles shared by vertices in Vertices array
+		TriangleSet.Reserve(Vertices.Num());
+		for (const int32 VID : Vertices)
+		{
+			for (const int32 TID : Mesh.VtxTrianglesItr(VID))
+			{
+				TriangleSet.Add(TID);
+			}
+		}
+
+		return TriangleSet.Array();
+	}
+
+	/**
+	 * Expand selection of vertices with one-ring neighbors.
+	 * 
+	 * @param Selection Array of Mesh vertex IDs.
+     * @param ExpandedSelection All vertices in Selection plus one-ring neighbors for each vertex.
+     * @param VIDToExpandedSelectionIdx Maps mesh Vertex ID to ExpandedSelection array index.
+     */
+    static void ExpandVertexSelectionToNeighbors(const TriangleMeshType& Mesh, 
+                                    	   		 const TArray<int32>& Selection, 
+                                    	   		 TArray<int32>& ExpandedSelection, 
+                                    	   		 TMap<int32,int32>& VIDToExpandedSelectionIdx)
+    {
+        VIDToExpandedSelectionIdx.Reserve(Selection.Num());
+        ExpandedSelection.Reserve(Selection.Num());
+
+        int32 Idx = 0;
+        for (const int32 SelectedVID : Selection)
+        {
+            if (!VIDToExpandedSelectionIdx.Contains(SelectedVID))
+            {
+                VIDToExpandedSelectionIdx.Add(SelectedVID, Idx++);
+                ExpandedSelection.Add(SelectedVID);
+            }
+
+            for (const int32 NeighborVID : Mesh.VtxVerticesItr(SelectedVID))
+            {
+                if (!VIDToExpandedSelectionIdx.Contains(NeighborVID))
+                {
+                    VIDToExpandedSelectionIdx.Add(NeighborVID, Idx++);
+                    ExpandedSelection.Add(NeighborVID);
+                }
+            }
+        }
+    }
 };
 
 
