-Added an operator for smoothing overlays of any dimension containing any scalar data
-Added a function for computing per-edge cotangent weights
-Added a function for computing edge lengths
-Added a geometry script function for bluring colors

diff --git a/Engine/Plugins/Experimental/GeometryScripting/Source/GeometryScriptingCore/Private/MeshVertexColorFunctions.cpp b/Engine/Plugins/Experimental/GeometryScripting/Source/GeometryScriptingCore/Private/MeshVertexColorFunctions.cpp
--- a/Engine/Plugins/Experimental/GeometryScripting/Source/GeometryScriptingCore/Private/MeshVertexColorFunctions.cpp
+++ b/Engine/Plugins/Experimental/GeometryScripting/Source/GeometryScriptingCore/Private/MeshVertexColorFunctions.cpp
@@ -5,6 +5,8 @@
 #include "DynamicMesh/DynamicMesh3.h"
 #include "UDynamicMesh.h"
 #include "Util/ColorConstants.h"
+#include "Operations/SmoothDynamicMeshAttributes.h"
+#include "DynamicMesh/DynamicMeshOverlay.h"
 
 #include UE_INLINE_GENERATED_CPP_BY_NAME(MeshVertexColorFunctions)
 
@@ -407,6 +409,76 @@ UDynamicMesh* UGeometryScriptLibrary_MeshVertexColorFunctions::ConvertMeshVertex
 	return TargetMesh;
 }
 
+UDynamicMesh* UGeometryScriptLibrary_MeshVertexColorFunctions::BlurMeshVertexColors(
+	UDynamicMesh* TargetMesh,
+	FGeometryScriptMeshSelection Selection,
+	int NumIterations,
+	double Strength,
+	EGeometryScriptBlurColorMode BlurMode,
+	FGeometryScriptBlurMeshVertexColorsOptions Options,
+	UGeometryScriptDebug* Debug)
+{
+	if (TargetMesh == nullptr)
+	{
+		UE::Geometry::AppendError(Debug, EGeometryScriptErrorType::InvalidInputs, LOCTEXT("BlurMeshVertexColors_InvalidInput", "BlurMeshVertexColors: TargetMesh is Null."));
+		return TargetMesh;
+	}
+
+	if (NumIterations < 0)
+	{
+		UE::Geometry::AppendError(Debug, EGeometryScriptErrorType::InvalidInputs, LOCTEXT("BlurMeshVertexColors_InvalidIterationNumber", "BlurMeshVertexColors: Number of iterations must be non-negative."));
+		return TargetMesh;
+	}
+
+	if (Strength < 0)
+	{
+		UE::Geometry::AppendError(Debug, EGeometryScriptErrorType::InvalidInputs, LOCTEXT("BlurMeshVertexColors_InvalidStrength", "BlurMeshVertexColors: Blur strength must be non-negative."));
+		return TargetMesh;
+	}
+
+	TargetMesh->EditMesh([&](FDynamicMesh3& EditMesh) 
+	{
+		if (EditMesh.HasAttributes() == false || EditMesh.Attributes()->HasPrimaryColors() == false)
+		{
+			return;
+		}
+
+		FDynamicMeshColorOverlay* Colors = EditMesh.Attributes()->PrimaryColors();
+		if (Colors->ElementCount() > 0)
+		{
+			FSmoothDynamicMeshAttributes BlurOp(EditMesh);
+			BlurOp.bUseParallel = true;
+			BlurOp.NumIterations = NumIterations;
+			BlurOp.Strength = Strength;
+			BlurOp.EdgeWeightMethod = static_cast<FSmoothDynamicMeshAttributes::EEdgeWeights>(BlurMode);
+
+			TArray<bool> ColorsToSmooth;
+			ColorsToSmooth.Reserve(4);
+			ColorsToSmooth.Add(Options.Red);
+			ColorsToSmooth.Add(Options.Green);
+			ColorsToSmooth.Add(Options.Blue);
+			ColorsToSmooth.Add(Options.Alpha);
+
+			if (!Selection.IsEmpty())
+			{
+				Selection.ProcessByVertexID(EditMesh, [&](int32 VertexID)
+				{
+					BlurOp.Selection.Add(VertexID);
+				});
+			}
+
+			if (!BlurOp.SmoothOverlay(Colors, ColorsToSmooth))
+			{
+				UE::Geometry::AppendError(Debug, EGeometryScriptErrorType::OperationFailed, LOCTEXT("BlurMeshVertexColors_BlurFailed", "BlurMeshVertexColors: Failed to blur the colors."));
+				return;
+			}
+		}
+	}, EDynamicMeshChangeType::GeneralEdit, EDynamicMeshAttributeChangeFlags::Unknown, false);
+	
+	return TargetMesh;
+}
+
+
 
 
 
diff --git a/Engine/Plugins/Experimental/GeometryScripting/Source/GeometryScriptingCore/Public/GeometryScript/MeshVertexColorFunctions.h b/Engine/Plugins/Experimental/GeometryScripting/Source/GeometryScriptingCore/Public/GeometryScript/MeshVertexColorFunctions.h
--- a/Engine/Plugins/Experimental/GeometryScripting/Source/GeometryScriptingCore/Public/GeometryScript/MeshVertexColorFunctions.h
+++ b/Engine/Plugins/Experimental/GeometryScripting/Source/GeometryScriptingCore/Public/GeometryScript/MeshVertexColorFunctions.h
@@ -11,6 +11,40 @@
 class UDynamicMesh;
 
 
+UENUM(BlueprintType)
+enum class EGeometryScriptBlurColorMode : uint8
+{
+	/** Blur the attributes where each neighbor is weighted equally. */
+	Uniform = 0,
+	/** Blur the attributes where each neighbor is weighted proportionally to the shared edge length. */
+	EdgeLength = 1,
+	/** Blur the attributes where each neighbor is weighted proportionally to the cotangent weight of the shared edge. */
+	CotanWeights = 2
+};
+
+USTRUCT(BlueprintType)
+struct GEOMETRYSCRIPTINGCORE_API FGeometryScriptBlurMeshVertexColorsOptions
+{
+	GENERATED_BODY()
+
+	/** Blur red channel. */
+	UPROPERTY(EditAnywhere, BlueprintReadWrite, Category = Options)
+	bool Red = true;
+
+	/** Blur green channel. */
+	UPROPERTY(EditAnywhere, BlueprintReadWrite, Category = Options)
+	bool Green = true;
+
+	/** Blur blue channel. */
+	UPROPERTY(EditAnywhere, BlueprintReadWrite, Category = Options)
+	bool Blue = true;
+
+	/** Blur alpha channel. */
+	UPROPERTY(EditAnywhere, BlueprintReadWrite, Category = Options)
+	bool Alpha = true;
+};
+
+
 
 UCLASS(meta = (ScriptName = "GeometryScript_VertexColors"))
 class GEOMETRYSCRIPTINGCORE_API UGeometryScriptLibrary_MeshVertexColorFunctions : public UBlueprintFunctionLibrary
@@ -100,5 +134,26 @@ public:
 	ConvertMeshVertexColorsLinearToSRGB(
 		UDynamicMesh* TargetMesh,
 		UGeometryScriptDebug* Debug = nullptr);
+
+
+	/**
+     * Blur the color attribute of the mesh. If the mesh has no color attribute, the function returns the mesh unchanged.
+	 * 
+     * @param TargetMesh The mesh containing the color attribute. 
+     * @param Selection Only vertices in the selection will have their color attribute blurred.
+     * @param NumIterations The number of blur iterations.
+     * @param Strength Each iteration, we will blur between the vertex of the color at the previous iteration and its neighbors' average by Strength amount (expected to be in the zero to one range).
+     * @param BlurMode Determines how neighbors are weighted when computing their average.
+	 */
+	UFUNCTION(BlueprintCallable, Category = "GeometryScript|VertexColor", meta=(ScriptMethod))
+	static UPARAM(DisplayName = "Target Mesh") UDynamicMesh*
+	BlurMeshVertexColors(
+		UDynamicMesh* TargetMesh,
+		FGeometryScriptMeshSelection Selection,
+		int NumIterations = 1,
+		double Strength = 0.5,
+		EGeometryScriptBlurColorMode BlurMode = EGeometryScriptBlurColorMode::Uniform,
+		FGeometryScriptBlurMeshVertexColorsOptions Options = FGeometryScriptBlurMeshVertexColorsOptions(),
+		UGeometryScriptDebug* Debug = nullptr);
 };
 
diff --git a/Engine/Plugins/Runtime/GeometryProcessing/Source/DynamicMesh/Private/Operations/SmoothDynamicMeshAttributes.cpp b/Engine/Plugins/Runtime/GeometryProcessing/Source/DynamicMesh/Private/Operations/SmoothDynamicMeshAttributes.cpp
new file mode 100644
--- /dev/null
+++ b/Engine/Plugins/Runtime/GeometryProcessing/Source/DynamicMesh/Private/Operations/SmoothDynamicMeshAttributes.cpp
@@ -0,0 +1,211 @@
+// Copyright Epic Games, Inc. All Rights Reserved.
+
+#include "Operations/SmoothDynamicMeshAttributes.h"
+#include "DynamicMesh/DynamicMesh3.h"
+#include "DynamicMesh/DynamicMeshOverlay.h"
+#include "Solvers/PrecomputedMeshWeightData.h"
+#include "Util/ProgressCancel.h"
+#include "Async/ParallelFor.h"
+#include "MeshQueries.h"
+#include "Math/UnrealMathUtility.h"
+#include "Operations/AverageOverlayToVertices.h"
+
+using namespace UE::Geometry;
+using namespace UE::MeshDeformation; 
+
+FSmoothDynamicMeshAttributes::FSmoothDynamicMeshAttributes(const FDynamicMesh3& InMesh)
+:
+Mesh(&InMesh)
+{
+}
+
+bool FSmoothDynamicMeshAttributes::Cancelled()
+{
+	return (Progress == nullptr) ? false : Progress->Cancelled();
+}
+
+template<typename RealType, int ElementSize>
+bool FSmoothDynamicMeshAttributes::SmoothOverlay(TDynamicMeshOverlay<RealType, ElementSize>* InOverlay, const TArray<bool>& DimensionsToSmooth)
+{
+	if (NumIterations <= 0 || Strength <= 0 || Mesh->EdgeCount() == 0)
+	{
+		return true;
+	}
+
+	if (!(DimensionsToSmooth.IsEmpty() || DimensionsToSmooth.Num() == ElementSize))
+	{
+		return false;
+	}
+
+    //
+    // Check if we need to pre-compute cotangent weights or edge lengths. Avoid re-compute on subsequent calls.
+    //
+    if (EdgeWeightMethod == EEdgeWeights::CotanWeights && CotangentEdgeWeights.Num() != Mesh->MaxEdgeID())
+    {
+        ConstructEdgeCotanWeightsDataArray(*Mesh, CotangentEdgeWeights);
+    }
+    if (EdgeWeightMethod == EEdgeWeights::EdgeLength && EdgeLengthWeights.Num() != Mesh->MaxEdgeID())
+    {   
+        double TotalLength;
+        TMeshQueries<FDynamicMesh3>::GetAllEdgeLengths(*Mesh, EdgeLengthWeights, TotalLength);
+        EdgeLengthThreshold = 0.001*(TotalLength/Mesh->EdgeCount()); // minimum acceptable edge length for weight computation
+    }
+
+    //
+    // Compute an array of vertex indices that contains the selection and the all the one-ring neighbors
+    //
+    ExpandedSelection.Reset();
+    VIDToExpandedSelectionIdx.Reset();
+    if (Selection.Num())
+    {
+        TMeshQueries<FDynamicMesh3>::ExpandVertexSelectionToNeighbors(*Mesh, Selection, ExpandedSelection, VIDToExpandedSelectionIdx);
+    }
+    
+    //
+    // For every vertex in ExpandedSelection, compute the area/angle weighted average of the overlay elements on that vertex
+    //
+    TArray<RealType> VertexVectorValue;
+    TArray<bool> HasElements;
+    FAverageOverlayToVertices AverageOverlayOp(*Mesh);
+    if (Selection.Num())
+    {
+        AverageOverlayOp.Selection = MoveTemp(ExpandedSelection);
+    }
+    if (!AverageOverlayOp.AverageOverlay(InOverlay, VertexVectorValue, HasElements))
+    {
+        return false;
+    }
+
+    //
+	// Iterate over every vertex (in Selection or the whole mesh if the Selection is empty) and smooth
+    //
+    const int32 NumVertsToSmooth = Selection.Num() ? Selection.Num() : Mesh->MaxVertexID();
+    TArray<RealType> NewVertexVectorValue = VertexVectorValue;
+    for (int32 Iteration = 0; Iteration < NumIterations; ++Iteration)
+    {
+        ParallelFor(NumVertsToSmooth, [&](int32 VIdx)
+        {
+            const int32 VertexID = Selection.Num() ? Selection[VIdx] : VIdx;
+
+            if (!Mesh->IsVertex(VertexID))
+            {
+                return;
+            }
+        
+            RealType WeightII = 0.0;        // Total weight value for all elements
+            RealType VecAccum[ElementSize]; // Value of weighted elements
+            for (int32 EIdx = 0; EIdx < ElementSize; ++EIdx)
+            {
+                VecAccum[EIdx] = static_cast<RealType>(0.0);
+            }
+
+            for (const int32 EdgeId : Mesh->VtxEdgesItr(VertexID))
+            {   
+				const FIndex2i EdgeV = Mesh->GetEdgeV(EdgeId);
+				const int32 NextV = (EdgeV.A == VertexID) ? EdgeV.B : EdgeV.A;
+                const int32 NextSelectionID = MeshVIDToSelectionIdx(NextV);
+
+                // check if the vertex has at least one element and needs to be considered in the calculation
+                if (HasElements[NextSelectionID]) 
+                {
+                    RealType WeightIJ = static_cast<RealType>(1.0f); // EEdgeWeights::Uniform
+                    if (EdgeWeightMethod == EEdgeWeights::CotanWeights) 
+                    {
+                        // cotangent weights can be negative, hence we clamp them to zero
+                        WeightIJ = static_cast<RealType>(FMath::Clamp(CotangentEdgeWeights[EdgeId], (RealType)0.0, CotangentEdgeWeights[EdgeId])); 
+                    }
+                    else if (EdgeWeightMethod == EEdgeWeights::EdgeLength)
+                    {
+                        WeightIJ = static_cast<RealType>(1.0/FMath::Max(EdgeLengthThreshold, EdgeLengthWeights[EdgeId]));
+                    }
+                    
+                    for (int32 EIdx = 0; EIdx < ElementSize; ++EIdx)
+                    {
+                        VecAccum[EIdx] += VertexVectorValue[ElementSize*NextSelectionID + EIdx] * WeightIJ;
+                    }
+                    WeightII += WeightIJ;
+                }
+               
+            }
+
+            if (!FMath::IsNearlyZero((double)WeightII))
+            {
+                const RealType InvWeightII = RealType(1.0/WeightII);
+                for (int32 EIdx = 0; EIdx < ElementSize; ++EIdx)
+                {
+				    VecAccum[EIdx] *= InvWeightII;
+                }
+            }
+            
+            const int32 SelectionIdx = MeshVIDToSelectionIdx(VertexID);
+			for (int32 EIdx = 0; EIdx < ElementSize; ++EIdx)
+			{
+                const int32 Start = ElementSize*SelectionIdx;
+				NewVertexVectorValue[Start + EIdx] = FMath::Lerp(VertexVectorValue[Start + EIdx], VecAccum[EIdx], Strength);
+			}
+        }, bUseParallel ? EParallelForFlags::None : EParallelForFlags::ForceSingleThread);
+
+		VertexVectorValue = NewVertexVectorValue;
+    }
+
+
+    //
+    // Write the smoothed data back into the overlay
+    //
+    ParallelFor(NumVertsToSmooth, [&](int32 Idx)
+    {
+        const int32 VertexID = Selection.Num() ? Selection[Idx] : Idx;
+        if (!Mesh->IsVertex(VertexID))
+        {
+            return;
+        }
+
+        const int32 SelectionIdx = MeshVIDToSelectionIdx(VertexID);
+
+        TArray<int32> Elements;
+        InOverlay->GetVertexElements(VertexID, Elements);
+		RealType ElementData[ElementSize];
+       
+        for (const int32 ElementID : Elements)
+        { 
+			InOverlay->GetElement(ElementID, ElementData);
+
+            if (SplitVertexModel == ESplitVertexModel::SetAllToSmoothed)  
+            {   
+				for (int32 EIdx = 0; EIdx < ElementSize; ++EIdx)
+				{
+					if (DimensionsToSmooth.IsEmpty() || DimensionsToSmooth[EIdx])
+					{
+						ElementData[EIdx] = NewVertexVectorValue[SelectionIdx * ElementSize + EIdx];
+					}
+				}
+            }
+            else if (SplitVertexModel == ESplitVertexModel::BlendWithSmoothed)
+            {
+				for (int32 EIdx = 0; EIdx < ElementSize; ++EIdx)
+				{
+					if (DimensionsToSmooth.IsEmpty() || DimensionsToSmooth[EIdx])
+					{
+						ElementData[EIdx] = FMath::Lerp(ElementData[EIdx], NewVertexVectorValue[SelectionIdx * ElementSize + EIdx], BlendWithSmoothedStrength);
+					}
+				}
+            }
+            else
+            {
+                checkNoEntry(); // unsupported type
+            }
+
+            InOverlay->SetElement(ElementID, ElementData);
+        }
+    }, bUseParallel ? EParallelForFlags::None : EParallelForFlags::ForceSingleThread);
+
+	return true;
+}
+
+
+// template instantiation
+template DYNAMICMESH_API bool UE::Geometry::FSmoothDynamicMeshAttributes::SmoothOverlay<float,4>(class UE::Geometry::TDynamicMeshOverlay<float,4>*, class TArray<bool, class TSizedDefaultAllocator<32> > const&);
+template DYNAMICMESH_API bool UE::Geometry::FSmoothDynamicMeshAttributes::SmoothOverlay<float,3>(class UE::Geometry::TDynamicMeshOverlay<float,3>*, class TArray<bool, class TSizedDefaultAllocator<32> > const&);
+template DYNAMICMESH_API bool UE::Geometry::FSmoothDynamicMeshAttributes::SmoothOverlay<float,2>(class UE::Geometry::TDynamicMeshOverlay<float,2>*, class TArray<bool, class TSizedDefaultAllocator<32> > const&);
+template DYNAMICMESH_API bool UE::Geometry::FSmoothDynamicMeshAttributes::SmoothOverlay<float,1>(class UE::Geometry::TDynamicMeshOverlay<float,1>*, class TArray<bool, class TSizedDefaultAllocator<32> > const&);
+template DYNAMICMESH_API bool UE::Geometry::FSmoothDynamicMeshAttributes::SmoothOverlay<double,3>(class UE::Geometry::TDynamicMeshOverlay<double,3>*, class TArray<bool, class TSizedDefaultAllocator<32> > const&);
diff --git a/Engine/Plugins/Runtime/GeometryProcessing/Source/DynamicMesh/Private/Solvers/PrecomputedMeshWeightData.cpp b/Engine/Plugins/Runtime/GeometryProcessing/Source/DynamicMesh/Private/Solvers/PrecomputedMeshWeightData.cpp
--- a/Engine/Plugins/Runtime/GeometryProcessing/Source/DynamicMesh/Private/Solvers/PrecomputedMeshWeightData.cpp
+++ b/Engine/Plugins/Runtime/GeometryProcessing/Source/DynamicMesh/Private/Solvers/PrecomputedMeshWeightData.cpp
@@ -156,3 +156,25 @@ void UE::MeshDeformation::MeanValueTriangleData::Initialize(const FDynamicMesh3&
 #endif
 }
 
+void UE::MeshDeformation::ConstructEdgeCotanWeightsDataArray(const FDynamicMesh3& Mesh, TArray<double>& EdgeWeightsDataArray, double ClampMin, double ClampMax)
+{
+	TArray<UE::MeshDeformation::CotanTriangleData> CotangentTriangleDataArray;
+	CotangentTriangleDataArray.SetNumUninitialized(Mesh.MaxTriangleID());
+
+	for (const int32 TriID : Mesh.TriangleIndicesItr())
+	{
+		CotangentTriangleDataArray[TriID].Initialize(Mesh, TriID);
+	}
+
+	EdgeWeightsDataArray.SetNumUninitialized(Mesh.MaxEdgeID());
+	for (const int32 EdgeId : Mesh.EdgeIndicesItr())
+	{   
+		const FDynamicMesh3::FEdge Edge = Mesh.GetEdge(EdgeId);
+
+		const double CotanAlpha = CotangentTriangleDataArray[Edge.Tri[0]].GetOpposingCotangent(EdgeId);
+
+		const double CotanBeta = (Edge.Tri[1] != FDynamicMesh3::InvalidID) ? CotangentTriangleDataArray[Edge.Tri[1]].GetOpposingCotangent(EdgeId) : 0.0;
+
+		EdgeWeightsDataArray[EdgeId] = FMathd::Clamp(CotanAlpha + CotanBeta, ClampMin, ClampMax);
+	}
+}
\ No newline at end of file
diff --git a/Engine/Plugins/Runtime/GeometryProcessing/Source/DynamicMesh/Public/Operations/SmoothDynamicMeshAttributes.h b/Engine/Plugins/Runtime/GeometryProcessing/Source/DynamicMesh/Public/Operations/SmoothDynamicMeshAttributes.h
new file mode 100644
--- /dev/null
+++ b/Engine/Plugins/Runtime/GeometryProcessing/Source/DynamicMesh/Public/Operations/SmoothDynamicMeshAttributes.h
@@ -0,0 +1,160 @@
+// Copyright Epic Games, Inc. All Rights Reserved.
+
+#pragma once
+
+#include "CoreTypes.h"
+#include "Containers/Map.h"
+#include "Containers/Array.h"
+
+// Forward declarations
+class FProgressCancel;
+namespace UE::Geometry 
+{ 
+	class FDynamicMesh3; 
+    template<typename RealType, int ElementSize>
+    class TDynamicMeshOverlay; 
+}
+
+namespace UE::MeshDeformation
+{
+    class CotanTriangleData;
+}
+
+namespace UE
+{
+namespace Geometry
+{
+
+/** 
+ * Smooth scalar data of any dimension that is stored via TDynamicMeshOverlay. The operator doesn't know anything about 
+ * constraints that are associated with the actual overlay type. For example, if the overlay contains unit length 
+ * normals, the resulting "smoothed" normals might not have unit length.
+ * 
+ * The smoothing happens in stages:
+ * 
+ *  - First, we average the overlay element values at each vertex to obtain the per-vertex values.
+ * 
+ *  - For each vertex, we compute the weighted average of its connected neighbors using the EdgeWeightMethod weights.
+ *    The smoothed value equals the (1-Strength)*PerVertexValue + (Strength)*NeighborAverage. Strength is expected to be 
+ *    in the zero to one range, otherwise the operation can turn from smoothing to more of an unsharp filter. 
+ *    We repeat this NumIterations times.
+ * 
+ *  - Finally, the smoothed per-vertex values are written back to the overlay by either overwriting the content of 
+ *    all per-vertex elements (SplitVertexModel == ESplitVertexModel::SetAllToSmoothed) or blending on per-element 
+ *    basis using the BlendWithSmoothedStrength blend value. The calling code could use FSplitAttributeWelder to weld 
+ *    split-attributes that happen to have identical values after this smoothing (happens when the 
+ *    SplitVertexModel == ESplitVertexModel::SetAllToSmoothed).
+ * 
+ * 
+ * Example Usage:
+ * 
+ *  // Smoothing the color
+ *  FDynamicMesh3 Mesh = ...
+ *  FDynamicMeshColorOverlay* Colors = Mesh.Attributes()->PrimaryColors();
+ * 
+ *  FSmoothDynamicMeshAttributes BlurOp(&Mesh);
+ *  
+ *  // Setup the parameters
+ *  BlurOp.bUseParallel = true;
+ *  BlurOp.NumIterations = 10;
+ *  BlurOp.Strength = 0.2;
+ *  BlurOp.EdgeWeightMethod = EEdgeWeights::CotanWeights;
+ * 
+ *  // Run the smoothing
+ *  BlurOp.SmoothOverlay(Colors);
+ */
+class DYNAMICMESH_API FSmoothDynamicMeshAttributes
+{
+public:
+
+	//
+	// Optional Inputs
+	//
+	
+    /** Set this to be able to cancel the running operation. */
+	FProgressCancel* Progress = nullptr;
+
+    /** Enable/disable multi-threading. */
+    bool bUseParallel = true;
+
+    /** Subset of points to smooth. */
+    TArray<int32> Selection;
+
+    enum class EEdgeWeights : uint8
+    {
+        /** Smooth the attributes where each neighbor is weighted equally. */
+        Uniform = 0,
+        /** Smooth the attributes where each neighbor is weighted proportionally to the shared edge length. */
+        EdgeLength = 1,
+		/** Smooth the attributes where each neighbor is weighted proportionally to the cotangent weight of the shared edge. */
+		CotanWeights = 2
+    };
+    EEdgeWeights EdgeWeightMethod = EEdgeWeights::CotanWeights;
+
+    /** How to deal with the split vertices. */
+    enum class ESplitVertexModel : uint8
+    {
+        /** Each element will be assigned the same value. */
+        SetAllToSmoothed, 
+        /** Each element will be blended with he smoothed value. Set BlendWithSmoothedStrength to control the amount of blend. */
+        BlendWithSmoothed 
+    };
+    ESplitVertexModel SplitVertexModel = ESplitVertexModel::SetAllToSmoothed;
+
+    /** The number of smoothing iterations. */
+    int32 NumIterations = 0;
+
+    /** The strength of each smoothing iteration. */
+    double Strength = 0.0; 
+
+    /** 
+     * Control the amount of blending between the final smoothed value and the original element value. Only used when 
+     * the SplitVertexModel==ESplitVertexModel::BlendWithSmoothed 
+     */
+    double BlendWithSmoothedStrength = 0.0;
+
+protected:
+		
+	const FDynamicMesh3* Mesh;
+
+    TArray<double> CotangentEdgeWeights;
+
+    TArray<double> EdgeLengthWeights;
+
+    // If Selection array is not empty, will contain all vertices in Selection plus one-ring neighbors for each vertex
+    TArray<int32> ExpandedSelection;
+	
+    // Maps mesh Vertex ID to ExpandedSelection index
+    TMap<int32, int32> VIDToExpandedSelectionIdx;
+
+    double EdgeLengthThreshold = 1.0;
+
+public:
+
+	FSmoothDynamicMeshAttributes(const FDynamicMesh3& InMesh);
+
+	virtual ~FSmoothDynamicMeshAttributes() = default;
+
+	/** 
+     * @param InOverlay Overlay to smooth
+     * @param DimensionsToSmooth ElementSize array, where DimensionsToSmooth[Idx] is true if we want smooth the 
+     *                           dimension Idx. If empty, we smooth all dimensions.
+     * @return true if the algorithm succeeds, false if it failed or was canceled by the user. 
+     */
+    template<typename RealType, int ElementSize>
+	bool SmoothOverlay(TDynamicMeshOverlay<RealType, ElementSize>* InOverlay, const TArray<bool>& DimensionsToSmooth = TArray<bool>());
+
+protected:
+	
+    /** @return if true, abort the computation. */
+	virtual bool Cancelled();
+
+    /** @return the index of the vertex into ExpandedSelection array if Selection is not empty, otherwise returns VertexID. */
+    FORCEINLINE int32 MeshVIDToSelectionIdx(int32 VertexID) const
+    {
+        return !VIDToExpandedSelectionIdx.IsEmpty() && ensure(VIDToExpandedSelectionIdx.Contains(VertexID)) ? VIDToExpandedSelectionIdx[VertexID] : VertexID;
+    }
+};
+
+}
+}
\ No newline at end of file
diff --git a/Engine/Plugins/Runtime/GeometryProcessing/Source/DynamicMesh/Public/Solvers/PrecomputedMeshWeightData.h b/Engine/Plugins/Runtime/GeometryProcessing/Source/DynamicMesh/Public/Solvers/PrecomputedMeshWeightData.h
--- a/Engine/Plugins/Runtime/GeometryProcessing/Source/DynamicMesh/Public/Solvers/PrecomputedMeshWeightData.h
+++ b/Engine/Plugins/Runtime/GeometryProcessing/Source/DynamicMesh/Public/Solvers/PrecomputedMeshWeightData.h
@@ -221,6 +221,11 @@ namespace UE
 		}
 
 
-
+	  	/**
+		 * @return Edge cotanget weights.
+		 * @todo Add an option to return the weights computed on an intrinsic mesh.
+		 */
+		void DYNAMICMESH_API ConstructEdgeCotanWeightsDataArray(const FDynamicMesh3& Mesh, TArray<double>& EdgeWeightsDataArray, double ClampMin = -1.e5, double ClampMax = 1.e5);
+		
 	}
 }
\ No newline at end of file
