Chaos Cloth Asset: TransferBoneWeights bugfixes.

* Fail gracefully if you don't have any initial matches when doing Inpaint Weights.
* Fix ClothPatternToDynamicMesh to actually populate the NonManifoldMapping data.

diff --git a/Engine/Plugins/Experimental/ChaosClothAssetEditor/Source/ChaosClothAssetTools/Private/ChaosClothAsset/ClothPatternToDynamicMesh.cpp b/Engine/Plugins/Experimental/ChaosClothAssetEditor/Source/ChaosClothAssetTools/Private/ChaosClothAsset/ClothPatternToDynamicMesh.cpp
--- a/Engine/Plugins/Experimental/ChaosClothAssetEditor/Source/ChaosClothAssetTools/Private/ChaosClothAsset/ClothPatternToDynamicMesh.cpp
+++ b/Engine/Plugins/Experimental/ChaosClothAssetEditor/Source/ChaosClothAssetTools/Private/ChaosClothAsset/ClothPatternToDynamicMesh.cpp
@@ -642,6 +642,7 @@ void FClothPatternToDynamicMesh::Convert(const TSharedRef<const FManagedArrayCol
 	}
 	if (!bVertexIndicesMatch)
 	{
+		MeshOut.EnableAttributes();
 		UE::Geometry::FNonManifoldMappingSupport::AttachNonManifoldVertexMappingData(PatternToDynamicMesh.ToSrcVertIDMap, MeshOut);
 	}
 
diff --git a/Engine/Plugins/Runtime/GeometryProcessing/Source/DynamicMesh/Private/Operations/TransferBoneWeights.cpp b/Engine/Plugins/Runtime/GeometryProcessing/Source/DynamicMesh/Private/Operations/TransferBoneWeights.cpp
--- a/Engine/Plugins/Runtime/GeometryProcessing/Source/DynamicMesh/Private/Operations/TransferBoneWeights.cpp
+++ b/Engine/Plugins/Runtime/GeometryProcessing/Source/DynamicMesh/Private/Operations/TransferBoneWeights.cpp
@@ -344,6 +344,12 @@ bool FTransferBoneWeights::TransferWeightsToMesh(FDynamicMesh3& InOutTargetMesh,
 			return true;
 		}
 
+		// If no vertices matched, we have nothing to inpaint.
+		if (NumMatched == 0)
+		{
+			return false;
+		}
+
 		// Setup the sparse matrix FixedValues of known (matched) weight values and the array (FixedIndices) of the matched vertex IDs
 		const int32 TargetNumBones = InOutTargetMesh.Attributes()->GetBoneNames()->GetAttribValues().Num();
 		FSparseMatrixD FixedValues;
