add variant of mesh projection hull that keeps 3d hull on one side of the projection

diff --git a/Engine/Plugins/Runtime/GeometryProcessing/Source/DynamicMesh/Private/Operations/MeshProjectionHull.cpp b/Engine/Plugins/Runtime/GeometryProcessing/Source/DynamicMesh/Private/Operations/MeshProjectionHull.cpp
--- a/Engine/Plugins/Runtime/GeometryProcessing/Source/DynamicMesh/Private/Operations/MeshProjectionHull.cpp
+++ b/Engine/Plugins/Runtime/GeometryProcessing/Source/DynamicMesh/Private/Operations/MeshProjectionHull.cpp
@@ -5,6 +5,7 @@
 #include "MeshSimplification.h"
 #include "DynamicMesh/MeshNormals.h"
 #include "CompGeom/ConvexHull2.h"
+#include "CompGeom/ConvexHull3.h"
 #include "Generators/SweepGenerator.h"
 
 using namespace UE::Geometry;
@@ -59,6 +60,11 @@ bool FMeshProjectionHull::Compute()
 		CenterFrame.Origin -= (ExtrudeLength * 0.5) * ProjAxis;
 	}
 
+	if (Keep3DHullSide != EKeep3DHullSide::None)
+	{
+		return ComputeWith3DHullSide(CenterFrame, ExtrudeLength);
+	}
+
 	// generate the swept-polygon mesh
 	FGeneralizedCylinderGenerator MeshGen;
 	MeshGen.CrossSection = (bSimplifyPolygon) ? SimplifiedHull2D : ConvexHull2D;
@@ -69,5 +75,88 @@ bool FMeshProjectionHull::Compute()
 	MeshGen.Generate();
 	ConvexHull3D.Copy(&MeshGen);
 
+	return true;
+}
+
+
+bool FMeshProjectionHull::ComputeWith3DHullSide(FFrame3d CenterFrame, double ExtrudeLength)
+{
+	FVector3d ProjAxis = ProjectionFrame.Z();
+
+	TArray<FVector3d> Verts3D;
+	const FPolygon2d* UsePolygon = bSimplifyPolygon ? &SimplifiedHull2D : &ConvexHull2D;
+	Verts3D.Reserve(Mesh->VertexCount() + UsePolygon->VertexCount());
+
+	if (bSimplifyPolygon)
+	{
+		// if the polygon was simplified, project all mesh vertices inside the polygon
+		for (FVector3d Position : Mesh->VerticesItr())
+		{
+			FVector2d FramePos = ProjectionFrame.ToPlaneUV(Position, 2);
+			FVector2d InPolyPos = FramePos;
+			bool bProjected = CurveUtil::ProjectPointInsideConvexPolygon<double>(SimplifiedHull2D.GetVertices(), InPolyPos, false);
+			if (bProjected)
+			{
+				FVector3d UnprojectedOriginal = ProjectionFrame.FromPlaneUV(FramePos, 2);
+				FVector3d UnprojectedInside = ProjectionFrame.FromPlaneUV(InPolyPos, 2);
+				Position += UnprojectedInside - UnprojectedOriginal;
+			}
+			Verts3D.Add(Position);
+		}
+	}
+	else
+	{
+		// otherwise, can just take the vertices as-is
+		for (FVector3d Position : Mesh->VerticesItr())
+		{
+			Verts3D.Add(Position);
+		}
+	}
+
+	// Add the polygon vertices, aligned to the flat side
+	if (Keep3DHullSide == EKeep3DHullSide::Back)
+	{
+		CenterFrame.Origin += ExtrudeLength * ProjAxis;
+	}
+	const TArray<FVector2d>& Verts = UsePolygon->GetVertices();
+	for (const FVector2d& Vert : Verts)
+	{
+		Verts3D.Add(CenterFrame.FromPlaneUV(Vert, 2));
+	}
+
+	// Compute a 3D hull of the flat polygon face vertices and the 3D vertices
+	FConvexHull3d HullCompute3d;
+	bool bOK = HullCompute3d.Solve<FVector3d>(Verts3D);
+	if (!bOK)
+	{
+		return false;
+	}
+
+	// Convert the hull triangles to a compact dynamic mesh with no inner/unused vertices
+	TArray<int32> SourceVertToHullVert;
+	SourceVertToHullVert.Init(-1, Verts3D.Num());
+	ConvexHull3D = FDynamicMesh3(EMeshComponents::None);
+	HullCompute3d.GetTriangles([&](FIndex3i Triangle)
+		{
+			for (int32 SubIdx = 0; SubIdx < 3; ++SubIdx)
+			{
+				int32 SourceIndex = Triangle[SubIdx];
+				int32 HullIndex = SourceVertToHullVert[SourceIndex];
+				if (HullIndex == -1)
+				{
+					FVector3d OrigPos = Verts3D[SourceIndex];
+					int32 NewVID = ConvexHull3D.AppendVertex(OrigPos);
+					SourceVertToHullVert[SourceIndex] = NewVID;
+					Triangle[SubIdx] = NewVID;
+				}
+				else
+				{
+					Triangle[SubIdx] = HullIndex;
+				}
+			}
+
+			ConvexHull3D.AppendTriangle(Triangle);
+		});
+
 	return true;
 }
\ No newline at end of file
diff --git a/Engine/Plugins/Runtime/GeometryProcessing/Source/DynamicMesh/Public/Operations/MeshProjectionHull.h b/Engine/Plugins/Runtime/GeometryProcessing/Source/DynamicMesh/Public/Operations/MeshProjectionHull.h
--- a/Engine/Plugins/Runtime/GeometryProcessing/Source/DynamicMesh/Public/Operations/MeshProjectionHull.h
+++ b/Engine/Plugins/Runtime/GeometryProcessing/Source/DynamicMesh/Public/Operations/MeshProjectionHull.h
@@ -46,6 +46,20 @@ public:
 	/** Output swept-polygon convex hull */
 	FDynamicMesh3 ConvexHull3D;
 
+	enum class EKeep3DHullSide : uint8
+	{
+		// Both sides of the projection hull will be flat
+		None,
+		// The front of the projection hull will follow the 3D convex hull, and the back will be flat
+		Front,
+		// The back of the projection hull will follow the 3D convex hull, and the front will be flat
+		Back
+		// Note: If both the front and back follow the 3D hull, that is just a regular convex hull; see MeshConvexHull.h
+	};
+
+	/** Whether to conform to the 3D convex hull surface on the front or back side of the sweep, or to use a flat surface on both sides of the swept hull */
+	EKeep3DHullSide Keep3DHullSide = EKeep3DHullSide::None;
+
 public:
 	FMeshProjectionHull(const FDynamicMesh3* MeshIn)
 	{
@@ -58,6 +72,11 @@ public:
 	 */
 	bool Compute();
 
+private:
+
+	// Helper to compute the 3D hull when Keep3DHullSide is not None; called by Compute after the 2D projection hull is computed
+	bool ComputeWith3DHullSide(FFrame3d CenterFrame, double ExtrudeLength);
+
 };
 
 } // end namespace UE::Geometry
diff --git a/Engine/Source/Runtime/GeometryCore/Public/Curve/CurveUtil.h b/Engine/Source/Runtime/GeometryCore/Public/Curve/CurveUtil.h
--- a/Engine/Source/Runtime/GeometryCore/Public/Curve/CurveUtil.h
+++ b/Engine/Source/Runtime/GeometryCore/Public/Curve/CurveUtil.h
@@ -354,6 +354,39 @@ using namespace UE::Math;
 			&& int(TMathUtil<RealType>::Round(AngleSum / TMathUtil<RealType>::TwoPi)) == (int)bSeenPos - (int)bSeenNeg;
 	}
 
+	/**
+	 * Project point inside a convex polygon with known orientation
+	 *
+	 * @param ProjPt				Point to project
+	 * @param bReverseOrientation	Whether convex polygon orientation is reversed (i.e., has negative signed area)
+	 * @return						true if the point was projected
+	 */
+	template<typename RealType>
+	bool ProjectPointInsideConvexPolygon(const TArrayView<const TVector2<RealType>> Vertices, TVector2<RealType>& ProjPt, bool bReverseOrientation = false)
+	{
+		bool bProjected = false;
+		const int32 N = Vertices.Num();
+		for (int32 Idx = 0, PrevIdx = N - 1; Idx < N; PrevIdx = Idx++)
+		{
+			const TVector2<RealType>& V1 = Vertices[PrevIdx];
+			const TVector2<RealType>& V2 = Vertices[Idx];
+			TVector2<RealType> ToPt = ProjPt - V1;
+			TVector2<RealType> Normal = PerpCW<RealType>(V2 - V1);
+			RealType Orient = ToPt.Dot(Normal);
+			bool bOutside = bReverseOrientation ? Orient < 0 : Orient > 0;
+			if (bOutside)
+			{
+				RealType SqLen = Normal.SquaredLength();
+				if (SqLen > (RealType)FLT_MIN)
+				{
+					bProjected = true;
+					ProjPt -= Normal * Orient / SqLen;
+				}
+			}
+		}
+		return bProjected;
+	}
+
 
 	/**
 	 * smooth vertices in-place (will not produce a symmetric result, but does not require extra buffer)
