PR #10645: Fix extra parentheses compilation warning with clang-cl

diff --git a/Engine/Plugins/Runtime/GeometryProcessing/Source/DynamicMesh/Private/Operations/MeshBevel.cpp b/Engine/Plugins/Runtime/GeometryProcessing/Source/DynamicMesh/Private/Operations/MeshBevel.cpp
--- a/Engine/Plugins/Runtime/GeometryProcessing/Source/DynamicMesh/Private/Operations/MeshBevel.cpp
+++ b/Engine/Plugins/Runtime/GeometryProcessing/Source/DynamicMesh/Private/Operations/MeshBevel.cpp
@@ -28,7 +28,7 @@ using namespace UE::Geometry;
 #define MESH_BEVEL_DEBUG_ENSURE(Expr) ensure(Expr)
 #else
 #define MESH_BEVEL_DEBUG_CHECK(Expr) 
-#define MESH_BEVEL_DEBUG_ENSURE(Expr) (Expr)
+#define MESH_BEVEL_DEBUG_ENSURE(Expr) !!(Expr)
 #endif
 
 
