Added ability to specify force inpaint mask via weight map name in addition to just plain array.

#rnx

diff --git a/Engine/Plugins/Runtime/GeometryProcessing/Source/DynamicMesh/Private/Operations/TransferBoneWeights.cpp b/Engine/Plugins/Runtime/GeometryProcessing/Source/DynamicMesh/Private/Operations/TransferBoneWeights.cpp
--- a/Engine/Plugins/Runtime/GeometryProcessing/Source/DynamicMesh/Private/Operations/TransferBoneWeights.cpp
+++ b/Engine/Plugins/Runtime/GeometryProcessing/Source/DynamicMesh/Private/Operations/TransferBoneWeights.cpp
@@ -291,8 +291,23 @@ bool FTransferBoneWeights::TransferWeightsToMesh(FDynamicMesh3& InOutTargetMesh,
          *      i.e. W(i,j) = KnownWeights(i,j) where i is a vertex for which we found a match on the body.
 		 */
 		
+		// Check if the target mesh contains the user specifed force inpaint weight map 
+		const FDynamicMeshWeightAttribute* ForceInpaintLayer = nullptr;
+		if (!ForceInpaintWeightMapName.IsNone() && (ForceInpaint.IsEmpty() || ForceInpaint.Num() != InOutTargetMesh.MaxVertexID())) // ForceInpaint array takes priority if valid
+		{
+			for (int32 Idx = 0; Idx < InOutTargetMesh.Attributes()->NumWeightLayers(); ++Idx)
+			{	
+				const FDynamicMeshWeightAttribute* WeightLayer = InOutTargetMesh.Attributes()->GetWeightLayer(Idx);
+				if (WeightLayer && WeightLayer->GetName() == ForceInpaintWeightMapName)
+				{
+					ForceInpaintLayer = WeightLayer;
+					break;
+				}
+			}
+		}
+
 		// For every vertex on the target mesh try to find the match on the source mesh using the distance and normal checks
-		ParallelFor(InOutTargetMesh.MaxVertexID(), [&](int32 VertexID)
+		ParallelFor(InOutTargetMesh.MaxVertexID(), [this, &InOutTargetMesh, &ForceInpaintLayer, &TargetBoneToIndex, &TargetSkinWeights](int32 VertexID)
 		{
 			if (Cancelled()) 
 			{
@@ -302,10 +317,19 @@ bool FTransferBoneWeights::TransferWeightsToMesh(FDynamicMesh3& InOutTargetMesh,
 			if (InOutTargetMesh.IsVertex(VertexID)) 
 			{
 				// check if we need to force the vertex to not have a match
-				if (ForceInpaint.Num() == InOutTargetMesh.MaxVertexID() && ForceInpaint[VertexID] > 0)
+				if (ForceInpaint.Num() == InOutTargetMesh.MaxVertexID() && ForceInpaint[VertexID] != 0)
 				{
 					return;
 				}
+				else if (ForceInpaintLayer != nullptr)
+				{
+					float Value;
+					ForceInpaintLayer->GetValue(VertexID, &Value);
+					if (Value != 0)
+					{
+						return;
+					}
+				}
 
 				const FVector3d Point = InOutTargetMesh.GetVertex(VertexID);
 				const FVector3f Normal = NormalThreshold >= 0 ? InOutTargetMesh.Attributes()->PrimaryNormals()->GetElement(VertexID) : FVector3f::Zero();
diff --git a/Engine/Plugins/Runtime/GeometryProcessing/Source/DynamicMesh/Public/Operations/TransferBoneWeights.h b/Engine/Plugins/Runtime/GeometryProcessing/Source/DynamicMesh/Public/Operations/TransferBoneWeights.h
--- a/Engine/Plugins/Runtime/GeometryProcessing/Source/DynamicMesh/Public/Operations/TransferBoneWeights.h
+++ b/Engine/Plugins/Runtime/GeometryProcessing/Source/DynamicMesh/Public/Operations/TransferBoneWeights.h
@@ -152,7 +152,7 @@ public:
 	bool bUseIntrinsicLaplacian = false;
 
 	/** 
-	 * Optional mask where if ForceInpaint[VertexID] > 0 we want to force the weights for the vertex to be computed  
+	 * Optional mask where if ForceInpaint[VertexID] != 0 we want to force the weights for the vertex to be computed  
 	 * automatically.
 	 * 
 	 * @note Only used when TransferMethod == ETransferBoneWeightsMethod::InpaintWeights.
@@ -160,6 +160,12 @@ public:
 	 */
 	TArray<float> ForceInpaint;
 
+	/** 
+	 * Alternatively, if the mask is stored as a target mesh weight map attribute, specify its name. Will be ignored if 
+	 * the ForceInpaint array is not empty and valid.
+	 */
+	FName ForceInpaintWeightMapName = NAME_None;
+
 
 	//
 	// Outputs
