Fixed a bug where overlay normals were used instead of the per-vertex normals.

#rnx

diff --git a/Engine/Plugins/Runtime/GeometryProcessing/Source/DynamicMesh/Private/Operations/TransferBoneWeights.cpp b/Engine/Plugins/Runtime/GeometryProcessing/Source/DynamicMesh/Private/Operations/TransferBoneWeights.cpp
--- a/Engine/Plugins/Runtime/GeometryProcessing/Source/DynamicMesh/Private/Operations/TransferBoneWeights.cpp
+++ b/Engine/Plugins/Runtime/GeometryProcessing/Source/DynamicMesh/Private/Operations/TransferBoneWeights.cpp
@@ -6,6 +6,7 @@
 #include "DynamicMesh/DynamicBoneAttribute.h"
 #include "DynamicMesh/DynamicMeshAttributeSet.h"
 #include "DynamicMesh/DynamicMeshAABBTree3.h"
+#include "DynamicMesh/MeshNormals.h"
 #include "Async/ParallelFor.h"
 #include "Util/ProgressCancel.h"
 #include "BoneIndices.h"
@@ -110,6 +111,11 @@ namespace TransferBoneWeightsLocals
 		}
 		return Attribute;
 	}
+
+	static FVector3f ToUENormal(const FVector3d& Normal)
+	{
+		return FVector3f((float)Normal.X, (float)Normal.Y, (float)Normal.Z);
+	}
 }
 
 FTransferBoneWeights::FTransferBoneWeights(const FDynamicMesh3* InSourceMesh, 
@@ -164,16 +170,13 @@ EOperationValidationResult FTransferBoneWeights::Validate()
 		return EOperationValidationResult::Failed_UnknownReason;
 	}
 
-	if (NormalThreshold >= 0 && !SourceMesh->Attributes()->PrimaryNormals())
-	{
-		return EOperationValidationResult::Failed_UnknownReason;
-	}
-
 	return EOperationValidationResult::Ok;
 }
 
 bool FTransferBoneWeights::TransferWeightsToMesh(FDynamicMesh3& InOutTargetMesh, const FName& InTargetProfileName)
 {	
+	using TransferBoneWeightsLocals::ToUENormal;
+
 	if (Validate() != EOperationValidationResult::Ok) 
 	{
 		return false;
@@ -189,9 +192,22 @@ bool FTransferBoneWeights::TransferWeightsToMesh(FDynamicMesh3& InOutTargetMesh,
 		return false; // the target mesh must have bone attributes
 	}
 
-	if (NormalThreshold >= 0 && !InOutTargetMesh.Attributes()->PrimaryNormals())
+	// If we need to compare normals, make sure both the target and the source meshes have per-vertex normals data
+	TUniquePtr<FMeshNormals> InternalTargetMeshNormals;
+	if (NormalThreshold >= 0)
 	{
-		return false; // the target mesh must have normal attributes if we are comparing source and target normals
+		if (!SourceMesh->HasVertexNormals() && !InternalSourceMeshNormals)
+		{
+			// only do this once for the source mesh in case of subsequent calls to the method
+			InternalSourceMeshNormals = MakeUnique<FMeshNormals>(SourceMesh);
+			InternalSourceMeshNormals->ComputeVertexNormals();
+		}
+
+		if (!InOutTargetMesh.HasVertexNormals())
+		{
+			InternalTargetMeshNormals = MakeUnique<FMeshNormals>(&InOutTargetMesh);
+			InternalTargetMeshNormals->ComputeVertexNormals();
+		}
 	}
 	
 	FDynamicMeshVertexSkinWeightsAttribute* TargetSkinWeights = TransferBoneWeightsLocals::GetOrCreateSkinWeightsAttribute(InOutTargetMesh, InTargetProfileName);
@@ -229,7 +245,7 @@ bool FTransferBoneWeights::TransferWeightsToMesh(FDynamicMesh3& InOutTargetMesh,
 
 	if (TransferMethod == ETransferBoneWeightsMethod::ClosestPointOnSurface)
 	{
-		ParallelFor(InOutTargetMesh.MaxVertexID(), [&](int32 VertexID)
+		ParallelFor(InOutTargetMesh.MaxVertexID(), [this, &InOutTargetMesh, &TargetBoneToIndex, &TargetSkinWeights, &InternalTargetMeshNormals](int32 VertexID)
 		{
 			if (Cancelled()) 
 			{
@@ -239,7 +255,16 @@ bool FTransferBoneWeights::TransferWeightsToMesh(FDynamicMesh3& InOutTargetMesh,
 			if (InOutTargetMesh.IsVertex(VertexID)) 
 			{
 				const FVector3d Point = InOutTargetMesh.GetVertex(VertexID);
-				const FVector3f Normal = NormalThreshold >= 0 ? InOutTargetMesh.Attributes()->PrimaryNormals()->GetElement(VertexID) : FVector3f::Zero();
+
+				FVector3f Normal = FVector3f::UnitY();
+				if (NormalThreshold >= 0) 
+				{
+					const bool bHasNormals = InOutTargetMesh.HasVertexNormals();
+					if (ensure(bHasNormals || InternalTargetMeshNormals))
+					{
+						Normal = bHasNormals ? InOutTargetMesh.GetVertexNormal(VertexID) : ToUENormal(InternalTargetMeshNormals->GetNormals()[VertexID]);
+					}
+				}
 
 				FBoneWeights Weights;
 				if (TransferWeightsToPoint(Weights, Point, TargetBoneToIndex.Get(), Normal))
@@ -307,7 +332,7 @@ bool FTransferBoneWeights::TransferWeightsToMesh(FDynamicMesh3& InOutTargetMesh,
 		}
 
 		// For every vertex on the target mesh try to find the match on the source mesh using the distance and normal checks
-		ParallelFor(InOutTargetMesh.MaxVertexID(), [this, &InOutTargetMesh, &ForceInpaintLayer, &TargetBoneToIndex, &TargetSkinWeights](int32 VertexID)
+		ParallelFor(InOutTargetMesh.MaxVertexID(), [this, &InOutTargetMesh, &ForceInpaintLayer, &TargetBoneToIndex, &TargetSkinWeights, &InternalTargetMeshNormals](int32 VertexID)
 		{
 			if (Cancelled()) 
 			{
@@ -332,7 +357,15 @@ bool FTransferBoneWeights::TransferWeightsToMesh(FDynamicMesh3& InOutTargetMesh,
 				}
 
 				const FVector3d Point = InOutTargetMesh.GetVertex(VertexID);
-				const FVector3f Normal = NormalThreshold >= 0 ? InOutTargetMesh.Attributes()->PrimaryNormals()->GetElement(VertexID) : FVector3f::Zero();
+				FVector3f Normal = FVector3f::UnitY();
+				if (NormalThreshold >= 0) 
+				{
+					const bool bHasNormals = InOutTargetMesh.HasVertexNormals();
+					if (ensure(bHasNormals || InternalTargetMeshNormals.IsValid()))
+					{
+						Normal = bHasNormals ? InOutTargetMesh.GetVertexNormal(VertexID) : ToUENormal(InternalTargetMeshNormals->GetNormals()[VertexID]);
+					}
+				}
 
 				FBoneWeights Weights;
 				if (TransferWeightsToPoint(Weights, Point, TargetBoneToIndex.Get(), Normal))
@@ -520,6 +553,8 @@ bool FTransferBoneWeights::TransferWeightsToPoint(UE::AnimationCore::FBoneWeight
 												  const TMap<FName, uint16>* TargetBoneToIndex,
 												  const FVector3f& InNormal)
 {	
+	using TransferBoneWeightsLocals::ToUENormal;
+
 	// Find the containing triangle and the barycentric coordinates of the closest point
 	int32 TriID; 
 	FVector3d Bary;
@@ -554,8 +589,16 @@ bool FTransferBoneWeights::TransferWeightsToPoint(UE::AnimationCore::FBoneWeight
 		bool bPassedNormalsCheck = true;
 		if (NormalThreshold >= 0)
 		{
-			FVector3f Normal0, Normal1, Normal2;
-			SourceMesh->Attributes()->PrimaryNormals()->GetTriElements(TriID, Normal0, Normal1, Normal2);
+			FVector3f Normal0 = FVector3f::UnitY();
+			FVector3f Normal1 = FVector3f::UnitY();
+			FVector3f Normal2 = FVector3f::UnitY();
+			const bool bHasSourceNormals = SourceMesh->HasVertexNormals();
+			if (ensure(bHasSourceNormals || InternalSourceMeshNormals.IsValid()))
+			{
+				Normal0 = bHasSourceNormals ? SourceMesh->GetVertexNormal(TriVertex[0]) : ToUENormal(InternalSourceMeshNormals->GetNormals()[TriVertex[0]]);
+				Normal1 = bHasSourceNormals ? SourceMesh->GetVertexNormal(TriVertex[1]) : ToUENormal(InternalSourceMeshNormals->GetNormals()[TriVertex[1]]);
+				Normal2 = bHasSourceNormals ? SourceMesh->GetVertexNormal(TriVertex[2]) : ToUENormal(InternalSourceMeshNormals->GetNormals()[TriVertex[2]]);
+			}
 
 			const FVector3f MatchedNormal = Normalized(BaryF[0]*Normal0 + BaryF[1]*Normal1 + BaryF[2]*Normal2);
 			const FVector3f InNormalNormalized = Normalized(InNormal);
@@ -590,7 +633,7 @@ bool FTransferBoneWeights::TransferWeightsToPoint(TArray<BoneIndexType>& OutBone
 												  const UE::Math::TVector<NormalVectorType>& InNormal)
 {
 	FBoneWeights BoneWeights;
-	if (!this->TransferWeightsToPoint(BoneWeights, FVector3d(InPoint), TargetBoneToIndex, FVector3f(InNormal)))
+	if (!this->TransferWeightsToPoint(BoneWeights, FVector3d(InPoint.X, InPoint.Y, InPoint.Z), TargetBoneToIndex, FVector3f((float)InNormal.X, (float)InNormal.Y, (float)InNormal.Z)))
 	{
 		return false;
 	}
diff --git a/Engine/Plugins/Runtime/GeometryProcessing/Source/DynamicMesh/Public/Operations/TransferBoneWeights.h b/Engine/Plugins/Runtime/GeometryProcessing/Source/DynamicMesh/Public/Operations/TransferBoneWeights.h
--- a/Engine/Plugins/Runtime/GeometryProcessing/Source/DynamicMesh/Public/Operations/TransferBoneWeights.h
+++ b/Engine/Plugins/Runtime/GeometryProcessing/Source/DynamicMesh/Public/Operations/TransferBoneWeights.h
@@ -16,6 +16,7 @@ namespace UE::Geometry
 	class FDynamicMesh3;
 	template<typename MeshType> class TMeshAABBTree3;
 	typedef TMeshAABBTree3<FDynamicMesh3> FDynamicMeshAABBTree3;
+	class FMeshNormals;
 }
 
 
@@ -191,6 +192,9 @@ protected:
 	/** If the caller doesn't pass BVH for the source mesh then we compute one. */
 	TUniquePtr<FDynamicMeshAABBTree3> InternalSourceBVH;
 
+	/** If the source mesh doesn't have per-vertex normals then compute them */
+	TUniquePtr<FMeshNormals> InternalSourceMeshNormals;
+
 public:
 	
 	/**
