ModelingTools: Fixed gap in gutter when baking MaterialID in the Bake Textures tool.

diff --git a/Engine/Plugins/Runtime/GeometryProcessing/Source/DynamicMesh/Private/Sampling/MeshMapBaker.cpp b/Engine/Plugins/Runtime/GeometryProcessing/Source/DynamicMesh/Private/Sampling/MeshMapBaker.cpp
--- a/Engine/Plugins/Runtime/GeometryProcessing/Source/DynamicMesh/Private/Sampling/MeshMapBaker.cpp
+++ b/Engine/Plugins/Runtime/GeometryProcessing/Source/DynamicMesh/Private/Sampling/MeshMapBaker.cpp
@@ -35,8 +35,8 @@ void ComputeGutterTexelsUsingFilterKernelCoverage(
 	const int64 TexelsPerPaddedTile = OccupancyMap.Tile.Num();
 	const int32 SamplesPerTexel = OccupancyMap.PixelSampler.Num();
 
-	TArray64<int8> TexelIsGutter;
-	TexelIsGutter.Init(true, Tile.Num());
+	TArray64<int8> TexelType;
+	TexelType.Init(FImageOccupancyMap::GutterTexel, Tile.Num());
 
 	// First pass: Identify Gutter texels as the texels whose centers are not covered by any filter kernels
 	// placed at interior/border sample points of texels in the PaddedTile.
@@ -99,7 +99,7 @@ void ComputeGutterTexelsUsingFilterKernelCoverage(
 					if (IsInFilterRegionEval(TexelDistance))
 					{
 						const int64 NbrTexelIndexInTile = Tile.GetIndexFromSourceCoords(NbrTexelCoordsInImage);
-						TexelIsGutter[NbrTexelIndexInTile] = false;
+						TexelType[NbrTexelIndexInTile] = OccupancyMap.IsInterior(NbrTexelIndexInPaddedTile) ? FImageOccupancyMap::InteriorTexel : FImageOccupancyMap::BorderTexel;
 					}
 				}
 			}
@@ -111,9 +111,10 @@ void ComputeGutterTexelsUsingFilterKernelCoverage(
 	{
 		const FVector2i GutterTexelCoordsInImage = Tile.GetSourceCoords(GutterTexelIndexInTile);
 		const int64 GutterTexelIndexInImage = Dimensions.GetIndex(GutterTexelCoordsInImage);
-		if (!TexelIsGutter[GutterTexelIndexInTile])
+		const int8 GutterTexelType = TexelType[GutterTexelIndexInTile];
+		if (GutterTexelType != FImageOccupancyMap::BorderTexel && GutterTexelType != FImageOccupancyMap::GutterTexel)
 		{
-			continue; // Not a Gutter texel, do nothing
+			continue; // Not a Border/Gutter texel, do nothing
 		}
 
 		for (int32 SampleIndexInTexel = 0; SampleIndexInTexel < SamplesPerTexel; ++SampleIndexInTexel)
@@ -133,11 +134,19 @@ void ComputeGutterTexelsUsingFilterKernelCoverage(
 				// search direction along which we can find the right texel, or better, use filter kernels to propogate
 				// known data into the gutter. Search :NearestInteriorGutterTexel for a related problem in the OccupancyMap
 				//
-				const FVector2d SampleUV = (FVector2d)OccupancyMap.TexelQueryUV[SampleIndexInPaddedTile];
+				const FVector2d SampleUV = static_cast<FVector2d>(OccupancyMap.TexelQueryUV[SampleIndexInPaddedTile]);
 				const FVector2i NearestTexelCoordsInImage = Dimensions.UVToCoords(SampleUV);
 				const int64 MappedTexelIndexInImage = Dimensions.GetIndex(NearestTexelCoordsInImage);
 				const TTuple<int64, int64> GutterTexel(GutterTexelIndexInImage, MappedTexelIndexInImage);
-				OccupancyMap.GutterTexels.Add(GutterTexel);
+
+				if (TexelType[GutterTexelIndexInTile] == FImageOccupancyMap::BorderTexel)
+				{
+					OccupancyMap.BorderTexels.Add(GutterTexel);
+				}
+				else if (TexelType[GutterTexelIndexInTile] == FImageOccupancyMap::GutterTexel)
+				{
+					OccupancyMap.GutterTexels.Add(GutterTexel);
+				}
 
 				break;
 			}
@@ -356,7 +365,9 @@ void FMeshMapBaker::Bake()
 	// Tile the image
 	FImageTiling Tiles(Dimensions, TileSize, TileSize);
 	const int32 NumTiles = Tiles.Num();
+	TArray<TArray64<TTuple<int64, int64>>> BorderTexelsPerTile;
 	TArray<TArray64<TTuple<int64, int64>>> GutterTexelsPerTile;
+	BorderTexelsPerTile.SetNum(NumTiles);
 	GutterTexelsPerTile.SetNum(NumTiles);
 
 	// WriteToOutputBuffer transfers local tile data (TileBuffer) to the image output buffer (FullImageTileBuffer).
@@ -433,7 +444,7 @@ void FMeshMapBaker::Bake()
 	};
 
 	FMeshMapBakerQueue OutputQueue(NumTiles);
-	ParallelFor(NumTiles, [this, &Tiles, &GutterTexelsPerTile, &OutputQueue, &WriteToOutputBuffer, &WriteToOutputBufferQueued, &ComputeCorrespondenceSample](int32 TileIdx)
+	ParallelFor(NumTiles, [this, &Tiles, &BorderTexelsPerTile, &GutterTexelsPerTile, &OutputQueue, &WriteToOutputBuffer, &WriteToOutputBufferQueued, &ComputeCorrespondenceSample](int32 TileIdx)
 	{
 		TRACE_CPUPROFILER_EVENT_SCOPE(FMeshMapBaker::Bake_EvalTile);
 
@@ -452,6 +463,7 @@ void FMeshMapBaker::Bake()
 		const auto GetTriangleIDFunc = [this](int32 TriangleID) { return FlatMesh.GetTriangleGroup(TriangleID); };
 		OccupancyMap.ClassifySamplesFromUVSpaceMesh(FlatMesh, GetTriangleIDFunc, TargetMeshUVCharts);
 		ComputeGutterTexelsUsingFilterKernelCoverage(OccupancyMap, Tile, Dimensions, FilterKernelSize, IsInFilterRegionEval);
+		BorderTexelsPerTile[TileIdx] = OccupancyMap.BorderTexels;
 		GutterTexelsPerTile[TileIdx] = OccupancyMap.GutterTexels;
 
 		const int64 NumTilePixels = Tile.Num();
@@ -674,7 +686,7 @@ void FMeshMapBaker::Bake()
 		FScopedDurationTimer WriteToGutterTimer(BakeAnalytics.WriteToGutterDuration);
 		
 		const int32 NumResults = BakeResults.Num();
-		ParallelFor(NumTiles, [this, &NumResults, &GutterTexelsPerTile](int32 TileIdx)
+		ParallelFor(NumTiles, [this, &NumResults, &BorderTexelsPerTile, &GutterTexelsPerTile](int32 TileIdx)
 		{
 			TRACE_CPUPROFILER_EVENT_SCOPE(FMeshMapBaker::Bake_WriteGutterPixels);
 
@@ -695,6 +707,27 @@ void FMeshMapBaker::Bake()
 				}
 			}
 
+			// For EAccumulateMode::Overwrite evaluators, Border pixels are gutter pixels. 
+			const int64 NumBorder = BorderTexelsPerTile[TileIdx].Num();
+			const TArray<int32>& EvaluatorOverwriteIds = EvaluatorIdsForMode(FMeshMapEvaluator::EAccumulateMode::Overwrite);
+			for (const int32 Idx : EvaluatorOverwriteIds)
+			{
+				const FMeshMapEvaluator::FEvaluationContext& Context = BakeContexts[Idx];
+				const int32 NumData = Context.DataLayout.Num();
+				const int32 ResultOffset = BakeOffsets[Idx];
+				for (int32 DataIdx = 0; DataIdx < NumData; ++DataIdx)
+				{
+					const int32 ResultIdx = ResultOffset + DataIdx;
+					for (int64 BorderIdx = 0; BorderIdx < NumBorder; ++BorderIdx)
+					{
+						int64 BorderPixelTo;
+						int64 BorderPixelFrom;
+						Tie(BorderPixelTo, BorderPixelFrom) = BorderTexelsPerTile[TileIdx][BorderIdx];
+						BakeResults[ResultIdx]->CopyPixel(BorderPixelFrom, BorderPixelTo);
+					}
+				}
+			}
+
 			BakeAnalytics.NumGutterPixels += NumGutter;
 		}, !bParallel ? EParallelForFlags::ForceSingleThread : EParallelForFlags::None);
 	}
diff --git a/Engine/Source/Runtime/GeometryCore/Public/Image/ImageOccupancyMap.h b/Engine/Source/Runtime/GeometryCore/Public/Image/ImageOccupancyMap.h
--- a/Engine/Source/Runtime/GeometryCore/Public/Image/ImageOccupancyMap.h
+++ b/Engine/Source/Runtime/GeometryCore/Public/Image/ImageOccupancyMap.h
@@ -42,10 +42,10 @@ public:
 	int32 GutterSize = 4;
 
 	// texel types
-	const int8 EmptyTexel = 0;
-	const int8 InteriorTexel = 1;
-	//const int8 BorderTexel = 2;
-	const int8 GutterTexel = 3;
+	static constexpr int8 EmptyTexel = 0;
+	static constexpr int8 InteriorTexel = 1;
+	static constexpr int8 BorderTexel = 2;
+	static constexpr int8 GutterTexel = 3;
 
 	/**
 	 * Classification of each sample in Tile.
@@ -77,8 +77,16 @@ public:
 	 */
 	TArray64<int32> TexelQueryUVChart;
 
-	/** Set of Gutter Texels. Pair is <LinearIndexOfGutterTexel, LinearIndexOfNearestInteriorTexel>, so
-	    Gutter can be filled by directly copying from source to target. */
+	/**
+	 * Set of Border Texels, Pair is <LinearIndexOfBorderTexel, LinearIndexOfNearestInteriorTexel>, so
+	 * Border can be filled by directly copying from source to target. 
+	 */
+	TArray64<TTuple<int64, int64>> BorderTexels;
+	
+	/**
+	 * Set of Gutter Texels. Pair is <LinearIndexOfGutterTexel, LinearIndexOfNearestInteriorTexel>, so
+	 * Gutter can be filled by directly copying from source to target.
+	 */
 	TArray64<TTuple<int64, int64>> GutterTexels;
 
 	using FGridSampler = TGridSampler<double>;
