make 'mesh to collision' tool's box detection algorithm look for 6 shared planes rather than 6 shared triangle normal directions + make it early-out if it finds more than 6. Fixes staircases being detected as boxes.

diff --git a/Engine/Plugins/Runtime/GeometryProcessing/Source/DynamicMesh/Private/ShapeApproximation/ShapeDetection3.cpp b/Engine/Plugins/Runtime/GeometryProcessing/Source/DynamicMesh/Private/ShapeApproximation/ShapeDetection3.cpp
--- a/Engine/Plugins/Runtime/GeometryProcessing/Source/DynamicMesh/Private/ShapeApproximation/ShapeDetection3.cpp
+++ b/Engine/Plugins/Runtime/GeometryProcessing/Source/DynamicMesh/Private/ShapeApproximation/ShapeDetection3.cpp
@@ -1,7 +1,6 @@
 // Copyright Epic Games, Inc. All Rights Reserved.
 
 #include "ShapeApproximation/ShapeDetection3.h"
-#include "Sampling/VectorSetAnalysis.h"
 #include "DynamicMesh/DynamicMesh3.h"
 #include "MeshQueries.h"
 #include "FitCapsule3.h"
@@ -95,7 +94,7 @@ bool UE::Geometry::IsSphereMesh(const FDynamicMesh3& Mesh, FSphere3d& SphereOut,
 
 
 
-bool UE::Geometry::IsBoxMesh(const FDynamicMesh3& Mesh, FOrientedBox3d& BoxOut, double AngleToleranceDeg)
+bool UE::Geometry::IsBoxMesh(const FDynamicMesh3& Mesh, FOrientedBox3d& BoxOut, double AngleToleranceDeg, double PlaneDistanceTolerance)
 {
 	// minimal box has at least 6 vertices
 	if (Mesh.VertexCount() < 6)
@@ -128,15 +127,55 @@ bool UE::Geometry::IsBoxMesh(const FDynamicMesh3& Mesh, FOrientedBox3d& BoxOut,
 		}
 	}
 
-	// cluster normals
-	FVectorSetAnalysis3d Vectors;
-	Vectors.Initialize(Mesh.TriangleIndicesItr(),
-		[&](int32 TriangleID) { return Mesh.GetTriNormal(TriangleID); },
-		Mesh.TriangleCount(), true);
+	// Greedily cluster up to six planes
+	constexpr int32 ExpectPlanes = 6;
+	FVector3d Normals[ExpectPlanes];
+	double Distances[ExpectPlanes];
+	int32 FoundPlanes = 0;
 
-	// A box should have precisely 6 normals, in 3 parallel pairs
-	Vectors.GreedyClusterVectors(AngleToleranceDeg);
-	if (Vectors.NumClusters() != 6)		
+	// Helper to bin a new position+normal into a plane cluster, or return -1 if no valid cluster is found
+	auto NormalMatch = [&Normals, &Distances, &FoundPlanes, ExpectPlanes, ParallelDotTolerance, PlaneDistanceTolerance](FVector3d Pos, FVector3d Normal) -> int32
+	{
+		int32 Idx = 0;
+		for (; Idx < FoundPlanes; ++Idx)
+		{
+			if (Normals[Idx].Dot(Normal) > ParallelDotTolerance)
+			{
+				double Distance = Pos.Dot(Normal);
+				if (FMath::IsNearlyEqual(Distance, Distances[Idx], PlaneDistanceTolerance))
+				{
+					return Idx;
+				}
+				else
+				{
+					// A box cannot have multiple parallel planes with the same normal
+					// (note opposite faces will be parallel but with opposite normals)
+					return INDEX_NONE;
+				}
+			}
+		}
+		if (FoundPlanes < ExpectPlanes)
+		{
+			Normals[Idx] = Normal;
+			Distances[Idx] = Pos.Dot(Normal);
+			FoundPlanes++;
+			return Idx;
+		}
+		return INDEX_NONE;
+	};
+
+	for (int32 TriID : Mesh.TriangleIndicesItr())
+	{
+		FIndex3i Tri = Mesh.GetTriangle(TriID);
+		FVector3d Vert = Mesh.GetVertex(Tri.A);
+		FVector3d Normal = Mesh.GetTriNormal(TriID);
+		int32 PlaneIdx = NormalMatch(Vert, Normal);
+		if (PlaneIdx == INDEX_NONE)
+		{
+			return false;
+		}
+	}
+	if (FoundPlanes != ExpectPlanes)
 	{
 		return false;
 	}
@@ -150,12 +189,12 @@ bool UE::Geometry::IsBoxMesh(const FDynamicMesh3& Mesh, FOrientedBox3d& BoxOut,
 	{
 		if (bDone[k]) continue;
 
-		FVector3d Normal0 = Vectors.ClusterVectors[k];
+		FVector3d Normal0 = Normals[k];
 		int32 ParallelPair = -1;
 
 		for (int32 j = k + 1; j < 6; ++j)
 		{
-			double Dot = Normal0.Dot(Vectors.ClusterVectors[j]);
+			double Dot = Normal0.Dot(Normals[j]);
 			if (FMathd::Abs(Dot) > PerpDotTolerance)
 			{
 				if (Dot > -ParallelDotTolerance)		// if dot is not zero, it needs to be -1
@@ -176,14 +215,14 @@ bool UE::Geometry::IsBoxMesh(const FDynamicMesh3& Mesh, FOrientedBox3d& BoxOut,
 		}
 	}
 			
-	// if we found the 3 unique axes, it's a box and we know it's orientation, so just fit minimal 
+	// if we found the 3 unique axes, it's a box and we know its orientation, so just fit minimal 
 	// container aligned to box axes, and shift center point to center of that oriented-AABB
 	if (UniqueCount == 3)
 	{
 		// would be nice to cycle these so that X = most-aligned-with-X, etc, or longest?
-		FVector3d X = Vectors.ClusterVectors[UniqueAxes[0]];
-		FVector3d Y = Vectors.ClusterVectors[UniqueAxes[1]];
-		FVector3d Z = Vectors.ClusterVectors[UniqueAxes[2]];
+		FVector3d X = Normals[UniqueAxes[0]];
+		FVector3d Y = Normals[UniqueAxes[1]];
+		FVector3d Z = Normals[UniqueAxes[2]];
 		// compute AABB in the frame of the box
 		FQuaterniond Rotation(FMatrix3d(X, Y, Z, false));
 		FMatrix3d UnorientRotation = Rotation.Inverse().ToRotationMatrix();
diff --git a/Engine/Plugins/Runtime/GeometryProcessing/Source/DynamicMesh/Public/ShapeApproximation/ShapeDetection3.h b/Engine/Plugins/Runtime/GeometryProcessing/Source/DynamicMesh/Public/ShapeApproximation/ShapeDetection3.h
--- a/Engine/Plugins/Runtime/GeometryProcessing/Source/DynamicMesh/Public/ShapeApproximation/ShapeDetection3.h
+++ b/Engine/Plugins/Runtime/GeometryProcessing/Source/DynamicMesh/Public/ShapeApproximation/ShapeDetection3.h
@@ -25,12 +25,13 @@ namespace UE
 
 		/**
 		 * Detect if input Mesh is a meshed box, and if so return analytic box in BoxOut.
-		 * Clusters face normals, looking to find 6 unique normals grouped into 3 opposite-direction pairs.
+		 * Clusters face planes, looking to find 6 unique planes with normals grouped into 3 opposite-direction pairs.
 		 * If this configuraiton is found, computing minimal box is trivial.
 		 * @param AngleToleranceDeg normals are allowed to deviate by this amount and still be considered coplanar
+		 * @param PlaneDistanceTolerance planes with the same normal are allowed to be this far apart and still be considered coplanar
 		 * @return true if mesh is a Box and BoxOut is initialized
 		 */
-		bool DYNAMICMESH_API IsBoxMesh(const FDynamicMesh3& Mesh, FOrientedBox3d& BoxOut, double AngleToleranceDeg = 0.1);
+		bool DYNAMICMESH_API IsBoxMesh(const FDynamicMesh3& Mesh, FOrientedBox3d& BoxOut, double AngleToleranceDeg = 0.1, double PlaneDistanceTolerance = 1e-2);
 
 
 		/**
