Stop MeshRegionGraph's small region merging algorithm from repeatedly merging regions that have already grown larger than the minimum region size.

Fixes an issue where the Generate PolyGroups algorithm would randomly create very large, oddly-shaped regions in cases where the settings would normally generate many single-triangle regions.

diff --git a/Engine/Plugins/Runtime/GeometryProcessing/Source/DynamicMesh/Private/Parameterization/MeshRegionGraph.cpp b/Engine/Plugins/Runtime/GeometryProcessing/Source/DynamicMesh/Private/Parameterization/MeshRegionGraph.cpp
--- a/Engine/Plugins/Runtime/GeometryProcessing/Source/DynamicMesh/Private/Parameterization/MeshRegionGraph.cpp
+++ b/Engine/Plugins/Runtime/GeometryProcessing/Source/DynamicMesh/Private/Parameterization/MeshRegionGraph.cpp
@@ -159,6 +159,10 @@ bool FMeshRegionGraph::MergeSmallRegionsPass(int32 SmallThreshold,
 		int32 SmallRegionIdx = SmallRegions[j];
 		if (IsRegion(SmallRegionIdx) == false) continue;
 		const FRegion& SmallRegion = Regions[SmallRegionIdx];
+		if (SmallRegion.Triangles.Num() > SmallThreshold)
+		{
+			continue;
+		}
 
 		TArray<FMatch> Matches;
 		for (const FNeighbour& Nbr : SmallRegion.Neighbours)
