Make more collision geometry types support Min Thickness in the Mesh to Collision tool + Make the tool disable the slider for shapes that don't support it (Convex Decomp and Level Sets)

Add Min Thickness support to the Mesh Convex Hull class

diff --git a/Engine/Plugins/Experimental/MeshModelingToolsetExp/Source/MeshModelingToolsExp/Public/Physics/SetCollisionGeometryTool.h b/Engine/Plugins/Experimental/MeshModelingToolsetExp/Source/MeshModelingToolsExp/Public/Physics/SetCollisionGeometryTool.h
--- a/Engine/Plugins/Experimental/MeshModelingToolsetExp/Source/MeshModelingToolsExp/Public/Physics/SetCollisionGeometryTool.h
+++ b/Engine/Plugins/Experimental/MeshModelingToolsetExp/Source/MeshModelingToolsExp/Public/Physics/SetCollisionGeometryTool.h
@@ -125,7 +125,8 @@ public:
 	UPROPERTY(EditAnywhere, Category = Options, meta = (UIMin = "1", UIMax = "100", ClampMin = "1", ClampMax = "9999999", EditCondition = "bEnableMaxCount"))
 	int32 MaxCount = 50;
 
-	UPROPERTY(EditAnywhere, Category = Options, AdvancedDisplay)
+	/** Generated collision shapes will be expanded if they are smaller than this in any dimension. Not supported for Level Sets or Convex Decompositions (Convex Hulls with more than one hull per mesh). */
+	UPROPERTY(EditAnywhere, Category = Options, AdvancedDisplay, meta = (ClampMin = "0", UIMax = "10", EditCondition = "GeometryType != ECollisionGeometryType::LevelSets && (GeometryType != ECollisionGeometryType::ConvexHulls || MaxHullsPerMesh == 1)"))
 	float MinThickness = 0.01;
 
 	UPROPERTY(EditAnywhere, Category = AutoDetect)
diff --git a/Engine/Plugins/Runtime/GeometryProcessing/Source/DynamicMesh/Private/Operations/MeshConvexHull.cpp b/Engine/Plugins/Runtime/GeometryProcessing/Source/DynamicMesh/Private/Operations/MeshConvexHull.cpp
--- a/Engine/Plugins/Runtime/GeometryProcessing/Source/DynamicMesh/Private/Operations/MeshConvexHull.cpp
+++ b/Engine/Plugins/Runtime/GeometryProcessing/Source/DynamicMesh/Private/Operations/MeshConvexHull.cpp
@@ -75,51 +75,89 @@ bool FMeshConvexHull::SimplifyHull(FDynamicMesh3& HullMesh, int32 MaxTargetFaceC
 
 bool FMeshConvexHull::Compute_FullMesh(FProgressCancel* Progress)
 {
-	FConvexHull3d HullCompute;
-	HullCompute.Progress = Progress;
-	bool bOK = HullCompute.Solve(Mesh->MaxVertexID(),
-		[this](int32 Index) { return Mesh->GetVertex(Index); },
-		[this](int32 Index) { return Mesh->IsVertex(Index); });
-	if (!bOK)
-	{
-		return false;
-	}
+	return Compute_Helper(Progress, Mesh->MaxVertexID(), [this](int32 Index) { return Mesh->GetVertex(Index); }, [this](int32 Index) { return Mesh->IsVertex(Index); });
+}
 
-	TMap<int32, int32> HullVertMap;
+bool FMeshConvexHull::Compute_VertexSubset(FProgressCancel* Progress)
+{
+	return Compute_Helper(Progress, VertexSet.Num(), [this](int32 Index) { return Mesh->GetVertex(VertexSet[Index]); }, [](int32 Index) { return true; });
+}
 
-	ConvexHull = FDynamicMesh3(EMeshComponents::None);
-	HullCompute.GetTriangles([&](FIndex3i Triangle)
+bool FMeshConvexHull::Compute_Helper(FProgressCancel* Progress, int32 MaxVertexIndex, TFunctionRef<FVector3d(int32)> GetVertex, TFunctionRef<bool(int32)> IsVertex, bool bTestMinDimension)
+{
+	if (bTestMinDimension && MinDimension > 0)
 	{
-		for (int32 j = 0; j < 3; ++j)
+		// Use extreme points to quickly get an aligned box, which should align to very thin shapes
+		TExtremePoints3<double> ExtremePoints(MaxVertexIndex, GetVertex, IsVertex);
+		// Fill out any missing vectors in the basis
+		switch (ExtremePoints.Dimension)
 		{
-			int32 Index = Triangle[j];
-			if (HullVertMap.Contains(Index) == false)
+		case 0:
+			ExtremePoints.Basis[0] = FVector(1, 0, 0);
+			ExtremePoints.Basis[1] = FVector(0, 1, 0);
+			ExtremePoints.Basis[2] = FVector(0, 0, 1);
+			break;
+		case 1:
+			VectorUtil::MakePerpVectors(ExtremePoints.Basis[0], ExtremePoints.Basis[1], ExtremePoints.Basis[2]);
+			break;
+		case 2:
+			ExtremePoints.Basis[2] = ExtremePoints.Basis[0].Cross(ExtremePoints.Basis[1]);
+			break;
+		}
+		FAxisAlignedBox3d BasisBox;
+		for (int32 Index = 0; Index < MaxVertexIndex; ++Index)
+		{
+			if (!IsVertex(Index))
 			{
-				FVector3d OrigPos = Mesh->GetVertex(Index);
-				int32 NewVID = ConvexHull.AppendVertex(OrigPos);
-				HullVertMap.Add(Index, NewVID);
-				Triangle[j] = NewVID;
+				continue;
 			}
-			else
+			FVector3d Vertex = GetVertex(Index);
+			FVector3d InBasis(
+				ExtremePoints.Basis[0].Dot(Vertex),
+				ExtremePoints.Basis[1].Dot(Vertex),
+				ExtremePoints.Basis[2].Dot(Vertex)
+			);
+			BasisBox.Contain(InBasis);
+		}
+		int32 NumSmallDims = 0;
+		int32 SmallDims[3]{ -1,-1,-1 };
+		FVector3d ThickenVecs[3];
+		FVector3d BasisDims = BasisBox.Diagonal();
+		for (int32 DimIdx = 0; DimIdx < 3; ++DimIdx)
+		{
+			if (BasisDims[DimIdx] < MinDimension)
 			{
-				Triangle[j] = HullVertMap[Index];
+				int32 SmallIdx = NumSmallDims++;
+				SmallDims[SmallIdx] = DimIdx;
+				ThickenVecs[SmallIdx] = .5 * (MinDimension - BasisDims[DimIdx]) * ExtremePoints.Basis[DimIdx];
 			}
 		}
+		if (NumSmallDims > 0)
+		{
+			TArray<FVector3d> ThickenedHullVertices;
+			ThickenedHullVertices.Reserve(MaxVertexIndex * (NumSmallDims * 2));
+			for (int32 Index = 0; Index < MaxVertexIndex; ++Index)
+			{
+				if (!IsVertex(Index))
+				{
+					continue;
+				}
+				FVector3d Vertex = GetVertex(Index);
+				for (int32 SmallDimIdx = 0; SmallDimIdx < NumSmallDims; ++SmallDimIdx)
+				{
+					ThickenedHullVertices.Add(Vertex + ThickenVecs[SmallDimIdx]);
+					ThickenedHullVertices.Add(Vertex - ThickenVecs[SmallDimIdx]);
+				}
+			}
+			return Compute_Helper(Progress, ThickenedHullVertices.Num(), 
+				[&ThickenedHullVertices](int32 Index) { return ThickenedHullVertices[Index]; },
+				[](int32) { return true; }, false);
+		}
+	}
 
-		ConvexHull.AppendTriangle(Triangle);
-	});
-
-	return true;
-}
-
-
-
-bool FMeshConvexHull::Compute_VertexSubset(FProgressCancel* Progress)
-{
 	FConvexHull3d HullCompute;
 	HullCompute.Progress = Progress;
-	bool bOK = HullCompute.Solve(VertexSet.Num(),
-		[this](int32 Index) { return Mesh->GetVertex(VertexSet[Index]); });
+	bool bOK = HullCompute.Solve(MaxVertexIndex, GetVertex, IsVertex);
 	if (!bOK)
 	{
 		return false;
@@ -135,7 +173,7 @@ bool FMeshConvexHull::Compute_VertexSubset(FProgressCancel* Progress)
 			int32 Index = Triangle[j];
 			if (HullVertMap.Contains(Index) == false)
 			{
-				FVector3d OrigPos = Mesh->GetVertex(VertexSet[Index]);
+				FVector3d OrigPos = GetVertex(Index);
 				int32 NewVID = ConvexHull.AppendVertex(OrigPos);
 				HullVertMap.Add(Index, NewVID);
 				Triangle[j] = NewVID;
diff --git a/Engine/Plugins/Runtime/GeometryProcessing/Source/DynamicMesh/Private/ShapeApproximation/MeshSimpleShapeApproximation.cpp b/Engine/Plugins/Runtime/GeometryProcessing/Source/DynamicMesh/Private/ShapeApproximation/MeshSimpleShapeApproximation.cpp
--- a/Engine/Plugins/Runtime/GeometryProcessing/Source/DynamicMesh/Private/ShapeApproximation/MeshSimpleShapeApproximation.cpp
+++ b/Engine/Plugins/Runtime/GeometryProcessing/Source/DynamicMesh/Private/ShapeApproximation/MeshSimpleShapeApproximation.cpp
@@ -166,7 +166,7 @@ struct FSimpleShapeFitsResult
 
 
 static void ComputeSimpleShapeFits(const FDynamicMesh3& Mesh,
-									bool bSphere, bool bBox, bool bCapsule, bool bConvex, bool bUseExactComputationForBox, 
+									bool bSphere, bool bBox, bool bCapsule, double MinDimension, bool bUseExactComputationForBox,
 								   FSimpleShapeFitsResult& FitResult,
 								   FProgressCancel* Progress = nullptr)
 {
@@ -191,6 +191,11 @@ static void ComputeSimpleShapeFits(const FDynamicMesh3& Mesh,
 		{
 			FitResult.bHaveBox = true;
 			MinBoxCalc.GetResult(FitResult.Box);
+			double MinHalfDimension = MinDimension * .5;
+			for (int32 SubIdx = 0; SubIdx < 3; ++SubIdx)
+			{
+				FitResult.Box.Extents[SubIdx] = FMath::Max(MinHalfDimension, FitResult.Box.Extents[SubIdx]);
+			}
 		}
 	}
 
@@ -204,6 +209,7 @@ static void ComputeSimpleShapeFits(const FDynamicMesh3& Mesh,
 		{
 			FitResult.bHaveSphere = true;
 			MinSphereCalc.GetResult(FitResult.Sphere);
+			FitResult.Sphere.Radius = FMath::Max(MinDimension * .5, FitResult.Sphere.Radius);
 		}
 	}
 
@@ -212,17 +218,10 @@ static void ComputeSimpleShapeFits(const FDynamicMesh3& Mesh,
 	{
 		FitResult.bHaveCapsule = TFitCapsule3<double>::Solve(FromLinear.Num(),
 			[&](int32 Index) { return Mesh.GetVertex(FromLinear[Index]); }, FitResult.Capsule);
-	}
-
-	// todo: once we have computed convex we can use it to compute box
-	FitResult.bHaveConvex = false;
-	if (bConvex)
-	{
-		FMeshConvexHull Hull(&Mesh);
-		if (Hull.Compute(Progress))
+		if (FitResult.bHaveCapsule)
 		{
-			FitResult.bHaveConvex = true;
-			FitResult.Convex = MoveTemp(Hull.ConvexHull);
+			FitResult.Capsule.Radius = FMath::Max(MinDimension * .5, FitResult.Capsule.Radius);
+			// Note: No need to clamp Length based on MinDimension; a capsule's min dimension can't be along its length since a capsule w/ length of zero is still a sphere ...
 		}
 	}
 }
@@ -250,6 +249,10 @@ void FMeshSimpleShapeApproximation::Generate_AlignedBoxes(FSimpleShapeSet3d& Sha
 		{
 			FBoxShape3d NewBox;
 			NewBox.Box = FOrientedBox3d(Bounds);
+			for (int32 SubIdx = 0; SubIdx < 3; ++SubIdx)
+			{
+				NewBox.Box.Extents[SubIdx] = FMath::Max(MinDimension * .5, NewBox.Box.Extents[SubIdx]);
+			}
 
 			GeometryLock.Lock();
 			ShapeSetOut.Boxes.Add(NewBox);
@@ -274,7 +277,7 @@ void FMeshSimpleShapeApproximation::Generate_OrientedBoxes(FSimpleShapeSet3d& Sh
 
 		const FDynamicMesh3& SourceMesh = *SourceMeshes[idx];
 		FSimpleShapeFitsResult FitResult;
-		ComputeSimpleShapeFits(SourceMesh, false, true, false, false, bUseExactComputationForBox, FitResult, Progress);
+		ComputeSimpleShapeFits(SourceMesh, false, true, false, MinDimension, bUseExactComputationForBox, FitResult, Progress);
 
 		if (Progress && Progress->Cancelled())
 		{
@@ -304,7 +307,7 @@ void FMeshSimpleShapeApproximation::Generate_MinimalSpheres(FSimpleShapeSet3d& S
 
 		const FDynamicMesh3& SourceMesh = *SourceMeshes[idx];
 		FSimpleShapeFitsResult FitResult;
-		ComputeSimpleShapeFits(SourceMesh, true, false, false, false, bUseExactComputationForBox, FitResult);
+		ComputeSimpleShapeFits(SourceMesh, true, false, false, MinDimension, bUseExactComputationForBox, FitResult);
 
 		if (FitResult.bHaveSphere)
 		{
@@ -327,7 +330,7 @@ void FMeshSimpleShapeApproximation::Generate_Capsules(FSimpleShapeSet3d& ShapeSe
 
 		const FDynamicMesh3& SourceMesh = *SourceMeshes[idx];
 		FSimpleShapeFitsResult FitResult;
-		ComputeSimpleShapeFits(SourceMesh, false, false, true, false, bUseExactComputationForBox, FitResult);
+		ComputeSimpleShapeFits(SourceMesh, false, false, true, MinDimension, bUseExactComputationForBox, FitResult);
 
 		if (FitResult.bHaveCapsule)
 		{
@@ -355,6 +358,7 @@ void FMeshSimpleShapeApproximation::Generate_ConvexHulls(FSimpleShapeSet3d& Shap
 
 		Hull.bPostSimplify = bSimplifyHulls;
 		Hull.MaxTargetFaceCount = HullTargetFaceCount;
+		Hull.MinDimension = MinDimension;
 
 		if (Hull.Compute(Progress))
 		{
@@ -529,9 +533,14 @@ void FMeshSimpleShapeApproximation::Generate_MinVolume(FSimpleShapeSet3d& ShapeS
 		const FDynamicMesh3& SourceMesh = *SourceMeshes[idx];
 
 		FOrientedBox3d AlignedBox = FOrientedBox3d(SourceMesh.GetBounds());
+		double MinHalfDimension = MinDimension * .5;
+		for (int32 SubIdx = 0; SubIdx < 3; ++SubIdx)
+		{
+			AlignedBox.Extents[SubIdx] = FMath::Max(MinHalfDimension, AlignedBox.Extents[SubIdx]);
+		}
 
 		FSimpleShapeFitsResult FitResult;
-		ComputeSimpleShapeFits(SourceMesh, true, true, true, false, bUseExactComputationForBox, FitResult);
+		ComputeSimpleShapeFits(SourceMesh, true, true, true, MinDimension, bUseExactComputationForBox, FitResult);
 
 		double Volumes[4];
 		Volumes[0] = AlignedBox.Volume();
diff --git a/Engine/Plugins/Runtime/GeometryProcessing/Source/DynamicMesh/Public/Operations/MeshConvexHull.h b/Engine/Plugins/Runtime/GeometryProcessing/Source/DynamicMesh/Public/Operations/MeshConvexHull.h
--- a/Engine/Plugins/Runtime/GeometryProcessing/Source/DynamicMesh/Public/Operations/MeshConvexHull.h
+++ b/Engine/Plugins/Runtime/GeometryProcessing/Source/DynamicMesh/Public/Operations/MeshConvexHull.h
@@ -32,6 +32,9 @@ public:
 	/** Target triangle count of the output Convex Hull */
 	int32 MaxTargetFaceCount = 0;
 
+	/** Minimum extent along the shortest dimension; if greater than zero, the hull may be expanded if it would be too thin */
+	double MinDimension = 0;
+
 	/** Output convex hull */
 	FDynamicMesh3 ConvexHull;
 
@@ -75,6 +78,8 @@ public:
 protected:
 	bool Compute_FullMesh(FProgressCancel* Progress);
 	bool Compute_VertexSubset(FProgressCancel* Progress);
+private:
+	bool Compute_Helper(FProgressCancel* Progress, int32 MaxVertexIndex, TFunctionRef<FVector3d(int32)> GetVertex, TFunctionRef<bool(int32)> IsVertex, bool bTestMinDimension = true);
 };
 
 
