MeshModelingTools: Fix group edges not being properly scalable (only the endpoints moved properly before).

diff --git a/Engine/Plugins/Runtime/GeometryProcessing/Source/DynamicMesh/Private/Operations/GroupTopologyDeformer.cpp b/Engine/Plugins/Runtime/GeometryProcessing/Source/DynamicMesh/Private/Operations/GroupTopologyDeformer.cpp
--- a/Engine/Plugins/Runtime/GeometryProcessing/Source/DynamicMesh/Private/Operations/GroupTopologyDeformer.cpp
+++ b/Engine/Plugins/Runtime/GeometryProcessing/Source/DynamicMesh/Private/Operations/GroupTopologyDeformer.cpp
@@ -363,6 +363,11 @@ void FGroupTopologyDeformer::UpdateSolution(FDynamicMesh3* TargetMesh, const TFu
 		FSegment3d Seg(A, B);
 		for (int k = 1; k < NumVerts - 1; ++k)
 		{
+			if (HandleVertices.Contains(Span.Vertices[k]))
+			{
+				// If we're explicitly moving this vertex, we should keep the explicitly set location.
+				continue;
+			}
 			FVector3d NewPos = Seg.PointBetween(Encoding.Vertices[k].T);
 			NewPos += Encoding.Vertices[k].Delta;
 			TargetMesh->SetVertex(Span.Vertices[k], NewPos);
