Make group topology choose an average edge position instead of a midpoint when the edge is a self-loop

diff --git a/Engine/Plugins/Runtime/GeometryProcessing/Source/DynamicMesh/Private/GroupTopology.cpp b/Engine/Plugins/Runtime/GeometryProcessing/Source/DynamicMesh/Private/GroupTopology.cpp
--- a/Engine/Plugins/Runtime/GeometryProcessing/Source/DynamicMesh/Private/GroupTopology.cpp
+++ b/Engine/Plugins/Runtime/GeometryProcessing/Source/DynamicMesh/Private/GroupTopology.cpp
@@ -377,6 +377,31 @@ double FGroupTopology::GetEdgeArcLength(int32 GroupEdgeID, TArray<double>* PerVe
 }
 
 
+FVector3d FGroupTopology::GetEdgeAveragePosition(int32 GroupEdgeID) const
+{
+	const TArray<int32>& Vertices = GetGroupEdgeVertices(GroupEdgeID);
+
+	double WtSum = 0;
+	FVector3d Sum = FVector3d::ZeroVector;
+	FVector LastV = Mesh->GetVertex(Vertices[0]);
+	for (int32 Idx = 0; Idx + 1 < Vertices.Num(); ++Idx)
+	{
+		FVector3d A = LastV;
+		FVector3d B = Mesh->GetVertex(Vertices[Idx + 1]);
+		// weight edge centers by edge lengths
+		double Wt = FVector3d::Distance(A, B);
+		Sum += Wt * .5 * (A + B);
+		WtSum += Wt;
+		LastV = B;
+	}
+	if (WtSum < FMathd::Epsilon) // if edges were all ~ zero length, just use one of the vertices
+	{
+		return LastV;
+	}
+	return Sum / WtSum;
+}
+
+
 FVector3d FGroupTopology::GetEdgeMidpoint(int32 GroupEdgeID, double* ArcLengthOut, TArray<double>* PerVertexLengthsOut) const
 {
 	check(GroupEdgeID >= 0 && GroupEdgeID < Edges.Num());
@@ -805,7 +830,15 @@ FFrame3d FGroupTopology::GetSelectionFrame(const FGroupTopologySelection& Select
 			StartFrame.ConstrainedAlignAxis(0, Tangent, StartFrame.Z());
 		}
 
-		StartFrame.Origin = GetEdgeMidpoint(GroupEdgeID);
+		const TArray<int32>& Vertices = GetGroupEdgeVertices(GroupEdgeID);
+		if (Vertices[0] == Vertices.Last())
+		{
+			StartFrame.Origin = GetEdgeAveragePosition(GroupEdgeID);
+		}
+		else
+		{
+			StartFrame.Origin = GetEdgeMidpoint(GroupEdgeID);
+		}
 		return StartFrame;
 	}
 	if (NumCorners == 1)
@@ -829,8 +862,16 @@ FFrame3d FGroupTopology::GetSelectionFrame(const FGroupTopologySelection& Select
 	{
 		const FGroupEdge& Edge = Edges[EdgeID];
 		FVector3d StartPos = Mesh->GetVertex(Edge.Span.Vertices[0]);
-		FVector3d EndPos = Mesh->GetVertex(Edge.Span.Vertices[Edge.Span.Vertices.Num() - 1]);
-		AccumulatedOrigin +=  0.5*(StartPos + EndPos);
+		FVector3d EndPos = Mesh->GetVertex(Edge.Span.Vertices.Last());
+		// special case self-loops to be consistent with the single-edge loop special case, above
+		if (Edge.Span.Vertices[0] == Edge.Span.Vertices.Last())
+		{
+			AccumulatedOrigin += GetEdgeAveragePosition(EdgeID);
+		}
+		else
+		{
+			AccumulatedOrigin += 0.5 * (StartPos + EndPos);
+		}
 		AccumulatedNormal += FVector3d::UnitZ();
 		AccumCount++;
 	}
diff --git a/Engine/Plugins/Runtime/GeometryProcessing/Source/DynamicMesh/Public/GroupTopology.h b/Engine/Plugins/Runtime/GeometryProcessing/Source/DynamicMesh/Public/GroupTopology.h
--- a/Engine/Plugins/Runtime/GeometryProcessing/Source/DynamicMesh/Public/GroupTopology.h
+++ b/Engine/Plugins/Runtime/GeometryProcessing/Source/DynamicMesh/Public/GroupTopology.h
@@ -336,6 +336,9 @@ public:
 	/** @return arclength midpoint of edge, and optionally arclength info */
 	FVector3d GetEdgeMidpoint(int32 GroupEdgeID, double* ArcLengthOut = nullptr, TArray<double>* PerVertexLengthsOut = nullptr) const;
 
+	/** @return average position of edge (may not be on edge for curved edges) */
+	FVector3d GetEdgeAveragePosition(int32 GroupEdgeID) const;
+
 	/** Add all the groups connected to the given Corner to the GroupsOut list */
 	void FindCornerNbrGroups(int CornerID, TArray<int>& GroupsOut) const;
 	/** Add all the groups connected to the given Corners to the GroupsOut list */
