add an angle-based metric for controlling sphere approximation acceptance, and clarify that the existing radial tolerance is vs the best distance possible for an edge of that length rather than the actual distance to the sphere surface. The angle tolerance gives an intuitive control for how coarse a tessellation could be accepted as a sphere.
Set a default value of 72 degrees max angle, so that e.g. cone volumes aren't recognized as spheres.

diff --git a/Engine/Plugins/Runtime/GeometryProcessing/Source/DynamicMesh/Private/ShapeApproximation/ShapeDetection3.cpp b/Engine/Plugins/Runtime/GeometryProcessing/Source/DynamicMesh/Private/ShapeApproximation/ShapeDetection3.cpp
--- a/Engine/Plugins/Runtime/GeometryProcessing/Source/DynamicMesh/Private/ShapeApproximation/ShapeDetection3.cpp
+++ b/Engine/Plugins/Runtime/GeometryProcessing/Source/DynamicMesh/Private/ShapeApproximation/ShapeDetection3.cpp
@@ -8,7 +8,7 @@
 
 using namespace UE::Geometry;
 
-bool UE::Geometry::IsSphereMesh(const FDynamicMesh3& Mesh, FSphere3d& SphereOut, double RelativeDeviationTol)
+bool UE::Geometry::IsSphereMesh(const FDynamicMesh3& Mesh, FSphere3d& SphereOut, double RelativeDeviationTol, double MaxAngleRangeDegrees)
 {
 	// assume that we aren't going to count it as a sphere unless it has 4 slices/sections, which means at least 10 vertices
 	if (Mesh.VertexCount() < 10)
@@ -74,11 +74,20 @@ bool UE::Geometry::IsSphereMesh(const FDynamicMesh3& Mesh, FSphere3d& SphereOut,
 	// computed with formula sagitta = r - sqrt(r*r - l*l), where l = chordlen/2
 	double UseRadius = SphereOut.Radius;
 	double DeviationTol = 2.0 * UseRadius * RelativeDeviationTol;
+	double CosAngleTolerance = FMathd::Cos(FMathd::DegToRad * MaxAngleRangeDegrees);
 	for (int32 EdgeID : Mesh.EdgeIndicesItr())
 	{
 		FVector3d A, B;
 		Mesh.GetEdgeV(EdgeID, A, B);
 
+		// if a single edge spans too wide an angular range, the shape is too coarsely tesselated to be considered a sphere
+		FVector3d ToA = A - SphereOut.Center;
+		FVector3d ToB = B - SphereOut.Center;
+		if (ToA.Dot(ToB) <= CosAngleTolerance)
+		{
+			return false;
+		}
+
 		double HalfChordLen = Distance(A, B) * 0.5;
 		double MaxChordHeight = UseRadius - FMathd::Sqrt(UseRadius*UseRadius - HalfChordLen*HalfChordLen);   // "sagitta" height
 
diff --git a/Engine/Plugins/Runtime/GeometryProcessing/Source/DynamicMesh/Public/ShapeApproximation/ShapeDetection3.h b/Engine/Plugins/Runtime/GeometryProcessing/Source/DynamicMesh/Public/ShapeApproximation/ShapeDetection3.h
--- a/Engine/Plugins/Runtime/GeometryProcessing/Source/DynamicMesh/Public/ShapeApproximation/ShapeDetection3.h
+++ b/Engine/Plugins/Runtime/GeometryProcessing/Source/DynamicMesh/Public/ShapeApproximation/ShapeDetection3.h
@@ -17,10 +17,11 @@ namespace UE
 		/**
 		 * Detect if input Mesh is a meshed approximation of an analytic Sphere, and if so return best guess in SphereOut.
 		 * Fits a sphere to input points with several rounds of incremental improvement, then measures chordal deviation of edge midpoints.
-		 * @param RelativeDeviationTol distances from edge midpoints to sphere surface are allowed to deviate by 2*Radius*RelativeDeviationTol
+		 * @param RelativeDeviationTol Scaled by sphere diameter. The allowed difference in the distance from edge midpoints to the surface of the sphere vs the ideal distance for an edge of the same length.
+		 * @param MaxAngleRangeDegrees Maximum angle difference in vectors from center of sphere to the endpoints of any surface edge. Controls how coarsely tessellated a shpere can be before it is not considered a sphere.
 		 * @return true if mesh is a Sphere and SphereOut is initialized
 		 */
-		bool DYNAMICMESH_API IsSphereMesh(const FDynamicMesh3& Mesh, FSphere3d& SphereOut, double RelativeDeviationTol = 0.025);
+		bool DYNAMICMESH_API IsSphereMesh(const FDynamicMesh3& Mesh, FSphere3d& SphereOut, double RelativeDeviationTol = 0.025, double MaxAngleRangeDegrees = 72);
 
 
 		/**
