Make group self-loop edges get frame normals by plane fit, rather than picking the normal off an arbitrary triangle

diff --git a/Engine/Plugins/Runtime/GeometryProcessing/Source/DynamicMesh/Private/GroupTopology.cpp b/Engine/Plugins/Runtime/GeometryProcessing/Source/DynamicMesh/Private/GroupTopology.cpp
--- a/Engine/Plugins/Runtime/GeometryProcessing/Source/DynamicMesh/Private/GroupTopology.cpp
+++ b/Engine/Plugins/Runtime/GeometryProcessing/Source/DynamicMesh/Private/GroupTopology.cpp
@@ -811,33 +811,48 @@ FFrame3d FGroupTopology::GetSelectionFrame(const FGroupTopologySelection& Select
 	if (NumEdges == 1)
 	{
 		int32 GroupEdgeID = Selection.GetASelectedEdgeID();
-		int32 MeshEdgeID = GetGroupEdgeEdges(GroupEdgeID)[0];
-
-		// align Z axis of frame to face normal of one of the connected faces. 
-		FIndex2i EdgeTris = Mesh->GetEdgeT(MeshEdgeID);
-		int32 UseFace = (EdgeTris.B != IndexConstants::InvalidID) ? FMath::Min(EdgeTris.A, EdgeTris.B)
-			: EdgeTris.A;
-		FVector3d FaceNormal = Mesh->GetTriNormal(UseFace);
-		if (FaceNormal.Length() > 0.1)
-		{
-			StartFrame.AlignAxis(2, FaceNormal);
-		}
-
-		// align X axis along the edge, around the aligned Z axis
-		FVector3d Tangent;
-		if (GetGroupEdgeTangent(GroupEdgeID, Tangent))
-		{
-			StartFrame.ConstrainedAlignAxis(0, Tangent, StartFrame.Z());
-		}
 
 		const TArray<int32>& Vertices = GetGroupEdgeVertices(GroupEdgeID);
-		if (Vertices[0] == Vertices.Last())
+		bool bIsSelfLoop = Vertices[0] == Vertices.Last();
+
+		if (!bIsSelfLoop)
 		{
-			StartFrame.Origin = GetEdgeAveragePosition(GroupEdgeID);
+			// align Z axis of frame to face normal of one of the connected faces. 
+			int32 MeshEdgeID = GetGroupEdgeEdges(GroupEdgeID)[0];
+			FIndex2i EdgeTris = Mesh->GetEdgeT(MeshEdgeID);
+			int32 UseFace = (EdgeTris.B != IndexConstants::InvalidID) ? FMath::Min(EdgeTris.A, EdgeTris.B)
+				: EdgeTris.A;
+			FVector3d FaceNormal = Mesh->GetTriNormal(UseFace);
+			if (!FaceNormal.IsZero())
+			{
+				StartFrame.AlignAxis(2, FaceNormal);
+			}
+
+			// align X axis along the edge, around the aligned Z axis
+			FVector3d Tangent;
+			if (GetGroupEdgeTangent(GroupEdgeID, Tangent))
+			{
+				StartFrame.ConstrainedAlignAxis(0, Tangent, StartFrame.Z());
+			}
+			StartFrame.Origin = GetEdgeMidpoint(GroupEdgeID);
 		}
 		else
 		{
-			StartFrame.Origin = GetEdgeMidpoint(GroupEdgeID);
+			StartFrame.Origin = GetEdgeAveragePosition(GroupEdgeID);
+			// For self-loops, get polygon normal via newell's method
+			FVector FaceNormal = FVector::ZeroVector;
+			FVector Prev = Mesh->GetVertex(Vertices.Last());
+			for (int32 Idx = 0; Idx < Vertices.Num(); ++Idx)
+			{
+				FVector Cur = Mesh->GetVertex(Vertices[Idx]);
+				FaceNormal += FVector((Cur.Y - Prev.Y) * (Prev.Z + Cur.Z), (Cur.Z - Prev.Z) * (Prev.X + Cur.X), (Cur.X - Prev.X) * (Prev.Y + Cur.Y));
+				Prev = Cur;
+			}
+			// If we found a valid plane normal, align the frame to it; otherwise leave the frame unchanged from default
+			if (FaceNormal.Normalize())
+			{
+				StartFrame.AlignAxis(2, FaceNormal);
+			}
 		}
 		return StartFrame;
 	}
