Expose exterior visibility algorithm in GeometryProcessing DynamicMesh module, and use it to add an operation to determine which mesh instances are fully occluded. Also add this operation as a new Geometry Script function.

diff --git a/Engine/Plugins/Experimental/GeometryScripting/Source/GeometryScriptingCore/Private/SceneUtilityFunctions.cpp b/Engine/Plugins/Experimental/GeometryScripting/Source/GeometryScriptingCore/Private/SceneUtilityFunctions.cpp
--- a/Engine/Plugins/Experimental/GeometryScripting/Source/GeometryScriptingCore/Private/SceneUtilityFunctions.cpp
+++ b/Engine/Plugins/Experimental/GeometryScripting/Source/GeometryScriptingCore/Private/SceneUtilityFunctions.cpp
@@ -8,6 +8,7 @@
 #include "DynamicMesh/MeshTransforms.h"
 #include "DynamicMesh/MeshNormals.h"
 #include "DynamicMeshEditor.h"
+#include "Operations/DetectExteriorVisibility.h"
 
 #include "Engine/StaticMesh.h"
 #include "Engine/SkeletalMesh.h"
@@ -361,7 +362,52 @@ UDynamicMesh* UGeometryScriptLibrary_SceneUtilityFunctions::CopyCollisionMeshesF
 	return ToDynamicMesh;
 }
 
+void UGeometryScriptLibrary_SceneUtilityFunctions::DetermineMeshOcclusion(
+	const TArray<UDynamicMesh*>& SourceMeshes,
+	const TArray<FTransform>& SourceMeshTransforms,
+	TArray<bool>& OutMeshIsHidden,
+	const TArray<UDynamicMesh*>& OccludeMeshes,
+	const TArray<FTransform>& OccludeMeshTransforms,
+	const FGeometryScriptDetermineMeshOcclusionOptions& OcclusionOptions,
+	UGeometryScriptDebug* Debug)
+{
+	if (SourceMeshes.Num() != SourceMeshTransforms.Num())
+	{
+		UE::Geometry::AppendError(Debug, EGeometryScriptErrorType::InvalidInputs, LOCTEXT("DetermineMeshOcclusion_SourceArrayMismatch", "DetermineMeshOcclusion: SourceMeshes and SourceMeshTransforms arrays must have same length"));
+		return;
+	}
+	if (OccludeMeshes.Num() != OccludeMeshTransforms.Num())
+	{
+		UE::Geometry::AppendError(Debug, EGeometryScriptErrorType::InvalidInputs, LOCTEXT("DetermineMeshOcclusion_OccludeArrayMismatch", "DetermineMeshOcclusion: OccludeMeshes and OccludeMeshTransforms arrays must have same length"));
+		return;
+	}
 
+	FDetectPerDynamicMeshExteriorVisibility Occlusion;
+	for (int32 SourceMeshIndex = 0; SourceMeshIndex < SourceMeshes.Num(); ++SourceMeshIndex)
+	{
+		UDynamicMesh* Mesh = SourceMeshes[SourceMeshIndex];
+		if (!Mesh)
+		{
+			UE::Geometry::AppendError(Debug, EGeometryScriptErrorType::InvalidInputs, LOCTEXT("DetermineMeshOcclusion_InvalidSourceMesh", "DetermineMeshOcclusion: SourceMeshes array contained null mesh"));
+			return;
+		}
+		Occlusion.Instances.Emplace(Mesh->GetMeshPtr(), SourceMeshTransforms[SourceMeshIndex]);
+	}
+	for (int32 OccludeMeshIndex = 0; OccludeMeshIndex < OccludeMeshes.Num(); ++OccludeMeshIndex)
+	{
+		UDynamicMesh* Mesh = OccludeMeshes[OccludeMeshIndex];
+		if (!Mesh)
+		{
+			UE::Geometry::AppendError(Debug, EGeometryScriptErrorType::InvalidInputs, LOCTEXT("DetermineMeshOcclusion_InvalidOccludeMesh", "DetermineMeshOcclusion: OccludeMeshes array contained null mesh"));
+			return;
+		}
+		Occlusion.OccludeInstances.Emplace(Mesh->GetMeshPtr(), OccludeMeshTransforms[OccludeMeshIndex]);
+	}
+	Occlusion.SamplingParameters.bDoubleSided = OcclusionOptions.bDoubleSided;
+	Occlusion.SamplingParameters.SamplingDensity = OcclusionOptions.SamplingDensity;
+	Occlusion.SamplingParameters.NumSearchDirections = OcclusionOptions.NumSearchDirections;
 
+	Occlusion.ComputeHidden(OutMeshIsHidden);
+}
 
 #undef LOCTEXT_NAMESPACE
\ No newline at end of file
diff --git a/Engine/Plugins/Experimental/GeometryScripting/Source/GeometryScriptingCore/Public/GeometryScript/SceneUtilityFunctions.h b/Engine/Plugins/Experimental/GeometryScripting/Source/GeometryScriptingCore/Public/GeometryScript/SceneUtilityFunctions.h
--- a/Engine/Plugins/Experimental/GeometryScripting/Source/GeometryScriptingCore/Public/GeometryScript/SceneUtilityFunctions.h
+++ b/Engine/Plugins/Experimental/GeometryScripting/Source/GeometryScriptingCore/Public/GeometryScript/SceneUtilityFunctions.h
@@ -27,6 +27,25 @@ public:
 	FGeometryScriptMeshReadLOD RequestedLOD = FGeometryScriptMeshReadLOD();
 };
 
+USTRUCT(BlueprintType)
+struct GEOMETRYSCRIPTINGCORE_API FGeometryScriptDetermineMeshOcclusionOptions
+{
+	GENERATED_BODY()
+public:
+
+	// Approximate spacing between samples on triangle faces used for determining visibility
+	UPROPERTY(EditAnywhere, BlueprintReadWrite, Category = Options)
+	double SamplingDensity = 1.0;
+
+	// Whether to treat faces as double-sided when determining visibility
+	UPROPERTY(EditAnywhere, BlueprintReadWrite, Category = Options)
+	bool bDoubleSided = false;
+
+	// Number of directions to test for visibility
+	UPROPERTY(EditAnywhere, BlueprintReadWrite, Category = Options)
+	int32 NumSearchDirections = 128;
+};
+
 
 UCLASS(meta = (ScriptName = "GeometryScript_SceneUtils"))
 class GEOMETRYSCRIPTINGCORE_API UGeometryScriptLibrary_SceneUtilityFunctions : public UBlueprintFunctionLibrary
@@ -95,4 +114,24 @@ public:
 		int SphereResolution = 16,
 		UGeometryScriptDebug* Debug = nullptr);
 
+
+	/**
+	 * Determine which meshes are entirely hidden by other meshes in the set, when viewed from outside.
+	 * 
+	 * @param SourceMeshes			Meshes to test for occlusion. Note: The same mesh may appear multiple times in this array, if it is instanced with different transforms.
+	 * @param SourceMeshTransforms	A transform for each source mesh. Array must have the same length as SourceMeshes.
+	 * @param OutMeshIsHidden		Array will be filled with a bool per source mesh, indicating whether that mesh is hidden (true) or visible (false)
+	 * @param OccludeMeshes			Array of optional meshes which can occlude SourceMeshes, but for which we will not test occlusion.
+	 * @param OccludeMeshTransforms	Array of transforms for each occlude mesh. Array must have the same length as OccludeMeshes.
+	 * @param Options				Settings to control how occlusion is tested
+	 */
+	UFUNCTION(BlueprintCallable, Category = "GeometryScript|Scene", meta = (AutoCreateRefTerm = "OccludeMeshes, OccludeMeshTransforms, OcclusionOptions"))
+	static void DetermineMeshOcclusion(
+		const TArray<UDynamicMesh*>& SourceMeshes,
+		const TArray<FTransform>& SourceMeshTransforms,
+		TArray<bool>& OutMeshIsHidden,
+		const TArray<UDynamicMesh*>& OccludeMeshes,
+		const TArray<FTransform>& OccludeMeshTransforms,
+		const FGeometryScriptDetermineMeshOcclusionOptions& OcclusionOptions,
+		UGeometryScriptDebug* Debug = nullptr);
 };
\ No newline at end of file
diff --git a/Engine/Plugins/Runtime/GeometryProcessing/Source/DynamicMesh/Private/Operations/DetectExteriorVisibility.cpp b/Engine/Plugins/Runtime/GeometryProcessing/Source/DynamicMesh/Private/Operations/DetectExteriorVisibility.cpp
new file mode 100644
--- /dev/null
+++ b/Engine/Plugins/Runtime/GeometryProcessing/Source/DynamicMesh/Private/Operations/DetectExteriorVisibility.cpp
@@ -0,0 +1,431 @@
+// Copyright Epic Games, Inc. All Rights Reserved.
+
+#include "Operations/DetectExteriorVisibility.h"
+
+#include <atomic>
+
+#include "Async/ParallelFor.h"
+#include "DynamicMesh/DynamicMesh3.h"
+#include "DynamicMesh/DynamicMeshAABBTree3.h"
+#include "DynamicMeshEditor.h"
+#include "Sampling/SphericalFibonacci.h"
+
+using namespace UE::Geometry;
+
+
+void FExteriorVisibilitySampling::ComputePerTriangleOcclusion(const FDynamicMesh3& TargetMesh, TArray<bool>& OutTriOccluded, bool bPerPolyGroup, const FExteriorVisibilitySampling& SamplingParameters, TArrayView<const bool> SkipTris)
+{
+	const bool bDoubleSided = SamplingParameters.bDoubleSided;
+	const double SampleRadius = FMath::Max(UE_DOUBLE_KINDA_SMALL_NUMBER, SamplingParameters.SamplingDensity);
+	int32 NumVisibilityTestDirections = FMath::Max(3, SamplingParameters.NumSearchDirections);
+
+	FDynamicMeshAABBTree3 Spatial(&TargetMesh, true);
+	FAxisAlignedBox3d Bounds = Spatial.GetBoundingBox();
+	double Radius = Bounds.DiagonalLength();
+
+	// geometric magic numbers used below that have been slightly tuned...
+	double GlancingAngleDotTolerance = FMathd::Cos(85.0 * FMathd::DegToRad);
+	constexpr double TriScalingAlpha = 0.999;
+	constexpr double BaryCoordsThreshold = 0.001;
+
+	auto FindHitTriangleTest = [&](FVector3d TargetPosition, FVector3d TargetNormal, FVector3d FarPosition) -> int
+	{
+		FVector3d RayDir(TargetPosition - FarPosition);
+		double Distance = Normalize(RayDir);
+		if (bDoubleSided == false && RayDir.Dot(TargetNormal) > -0.001)
+		{
+			return IndexConstants::InvalidID;
+		}
+		FRay3d Ray(FarPosition, RayDir, true);
+		return Spatial.FindNearestHitTriangle(Ray, IMeshSpatial::FQueryOptions(Distance + 1.0));	// 1.0 is random fudge factor here...
+	};
+
+	// array of (+/-)X/Y/Z directions
+	TArray<FVector3d> CardinalDirections;
+	for (int32 k = 0; k < 3; ++k)
+	{
+		FVector3d Direction(0, 0, 0);
+		Direction[k] = 1.0;
+		CardinalDirections.Add(Direction);
+		CardinalDirections.Add(-Direction);
+	}
+
+	// TODO: it seems like a quite common failure case is triangles that are somewhat deeply
+	// nested inside cavities/etc. A possible way to help here would be to essentially raytrace
+	// orthographic images from top/bottom/etc. This could be done async and combined w/ the
+	// visibilty determination below...
+
+
+	// First pass. For each triangle, cast a ray at it's centroid from 
+	// outside the model, along the X/Y/Z directions and tri normal.
+	// If tri is hit we mark it as having 'known' status, allowing it
+	// to be skipped in the more expensive pass below
+	//
+	TArray<bool> TriStatusKnown;
+	TriStatusKnown.SetNumZeroed(TargetMesh.MaxTriangleID());
+	bool bHasSkipTris = SkipTris.Num() == TargetMesh.MaxTriangleID();
+	int32 MaxUnknownTID = -1;
+	for (int32 TID : TargetMesh.TriangleIndicesItr())
+	{
+		bool bKnown = bHasSkipTris && SkipTris[TID];
+		TriStatusKnown[TID] = bKnown;
+		if (!bKnown)
+		{
+			MaxUnknownTID = TID;
+		}
+	}
+	ParallelFor(MaxUnknownTID + 1, [&](int32 TID)
+	{
+		if (!TargetMesh.IsTriangle(TID))
+		{
+			TriStatusKnown[TID] = true;
+			return;
+		}
+		FVector3d Normal, Centroid; double Area;
+		TargetMesh.GetTriInfo(TID, Normal, Area, Centroid);
+		if (Normal.SquaredLength() < 0.1 || Area <= FMathd::ZeroTolerance)
+		{
+			TriStatusKnown[TID] = true;
+			return;
+		}
+
+		for (FVector3d Direction : CardinalDirections)
+		{
+			// if direction is orthogonal to the triangle, hit-test is unstable, but even
+			// worse, on rectilinear shapes (eg imagine some stacked cubes or adjacent parts)
+			// the ray can get "through" the cracks between adjacent connected triangles
+			// and manage to hit the search triangle
+			if (FMathd::Abs(Direction.Dot(Normal)) > GlancingAngleDotTolerance)
+			{
+				if (FindHitTriangleTest(Centroid, Normal, Centroid + Radius * Direction) == TID)
+				{
+					TriStatusKnown[TID] = true;
+					return;
+				}
+			}
+		}
+		if (FindHitTriangleTest(Centroid, Normal, Centroid + Radius * Normal) == TID)
+		{
+			TriStatusKnown[TID] = true;
+			return;
+		}
+
+		// triangle is not definitely visible or hidden
+	});
+
+	// Set up PolyGroup visibility if requested
+	TArray<bool> GroupStatusKnown;
+	TArray<std::atomic<bool>> ThreadSafeGroupVisible;
+	if (bPerPolyGroup)
+	{
+		// Init initial group status from initial tri status
+		GroupStatusKnown.SetNumZeroed(TargetMesh.MaxGroupID());
+		for (int32 TID : TargetMesh.TriangleIndicesItr())
+		{
+			if (TriStatusKnown[TID])
+			{
+				int32 GID = TargetMesh.GetTriangleGroup(TID);
+				if (GroupStatusKnown.IsValidIndex(GID))
+				{
+					GroupStatusKnown[GID] = true;
+				}
+			}
+		}
+		// Copy to thread safe values
+		ThreadSafeGroupVisible.SetNum(TargetMesh.MaxGroupID());
+		for (int32 GroupIdx = 0; GroupIdx < GroupStatusKnown.Num(); ++GroupIdx)
+		{
+			ThreadSafeGroupVisible[GroupIdx] = GroupStatusKnown[GroupIdx];
+		}
+		// Copy group statuses back to tri statuses
+		for (int32 TID : TargetMesh.TriangleIndicesItr())
+		{
+			int32 GID = TargetMesh.GetTriangleGroup(TID);
+			if (GroupStatusKnown.IsValidIndex(GID) && GroupStatusKnown[GID])
+			{
+				TriStatusKnown[TID] = true;
+			}
+		}
+	}
+
+	// Thread safe triangle visibility, atomics can be updated on any thread
+	TArray<std::atomic<bool>> ThreadSafeTriVisible;
+	ThreadSafeTriVisible.SetNum(TargetMesh.MaxTriangleID());
+	MaxUnknownTID = -1; // re-compute max unknown
+	for (int32 TID : TargetMesh.TriangleIndicesItr())
+	{
+		bool bKnown = TriStatusKnown[TID];
+		ThreadSafeTriVisible[TID] = bKnown;
+		if (!bKnown)
+		{
+			MaxUnknownTID = TID;
+		}
+	}
+
+	//
+	// Construct set of exterior visibility test directions, below we will check if sample 
+	// points on the mesh triangles are visible from the exterior along these directions.
+	// Order is modulo-shuffled in hopes that for visible tris we don't waste a bunch
+	// of time on the 'far' side
+	TSphericalFibonacci<double> SphereSampler(NumVisibilityTestDirections);
+	TArray<FVector3d> VisibilityDirections;
+	FModuloIteration ModuloIter(NumVisibilityTestDirections);
+	uint32 DirectionIndex = 0;
+	while (ModuloIter.GetNextIndex(DirectionIndex))
+	{
+		VisibilityDirections.Add(Normalized(SphereSampler[DirectionIndex]));
+	}
+	// Fibonacci set generally does not include the cardinal directions, but they are highly useful to check
+	VisibilityDirections.Append(CardinalDirections);
+	NumVisibilityTestDirections = VisibilityDirections.Num();
+
+	// For each triangle we will generate a set of sample points on the triangle surface,
+	// and then check if that point is visible along any of the sample directions.
+	// The number of sample points allocated to a triangle is based on it's area and the SampleRadius.
+	// However for small triangles this may be < 1, so we will clamp to at least this many samples.
+	// (this value should perhaps be relative to the mesh density, or exposed as a parameter...)
+	const int32 MinTriSamplesPerSamplePoint = 8;
+
+	// For each triangle, generate a set of sample points on the triangle surface,
+
+	// This is the expensive part!
+	ParallelFor(MaxUnknownTID + 1, [&](int32 TID)
+	{
+		if (!TargetMesh.IsTriangle(TID))
+		{
+			return;
+		}
+
+		int32 GID = 0;
+		bool bHasValidGroup = bPerPolyGroup;
+		if (bPerPolyGroup)
+		{
+			GID = TargetMesh.GetTriangleGroup(TID);
+			bHasValidGroup = GroupStatusKnown.IsValidIndex(GID);
+			if (bHasValidGroup && (GroupStatusKnown[GID] || ThreadSafeGroupVisible[GID]))
+			{
+				return;
+			}
+		}
+		// if we already found out this triangle is visible or hidden, we can skip it
+		if (TriStatusKnown[TID] || ThreadSafeTriVisible[TID]) return;
+
+		FVector3d A, B, C;
+		TargetMesh.GetTriVertices(TID, A, B, C);
+		FVector3d Centroid = (A + B + C) / 3.0;
+		double TriArea;
+		FVector3d TriNormal = VectorUtil::NormalArea(A, B, C, TriArea);		// TriStatusKnown should skip degen tris, do not need to check here
+
+		FFrame3d TriFrame(Centroid, TriNormal);
+		FTriangle2d UVTriangle(TriFrame.ToPlaneUV(A), TriFrame.ToPlaneUV(B), TriFrame.ToPlaneUV(C));
+
+		// Slightly shrink the triangle, this helps to avoid spurious hits
+		// TODO obviously should scale by an actual dimension and not just a relative %...
+		FVector2d Center = (UVTriangle.V[0] + UVTriangle.V[1] + UVTriangle.V[2]) / 3.0;
+		for (int32 k = 0; k < 3; ++k)
+		{
+			UVTriangle.V[k] = (1 - TriScalingAlpha) * Center + (TriScalingAlpha)*UVTriangle.V[k];
+		}
+
+		double DiscArea = (FMathd::Pi * SampleRadius * SampleRadius);
+		int NumSamples = FMath::Max((int)(TriArea / DiscArea), MinTriSamplesPerSamplePoint);
+		FVector2d V1 = UVTriangle.V[1] - UVTriangle.V[0];
+		FVector2d V2 = UVTriangle.V[2] - UVTriangle.V[0];
+
+		TArray<int32> HitTris, HitGroups;		// re-use this array in inner loop to avoid hitting atomics so often
+
+		int NumTested = 0; int Iterations = 0;
+		FRandomStream RandomStream(TID);
+		while (NumTested < NumSamples && Iterations++ < 10000)
+		{
+			double a1 = RandomStream.GetFraction();
+			double a2 = RandomStream.GetFraction();
+			FVector2d PointUV = UVTriangle.V[0] + a1 * V1 + a2 * V2;
+			if (UVTriangle.IsInside(PointUV))
+			{
+				NumTested++;
+				FVector3d Position = TriFrame.FromPlaneUV(PointUV, 2);
+
+				// cast ray from all exterior sample locations for this triangle sample point
+				HitTris.Reset();
+				for (int32 k = 0; k < NumVisibilityTestDirections; ++k)
+				{
+					FVector3d Direction = VisibilityDirections[k];
+					if (FMathd::Abs(Direction.Dot(TriNormal)) < GlancingAngleDotTolerance)
+					{
+						continue;
+					}
+
+					FVector3d RayFrom = Position + 2.0 * Radius * VisibilityDirections[k];
+					int32 HitTriID = FindHitTriangleTest(Position, TriNormal, RayFrom);
+					if (HitTriID != IndexConstants::InvalidID && TriStatusKnown[HitTriID] == false)
+					{
+						// Want to filter out on-edge triangle hits, as they are generally spurious and
+						// will result in interior triangles remaining visible
+						FRay3d Ray(RayFrom, Normalized(Position - RayFrom), true);
+						FIntrRay3Triangle3d RayHit = TMeshQueries<FDynamicMesh3>::RayTriangleIntersection(TargetMesh, HitTriID, Ray);
+						if (RayHit.IntersectionType == EIntersectionType::Point &&
+							RayHit.TriangleBaryCoords.GetMin() > BaryCoordsThreshold &&
+							RayHit.TriangleBaryCoords.GetMax() < (1.0 - BaryCoordsThreshold))
+						{
+
+							HitTris.AddUnique(HitTriID);		// we hit some triangle, whether or not it is the one we are testing...
+							if (HitTriID == TID)
+							{
+								break;
+							}
+						}
+					}
+				}
+
+				// mark any hit tris
+				for (int32 HitTriID : HitTris)
+				{
+					ThreadSafeTriVisible[HitTriID] = true;
+				}
+				if (bPerPolyGroup)
+				{
+					// Note: can tune the threshold at which we reset the hitgroups array, to trade off atomic accesses vs AddUnique search cost 
+					constexpr int32 ResetGroupSize = 20;
+					if (HitGroups.Num() >= ResetGroupSize)
+					{
+						HitGroups.Reset();
+					}
+					for (int32 HitTriID : HitTris)
+					{
+						int32 HitGID = TargetMesh.GetTriangleGroup(HitTriID);
+						if (GroupStatusKnown.IsValidIndex(HitGID) && !GroupStatusKnown[HitGID] && HitGroups.Num() == HitGroups.AddUnique(HitGID))
+						{
+							// Note: setting/checking the group visibility flag is optional and just allow threads to stop early
+							// so it should be fine to do so w/ the relaxed memory ordering
+							ThreadSafeGroupVisible[HitGID].store(true, std::memory_order_relaxed);
+						}
+					}
+				}
+
+				// if our triangle has become visible (in this thread or another) we can terminate now
+				if (ThreadSafeTriVisible[TID] || (bHasValidGroup && ThreadSafeGroupVisible[GID].load(std::memory_order_relaxed)))
+				{
+					return;
+				}
+			}
+		}
+
+		// Note: Consider periodically locking to update TriStatusKnown and GroupStatusKnown
+	});
+
+	// propagate tri visibility back through groups
+	if (bPerPolyGroup)
+	{
+		for (int32 TID : TargetMesh.TriangleIndicesItr())
+		{
+			if (ThreadSafeTriVisible[TID])
+			{
+				int32 GID = TargetMesh.GetTriangleGroup(TID);
+				if (ThreadSafeGroupVisible.IsValidIndex(GID))
+				{
+					ThreadSafeGroupVisible[GID] = true;
+				}
+			}
+		}
+		// Copy group statuses back to tri statuses
+		for (int32 TID : TargetMesh.TriangleIndicesItr())
+		{
+			int32 GID = TargetMesh.GetTriangleGroup(TID);
+			if (ThreadSafeGroupVisible.IsValidIndex(GID) && ThreadSafeGroupVisible[GID])
+			{
+				ThreadSafeTriVisible[TID] = true;
+			}
+		}
+	}
+
+	// Transfer visibility results to output array
+	OutTriOccluded.SetNum(ThreadSafeTriVisible.Num());
+	for (int32 Idx = 0; Idx < OutTriOccluded.Num(); ++Idx)
+	{
+		OutTriOccluded[Idx] = !ThreadSafeTriVisible[Idx];
+	}
+}
+
+
+void FDetectPerDynamicMeshExteriorVisibility::ComputeHidden(TArray<bool>& OutIsOccluded)
+{
+	OutIsOccluded.SetNumZeroed(Instances.Num());
+	if (Instances.Num() < 2)
+	{
+		// 0 or 1 meshes can't be fully occluded
+		return;
+	}
+
+	// Combine all Instances and OccludeInstances into a single AllMesh, to allow occlusion calculation via the single-mesh function above
+	// Note: To support large meshes with many instances, consider adding an alternate path which uses a separate tree per unique instance.
+	FDynamicMesh3 AllMesh(EMeshComponents::None);
+	FDynamicMeshEditor CombineEditor(&AllMesh);
+	FMeshIndexMappings Mappings;
+	TArray<int32> InstanceEndTID;
+	auto AddInstances = [this, &AllMesh, &CombineEditor, &Mappings, &InstanceEndTID](TArray<FDynamicMeshInstance>& ToAdd)
+	{
+		for (int32 InstIdx = 0; InstIdx < ToAdd.Num(); ++InstIdx)
+		{
+			CombineEditor.AppendMesh(ToAdd[InstIdx].SourceMesh, Mappings,
+			[ToAdd, InstIdx](int32 VID, const FVector3d Pos) -> FVector3d
+			{
+				return ToAdd[InstIdx].Transforms.TransformPosition(Pos);
+			});
+			InstanceEndTID.Add(AllMesh.MaxTriangleID());
+		}
+	};
+	AddInstances(Instances);
+	int32 OccludeStartID = AllMesh.MaxTriangleID();
+	AddInstances(OccludeInstances);
+
+	// Should be safe to assume contiguous triangle IDs because we've constructed this mesh ourselves
+	// if this assumption breaks, will need to change the logic below which assigns PolyGroup IDs per instance
+	if (!ensure(AllMesh.IsCompactT()))
+	{
+		return;
+	}
+
+	// Enable and set triangle groups *after* adding all meshes, so the groups initialize to zero and we are not affected by groups on the source meshes
+	AllMesh.EnableTriangleGroups(0);
+	for (int32 TID = InstanceEndTID[0], CurInstance = 1; TID < AllMesh.MaxTriangleID(); ++TID)
+	{
+		while (CurInstance < InstanceEndTID.Num() && TID == InstanceEndTID[CurInstance])
+		{
+			CurInstance++;
+		}
+		// Set triangle group by instance ID, but give all OccludeInstances the same ID
+		AllMesh.SetTriangleGroup(TID, CurInstance < Instances.Num() ? CurInstance : Instances.Num());
+	}
+
+	// Mark any triangles from Occlude Instances as SkipTris
+	TArray<bool> SkipTris;
+	SkipTris.SetNumZeroed(AllMesh.MaxTriangleID());
+	for (int32 TID = OccludeStartID; TID < AllMesh.MaxTriangleID(); ++TID)
+	{
+		SkipTris[TID] = true;
+	}
+
+	// Compute per-tri occlusion
+	TArray<bool> TriOccluded;
+	FExteriorVisibilitySampling::ComputePerTriangleOcclusion(AllMesh, TriOccluded, true /*bPerPolyGroup*/, SamplingParameters, SkipTris);
+
+	// Transfer triangle occlusion back to source mesh occlusion via the triangle groups
+	OutIsOccluded.Init(true, Instances.Num());
+	for (int32 TID : AllMesh.TriangleIndicesItr())
+	{
+		int32 GID = AllMesh.GetTriangleGroup(TID);
+		if (GID < 0 || GID >= Instances.Num()) // tri is from an OccludeInstances mesh
+		{
+			break; // OccludeInstances meshes are at the end; we can stop when we reach them
+		}
+		if (!OutIsOccluded[GID]) // already found visible
+		{
+			continue;
+		}
+		if (!TriOccluded[TID]) // found visible tri -> source is visible
+		{
+			OutIsOccluded[GID] = false;
+		}
+	}
+}
\ No newline at end of file
diff --git a/Engine/Plugins/Runtime/GeometryProcessing/Source/DynamicMesh/Public/Operations/DetectExteriorVisibility.h b/Engine/Plugins/Runtime/GeometryProcessing/Source/DynamicMesh/Public/Operations/DetectExteriorVisibility.h
new file mode 100644
--- /dev/null
+++ b/Engine/Plugins/Runtime/GeometryProcessing/Source/DynamicMesh/Public/Operations/DetectExteriorVisibility.h
@@ -0,0 +1,71 @@
+// Copyright Epic Games, Inc. All Rights Reserved.
+
+#pragma once
+
+#include "Containers/Array.h"
+#include "Containers/ArrayView.h"
+#include "TransformSequence.h"
+
+namespace UE::Geometry
+{
+
+class FDynamicMesh3;
+
+// Parameters controlling how exterior visiblity is tested
+struct FExteriorVisibilitySampling
+{
+	// Approximate spacing between triangle samples used for visibility tests
+	double SamplingDensity = 1.0;
+
+	// Whether to treat faces as double-sided
+	bool bDoubleSided = false;
+
+	// Number of directions to test for visibility
+	int32 NumSearchDirections = 128;
+
+	// Compute per-triangle visibility array for a triangle mesh
+	// @param OutTriOccluded	For each valid triangle ID, OutTriOccluded[ID] will be true if that triangle is hidden, false if it is visible
+	// @param SkipTris			True for triangles do not need visibility testing (will be marked as not occluded). If empty, no triangles will be skipped.
+	// @param bPerPolyGroup		If true, visiblity will be determined on a per group basis: If any triangle in a group is visible, the whole group will be marked as not-occluded
+	// @param bDoubleSided		Whether to treat faces a double-sided
+	// @param SamplingDensity	Approximate spacing between triangle samples used for visibility tests
+	DYNAMICMESH_API static void ComputePerTriangleOcclusion(const FDynamicMesh3& Mesh, TArray<bool>& OutTriOccluded, bool bPerPolyGroup, const FExteriorVisibilitySampling& SamplingParameters, TArrayView<const bool> SkipTris = TArrayView<const bool>());
+};
+
+// Determine which meshes are visible from exterior views of a list of meshes
+class FDetectPerDynamicMeshExteriorVisibility
+{
+public:
+
+	struct FDynamicMeshInstance
+	{
+		const FDynamicMesh3* SourceMesh;
+		FTransformSequence3d Transforms;
+		FDynamicMeshInstance() = default;
+		FDynamicMeshInstance(FDynamicMesh3* Mesh, const FTransform& InTransform) : SourceMesh(Mesh)
+		{
+			Transforms.Append(InTransform);
+		}
+	};
+
+	//
+	// Inputs
+	//
+
+	// Instances for which we need to determine visibility
+	TArray<FDynamicMeshInstance> Instances;
+	
+	// Instances which occlude visibility, but for which we don't need to determine visibility
+	TArray<FDynamicMeshInstance> OccludeInstances;
+
+	// Parameters controlling how the triangles are sampled for visibility
+	FExteriorVisibilitySampling SamplingParameters;
+
+	// Compute which source meshes are hidden
+	// @param OutIsOccluded		Indicate if each mesh in Instances is hidden (true) or visible (false)
+	DYNAMICMESH_API void ComputeHidden(TArray<bool>& OutIsOccluded);
+
+};
+
+
+} // end namespace UE::Geometry
