prevent cones and cylinders-with-open-ends from being auto-detected as capsules by mesh to collision / shape auto-detection, by adding additional criteria to the capsule detection

plus add a fix to the sphere auto-detection angle tolerance case

diff --git a/Engine/Plugins/Runtime/GeometryProcessing/Source/DynamicMesh/Private/ShapeApproximation/ShapeDetection3.cpp b/Engine/Plugins/Runtime/GeometryProcessing/Source/DynamicMesh/Private/ShapeApproximation/ShapeDetection3.cpp
--- a/Engine/Plugins/Runtime/GeometryProcessing/Source/DynamicMesh/Private/ShapeApproximation/ShapeDetection3.cpp
+++ b/Engine/Plugins/Runtime/GeometryProcessing/Source/DynamicMesh/Private/ShapeApproximation/ShapeDetection3.cpp
@@ -81,8 +81,8 @@ bool UE::Geometry::IsSphereMesh(const FDynamicMesh3& Mesh, FSphere3d& SphereOut,
 		Mesh.GetEdgeV(EdgeID, A, B);
 
 		// if a single edge spans too wide an angular range, the shape is too coarsely tesselated to be considered a sphere
-		FVector3d ToA = A - SphereOut.Center;
-		FVector3d ToB = B - SphereOut.Center;
+		FVector3d ToA = (A - SphereOut.Center).GetSafeNormal();
+		FVector3d ToB = (B - SphereOut.Center).GetSafeNormal();
 		if (ToA.Dot(ToB) <= CosAngleTolerance)
 		{
 			return false;
@@ -253,7 +253,7 @@ bool UE::Geometry::IsBoxMesh(const FDynamicMesh3& Mesh, FOrientedBox3d& BoxOut,
 
 
 
-bool UE::Geometry::IsCapsuleMesh(const FDynamicMesh3& Mesh, FCapsule3d& CapsuleOut, double RelativeDeviationTol)
+bool UE::Geometry::IsCapsuleMesh(const FDynamicMesh3& Mesh, FCapsule3d& CapsuleOut, double RelativeDeviationTol, double MaxAngleRangeDegrees)
 {
 	// minimal 4-slice capsule has at least 10 vertices
 	if (Mesh.VertexCount() < 10)		
@@ -285,26 +285,72 @@ bool UE::Geometry::IsCapsuleMesh(const FDynamicMesh3& Mesh, FCapsule3d& CapsuleO
 		return false;
 	}
 
-	// See IsSphereMesh() for explanation of logic here, essentially we are checking that chordal deviation
-	// along edges is within tolerance. This works because endcaps are spheres and so sphere test applies,
-	// and along middle of capsule the deviation should be zero, but this is not currently measured.
-	// If false positives occur due to this, can check it by computing segment parameter, in t=[0,1] range distance should be epsilon-ish
+	// We use logic similar to IsSphereMesh() to test edge midpoints vs the capsule's spherical endcaps or cylindrical middle.
+	// When the edge midpoint projects to the cylindrical middle, we project the problem to a segment-aligned plane,
+	// so it becomes a test vs the circular cross-section.
 	double UseRadius = CapsuleOut.Radius;
 	double DeviationTol = 2.0 * UseRadius * RelativeDeviationTol;
+	double CosAngleTolerance = FMathd::Cos(FMathd::DegToRad * MaxAngleRangeDegrees);
+
 	for (int32 EdgeID : Mesh.EdgeIndicesItr())
 	{
 		FVector3d A, B;
 		Mesh.GetEdgeV(EdgeID, A, B);
 
+		FVector3d MidPoint = (A + B) * .5;
+		double ProjParam = CapsuleOut.Segment.Project(MidPoint);
+
+		// Find the relevant point on the segment to use for distance calculations
+		FVector3d RefSegmentPt;
+		if (ProjParam + UE_DOUBLE_KINDA_SMALL_NUMBER >= CapsuleOut.Segment.Extent)
+		{
+			RefSegmentPt = CapsuleOut.Segment.EndPoint();
+		}
+		else if (ProjParam - UE_DOUBLE_KINDA_SMALL_NUMBER <= -CapsuleOut.Segment.Extent)
+		{
+			RefSegmentPt = CapsuleOut.Segment.StartPoint();
+		}
+		else
+		{
+			// Cylinder case: do a projection so we only measure distances in the space of the circular cross section
+			A = A - (A - CapsuleOut.Segment.Center).Dot(CapsuleOut.Segment.Direction) * CapsuleOut.Segment.Direction;
+			B = B - (B - CapsuleOut.Segment.Center).Dot(CapsuleOut.Segment.Direction) * CapsuleOut.Segment.Direction;
+			MidPoint = (A + B) * .5;
+			RefSegmentPt = CapsuleOut.Segment.Center;
+		}
+		
+		// if a single edge spans too wide an angular range, the shape is too coarsely tesselated to be considered a capsule
+		FVector3d ToA = (A - RefSegmentPt).GetSafeNormal();
+		FVector3d ToB = (B - RefSegmentPt).GetSafeNormal();
+		if (ToA.Dot(ToB) <= CosAngleTolerance)
+		{
+			return false;
+		}
+
 		double HalfChordLen = Distance(A, B) * 0.5;
 		double MaxChordHeight = UseRadius - FMathd::Sqrt(UseRadius * UseRadius - HalfChordLen * HalfChordLen);   // "sagitta" height
 
-		double MidpointSignedDist = CapsuleOut.SignedDistance( (A + B)*0.5 );
+		double MidpointSignedDist = FVector3d::Distance(MidPoint, RefSegmentPt) - UseRadius;
 		if (FMathd::Abs(MidpointSignedDist) > (MaxChordHeight + DeviationTol))
 		{
 			return false;
 		}
 	}
 
+	// track where the vertices project to along the segment, to see if the capsule vertices cover the expected range along the capsule's major axis
+	FInterval1d ProjRange;
+	for (int32 VertID : Mesh.VertexIndicesItr())
+	{
+		FVector3d Vertex = Mesh.GetVertex(VertID);
+		double ProjParam = CapsuleOut.Segment.Project(Vertex);
+		ProjRange.Contain(ProjParam);
+	}
+
+	// Expect to come within DeviationTol of covering the full range of the capsule
+	if (ProjRange.Max < CapsuleOut.Segment.Extent + CapsuleOut.Radius - DeviationTol || ProjRange.Min > -CapsuleOut.Segment.Extent - CapsuleOut.Radius + DeviationTol)
+	{
+		return false;
+	}
+
 	return true;
 }
diff --git a/Engine/Plugins/Runtime/GeometryProcessing/Source/DynamicMesh/Public/ShapeApproximation/ShapeDetection3.h b/Engine/Plugins/Runtime/GeometryProcessing/Source/DynamicMesh/Public/ShapeApproximation/ShapeDetection3.h
--- a/Engine/Plugins/Runtime/GeometryProcessing/Source/DynamicMesh/Public/ShapeApproximation/ShapeDetection3.h
+++ b/Engine/Plugins/Runtime/GeometryProcessing/Source/DynamicMesh/Public/ShapeApproximation/ShapeDetection3.h
@@ -18,7 +18,7 @@ namespace UE
 		 * Detect if input Mesh is a meshed approximation of an analytic Sphere, and if so return best guess in SphereOut.
 		 * Fits a sphere to input points with several rounds of incremental improvement, then measures chordal deviation of edge midpoints.
 		 * @param RelativeDeviationTol Scaled by sphere diameter. The allowed difference in the distance from edge midpoints to the surface of the sphere vs the ideal distance for an edge of the same length.
-		 * @param MaxAngleRangeDegrees Maximum angle difference in vectors from center of sphere to the endpoints of any surface edge. Controls how coarsely tessellated a shpere can be before it is not considered a sphere.
+		 * @param MaxAngleRangeDegrees Maximum angle difference in vectors from center of sphere to the endpoints of any surface edge. Controls how coarsely tessellated a sphere can be before it is not considered a sphere.
 		 * @return true if mesh is a Sphere and SphereOut is initialized
 		 */
 		bool DYNAMICMESH_API IsSphereMesh(const FDynamicMesh3& Mesh, FSphere3d& SphereOut, double RelativeDeviationTol = 0.025, double MaxAngleRangeDegrees = 72);
@@ -38,10 +38,11 @@ namespace UE
 		/**
 		 * Detect if input Mesh is a meshed approximation of an analytic Capsule, and if so return best guess in CapsuleOut.
 		 * Fits a capsule to input points, then measures chordal deviation of edge midpoints.
-		 * @param RelativeDeviationTol distances from edge midpoints to capsule surface are allowed to deviate by 2*Radius*RelativeDeviationTol
+		 * @param RelativeDeviationTol Scaled by capsule diameter. The allowed difference in the distance from edge midpoints to the surface of the capsule vs the ideal distance for an edge of the same length.
+		 * @param MaxAngleRangeDegrees Maximum angle range an edge can cover over the cylinder or the spherical endcaps. Controls how coarsely tessellated a capsule can be before it is not considered a capsule.
 		 * @return true if mesh is a Capsule and CapsuleOut is initialized
 		 */
-		bool DYNAMICMESH_API IsCapsuleMesh(const FDynamicMesh3& Mesh, FCapsule3d& CapsuleOut, double RelativeDeviationTol = 0.025);
+		bool DYNAMICMESH_API IsCapsuleMesh(const FDynamicMesh3& Mesh, FCapsule3d& CapsuleOut, double RelativeDeviationTol = 0.025, double MaxAngleRangeDegrees = 72);
 
 	}
 }
\ No newline at end of file
