Geometry: fix copy-pasta in computation of Voronoi area for cotangent weights

diff --git a/Engine/Plugins/Runtime/GeometryProcessing/Source/DynamicMesh/Private/MeshCurvature.cpp b/Engine/Plugins/Runtime/GeometryProcessing/Source/DynamicMesh/Private/MeshCurvature.cpp
--- a/Engine/Plugins/Runtime/GeometryProcessing/Source/DynamicMesh/Private/MeshCurvature.cpp
+++ b/Engine/Plugins/Runtime/GeometryProcessing/Source/DynamicMesh/Private/MeshCurvature.cpp
@@ -110,7 +110,7 @@ double TGaussianCurvature(const FDynamicMesh3& mesh, int32 v_i, GetPositionFuncT
 			double cot_alpha_ij = VectorUtil::VectorCot(-Vik, Vkj);
 			cot_alpha_ij = FMathd::Clamp(cot_alpha_ij, -CotClampRange, CotClampRange);
 			double cot_beta_ik = VectorUtil::VectorCot(-Vij, -Vkj);
-			cot_beta_ik = FMathd::Clamp(cot_alpha_ij, -CotClampRange, CotClampRange);
+			cot_beta_ik = FMathd::Clamp(cot_beta_ik, -CotClampRange, CotClampRange);
 			MixedAreaSum += LenSqrVij * cot_alpha_ij * 0.125;
 			MixedAreaSum += LenSqrVik * cot_beta_ik * 0.125;
 		}
diff --git a/Engine/Plugins/Runtime/GeometryProcessing/Source/DynamicMesh/Private/MeshWeights.cpp b/Engine/Plugins/Runtime/GeometryProcessing/Source/DynamicMesh/Private/MeshWeights.cpp
--- a/Engine/Plugins/Runtime/GeometryProcessing/Source/DynamicMesh/Private/MeshWeights.cpp
+++ b/Engine/Plugins/Runtime/GeometryProcessing/Source/DynamicMesh/Private/MeshWeights.cpp
@@ -353,7 +353,7 @@ double TVoronoiArea(const FDynamicMesh3& mesh, int32 v_i, GetPositionFuncType Ge
 			double cot_alpha_ij = VectorUtil::VectorCot(Vki, Vkj);
 			cot_alpha_ij = FMathd::Clamp(cot_alpha_ij, -CotClampRange, CotClampRange);
 			double cot_beta_ik = VectorUtil::VectorCot(Vji, -Vkj);
-			cot_beta_ik = FMathd::Clamp(cot_alpha_ij, -CotClampRange, CotClampRange);
+			cot_beta_ik = FMathd::Clamp(cot_beta_ik, -CotClampRange, CotClampRange);
 			areaSum += Vji.SquaredLength() * cot_alpha_ij * 0.125;
 			areaSum += Vki.SquaredLength() * cot_beta_ik * 0.125;
 		}
