Add support for transparent meshes/faces in detect exterior visibility algorithm + associated geometry script function.

Also improve behavior (+ add control of behavior) for meshes with degenerate triangles, so meshes with degenerate triangles will not always be treated as visible.

diff --git a/Engine/Plugins/Experimental/GeometryScripting/Source/GeometryScriptingCore/Private/SceneUtilityFunctions.cpp b/Engine/Plugins/Experimental/GeometryScripting/Source/GeometryScriptingCore/Private/SceneUtilityFunctions.cpp
--- a/Engine/Plugins/Experimental/GeometryScripting/Source/GeometryScriptingCore/Private/SceneUtilityFunctions.cpp
+++ b/Engine/Plugins/Experimental/GeometryScripting/Source/GeometryScriptingCore/Private/SceneUtilityFunctions.cpp
@@ -370,6 +370,9 @@ void UGeometryScriptLibrary_SceneUtilityFunctions::DetermineMeshOcclusion(
 	const TArray<UDynamicMesh*>& SourceMeshes,
 	const TArray<FTransform>& SourceMeshTransforms,
 	TArray<bool>& OutMeshIsHidden,
+	const TArray<UDynamicMesh*>& TransparentMeshes,
+	const TArray<FTransform>& TransparentMeshTransforms,
+	TArray<bool>& OutTransparentMeshIsHidden,
 	const TArray<UDynamicMesh*>& OccludeMeshes,
 	const TArray<FTransform>& OccludeMeshTransforms,
 	const FGeometryScriptDetermineMeshOcclusionOptions& OcclusionOptions,
@@ -385,33 +388,46 @@ void UGeometryScriptLibrary_SceneUtilityFunctions::DetermineMeshOcclusion(
 		UE::Geometry::AppendError(Debug, EGeometryScriptErrorType::InvalidInputs, LOCTEXT("DetermineMeshOcclusion_OccludeArrayMismatch", "DetermineMeshOcclusion: OccludeMeshes and OccludeMeshTransforms arrays must have same length"));
 		return;
 	}
+	if (TransparentMeshes.Num() != TransparentMeshTransforms.Num())
+	{
+		UE::Geometry::AppendError(Debug, EGeometryScriptErrorType::InvalidInputs, LOCTEXT("DetermineMeshOcclusion_TransparentArrayMismatch", "DetermineMeshOcclusion: TransparentMeshes and TransparentMeshTransforms arrays must have same length"));
+		return;
+	}
 
 	FDetectPerDynamicMeshExteriorVisibility Occlusion;
-	for (int32 SourceMeshIndex = 0; SourceMeshIndex < SourceMeshes.Num(); ++SourceMeshIndex)
+	auto AddInstances = [](const TArray<UDynamicMesh*>& Meshes, const TArray<FTransform>& Transforms, TArray<FDetectPerDynamicMeshExteriorVisibility::FDynamicMeshInstance>& OutInstances) -> bool
 	{
-		UDynamicMesh* Mesh = SourceMeshes[SourceMeshIndex];
-		if (!Mesh)
+		for (int32 MeshIndex = 0; MeshIndex < Meshes.Num(); ++MeshIndex)
 		{
-			UE::Geometry::AppendError(Debug, EGeometryScriptErrorType::InvalidInputs, LOCTEXT("DetermineMeshOcclusion_InvalidSourceMesh", "DetermineMeshOcclusion: SourceMeshes array contained null mesh"));
-			return;
+			UDynamicMesh* Mesh = Meshes[MeshIndex];
+			if (!Mesh)
+			{
+				return false;
+			}
+			OutInstances.Emplace(Mesh->GetMeshPtr(), Transforms[MeshIndex]);
 		}
-		Occlusion.Instances.Emplace(Mesh->GetMeshPtr(), SourceMeshTransforms[SourceMeshIndex]);
+		return true;
+	};
+	if (!AddInstances(SourceMeshes, SourceMeshTransforms, Occlusion.Instances))
+	{
+		UE::Geometry::AppendError(Debug, EGeometryScriptErrorType::InvalidInputs, LOCTEXT("DetermineMeshOcclusion_InvalidSourceMesh", "DetermineMeshOcclusion: SourceMeshes array contained null mesh"));
+		return;
 	}
-	for (int32 OccludeMeshIndex = 0; OccludeMeshIndex < OccludeMeshes.Num(); ++OccludeMeshIndex)
+	if (!AddInstances(OccludeMeshes, OccludeMeshTransforms, Occlusion.OccludeInstances))
 	{
-		UDynamicMesh* Mesh = OccludeMeshes[OccludeMeshIndex];
-		if (!Mesh)
-		{
-			UE::Geometry::AppendError(Debug, EGeometryScriptErrorType::InvalidInputs, LOCTEXT("DetermineMeshOcclusion_InvalidOccludeMesh", "DetermineMeshOcclusion: OccludeMeshes array contained null mesh"));
-			return;
-		}
-		Occlusion.OccludeInstances.Emplace(Mesh->GetMeshPtr(), OccludeMeshTransforms[OccludeMeshIndex]);
+		UE::Geometry::AppendError(Debug, EGeometryScriptErrorType::InvalidInputs, LOCTEXT("DetermineMeshOcclusion_InvalidOccludeMesh", "DetermineMeshOcclusion: OccludeMeshes array contained null mesh"));
+		return;
+	}
+	if (!AddInstances(TransparentMeshes, TransparentMeshTransforms, Occlusion.TransparentInstances))
+	{
+		UE::Geometry::AppendError(Debug, EGeometryScriptErrorType::InvalidInputs, LOCTEXT("DetermineMeshOcclusion_InvalidTransparentMesh", "DetermineMeshOcclusion: TransparentMeshes array contained null mesh"));
+		return;
 	}
 	Occlusion.SamplingParameters.bDoubleSided = OcclusionOptions.bDoubleSided;
 	Occlusion.SamplingParameters.SamplingDensity = OcclusionOptions.SamplingDensity;
 	Occlusion.SamplingParameters.NumSearchDirections = OcclusionOptions.NumSearchDirections;
 
-	Occlusion.ComputeHidden(OutMeshIsHidden);
+	Occlusion.ComputeHidden(OutMeshIsHidden, &OutTransparentMeshIsHidden);
 }
 
 #undef LOCTEXT_NAMESPACE
\ No newline at end of file
diff --git a/Engine/Plugins/Experimental/GeometryScripting/Source/GeometryScriptingCore/Public/GeometryScript/SceneUtilityFunctions.h b/Engine/Plugins/Experimental/GeometryScripting/Source/GeometryScriptingCore/Public/GeometryScript/SceneUtilityFunctions.h
--- a/Engine/Plugins/Experimental/GeometryScripting/Source/GeometryScriptingCore/Public/GeometryScript/SceneUtilityFunctions.h
+++ b/Engine/Plugins/Experimental/GeometryScripting/Source/GeometryScriptingCore/Public/GeometryScript/SceneUtilityFunctions.h
@@ -121,15 +121,21 @@ public:
 	 * @param SourceMeshes			Meshes to test for occlusion. Note: The same mesh may appear multiple times in this array, if it is instanced with different transforms.
 	 * @param SourceMeshTransforms	A transform for each source mesh. Array must have the same length as SourceMeshes.
 	 * @param OutMeshIsHidden		Array will be filled with a bool per source mesh, indicating whether that mesh is hidden (true) or visible (false)
+	 * @param TransparentMeshes		Transparent source meshes, to test for occlusion but which do not occlude.
+	 * @param TransparentMeshTransforms		Array of transforms for each transparent mesh
+	 * @param OutTransparentMeshIsHidden	Array will be filled with a bool per transparent mesh, indicating whether that mesh is hidden (true) or visible (false)
 	 * @param OccludeMeshes			Array of optional meshes which can occlude SourceMeshes, but for which we will not test occlusion.
 	 * @param OccludeMeshTransforms	Array of transforms for each occlude mesh. Array must have the same length as OccludeMeshes.
 	 * @param Options				Settings to control how occlusion is tested
 	 */
-	UFUNCTION(BlueprintCallable, Category = "GeometryScript|Scene", meta = (AutoCreateRefTerm = "OccludeMeshes, OccludeMeshTransforms, OcclusionOptions"))
+	UFUNCTION(BlueprintCallable, Category = "GeometryScript|Scene", meta = (AutoCreateRefTerm = "SourceMeshes, SourceMeshTransforms, TransparentMeshes, TransparentMeshTransforms, OccludeMeshes, OccludeMeshTransforms, OcclusionOptions"))
 	static void DetermineMeshOcclusion(
 		const TArray<UDynamicMesh*>& SourceMeshes,
 		const TArray<FTransform>& SourceMeshTransforms,
 		TArray<bool>& OutMeshIsHidden,
+		const TArray<UDynamicMesh*>& TransparentMeshes,
+		const TArray<FTransform>& TransparentMeshTransforms,
+		TArray<bool>& OutTransparentMeshIsHidden,
 		const TArray<UDynamicMesh*>& OccludeMeshes,
 		const TArray<FTransform>& OccludeMeshTransforms,
 		const FGeometryScriptDetermineMeshOcclusionOptions& OcclusionOptions,
diff --git a/Engine/Plugins/Runtime/GeometryProcessing/Source/DynamicMesh/Private/Operations/DetectExteriorVisibility.cpp b/Engine/Plugins/Runtime/GeometryProcessing/Source/DynamicMesh/Private/Operations/DetectExteriorVisibility.cpp
--- a/Engine/Plugins/Runtime/GeometryProcessing/Source/DynamicMesh/Private/Operations/DetectExteriorVisibility.cpp
+++ b/Engine/Plugins/Runtime/GeometryProcessing/Source/DynamicMesh/Private/Operations/DetectExteriorVisibility.cpp
@@ -12,14 +12,18 @@
 
 using namespace UE::Geometry;
 
+namespace UE::Private::DetectExteriorVisibilityInternal
+{
 
-void FExteriorVisibilitySampling::ComputePerTriangleOcclusion(const FDynamicMesh3& TargetMesh, TArray<bool>& OutTriOccluded, bool bPerPolyGroup, const FExteriorVisibilitySampling& SamplingParameters, TArrayView<const bool> SkipTris)
+// Templated helper to compute occlusion with custom AABB tree (can be specialized for different mesh types)
+// Note that for the has-transparent case, we assume that Spatial *does not* include the transparent triangles
+template<typename SpatialType, bool bHasTransparent>
+void ComputePerTriangleOcclusionHelper(const FDynamicMesh3& TargetMesh, const SpatialType& Spatial, TArray<bool>& OutTriOccluded, bool bPerPolyGroup, const FExteriorVisibilitySampling& SamplingParameters, TArrayView<const bool> SkipTris, TFunctionRef<bool(int32)> IsTransparent)
 {
 	const bool bDoubleSided = SamplingParameters.bDoubleSided;
 	const double SampleRadius = FMath::Max(UE_DOUBLE_KINDA_SMALL_NUMBER, SamplingParameters.SamplingDensity);
 	int32 NumVisibilityTestDirections = FMath::Max(3, SamplingParameters.NumSearchDirections);
 
-	FDynamicMeshAABBTree3 Spatial(&TargetMesh, true);
 	FAxisAlignedBox3d Bounds = Spatial.GetBoundingBox();
 	double Radius = Bounds.DiagonalLength();
 
@@ -28,7 +32,7 @@ void FExteriorVisibilitySampling::ComputePerTriangleOcclusion(const FDynamicMesh
 	constexpr double TriScalingAlpha = 0.999;
 	constexpr double BaryCoordsThreshold = 0.001;
 
-	auto FindHitTriangleTest = [&](FVector3d TargetPosition, FVector3d TargetNormal, FVector3d FarPosition) -> int
+	auto FindHitTriangleTest = [&](FVector3d TargetPosition, FVector3d TargetNormal, FVector3d FarPosition, int32 TID, FVector3d& OutBaryCoords) -> int32
 	{
 		FVector3d RayDir(TargetPosition - FarPosition);
 		double Distance = Normalize(RayDir);
@@ -37,7 +41,30 @@ void FExteriorVisibilitySampling::ComputePerTriangleOcclusion(const FDynamicMesh
 			return IndexConstants::InvalidID;
 		}
 		FRay3d Ray(FarPosition, RayDir, true);
-		return Spatial.FindNearestHitTriangle(Ray, IMeshSpatial::FQueryOptions(Distance + 1.0));	// 1.0 is random fudge factor here...
+
+		int32 HitTID;
+		double UnusedNearestT;
+		if constexpr (bHasTransparent)
+		{
+			if (IsTransparent(TID))
+			{
+				bool bHit = Spatial.FindNearestHitTriangle(Ray, UnusedNearestT, HitTID, OutBaryCoords, IMeshSpatial::FQueryOptions(Distance));
+				if (!bHit)
+				{
+					// Note bary coords are only used to test if the sample is inside the triangle, so we don't accurate bary coords for this case
+					// If you copy/adapt this lambda to some other use case, you will probably need to fix this!
+					OutBaryCoords = FVector3d(1 / 3, 1 / 3, 1 / 3);
+					return TID; // assume we reached the source (transparent) triangle, which isn't in the AABB tree
+				}
+				else
+				{
+					return HitTID;
+				}
+			} // not transparent, fall through to the standard test
+		}
+
+		bool bHit = Spatial.FindNearestHitTriangle(Ray, UnusedNearestT, HitTID, OutBaryCoords, IMeshSpatial::FQueryOptions(Distance + 1.0)); // 1.0 is random fudge factor here...
+		return HitTID;
 	};
 
 	// array of (+/-)X/Y/Z directions
@@ -60,16 +87,26 @@ void FExteriorVisibilitySampling::ComputePerTriangleOcclusion(const FDynamicMesh
 	// outside the model, along the X/Y/Z directions and tri normal.
 	// If tri is hit we mark it as having 'known' status, allowing it
 	// to be skipped in the more expensive pass below
-	//
-	TArray<bool> TriStatusKnown;
+	// Degenerate tris are treated as a special case: They will be treated as visible,
+	// unless bPerPolyGroup==true and there are non-degenerate tris in the same group which are all not visible
+	enum class EVisStatus : uint8
+	{
+		Unknown = 0,
+		Visible = 1,
+		DegenerateOrInvalid = 2
+	};
+	TArray<EVisStatus> TriStatusKnown;
 	TriStatusKnown.SetNumZeroed(TargetMesh.MaxTriangleID());
 	bool bHasSkipTris = SkipTris.Num() == TargetMesh.MaxTriangleID();
 	int32 MaxUnknownTID = -1;
 	for (int32 TID : TargetMesh.TriangleIndicesItr())
 	{
 		bool bKnown = bHasSkipTris && SkipTris[TID];
-		TriStatusKnown[TID] = bKnown;
-		if (!bKnown)
+		if (bKnown)
+		{
+			TriStatusKnown[TID] = EVisStatus::Visible; // skip-tris are treated as visible
+		}
+		else
 		{
 			MaxUnknownTID = TID;
 		}
@@ -78,17 +115,18 @@ void FExteriorVisibilitySampling::ComputePerTriangleOcclusion(const FDynamicMesh
 	{
 		if (!TargetMesh.IsTriangle(TID))
 		{
-			TriStatusKnown[TID] = true;
+			TriStatusKnown[TID] = EVisStatus::DegenerateOrInvalid;
 			return;
 		}
 		FVector3d Normal, Centroid; double Area;
 		TargetMesh.GetTriInfo(TID, Normal, Area, Centroid);
 		if (Normal.SquaredLength() < 0.1 || Area <= FMathd::ZeroTolerance)
 		{
-			TriStatusKnown[TID] = true;
+			TriStatusKnown[TID] = EVisStatus::DegenerateOrInvalid;
 			return;
 		}
 
+		FVector3d UnusedBaryCoords;
 		for (FVector3d Direction : CardinalDirections)
 		{
 			// if direction is orthogonal to the triangle, hit-test is unstable, but even
@@ -97,16 +135,16 @@ void FExteriorVisibilitySampling::ComputePerTriangleOcclusion(const FDynamicMesh
 			// and manage to hit the search triangle
 			if (FMathd::Abs(Direction.Dot(Normal)) > GlancingAngleDotTolerance)
 			{
-				if (FindHitTriangleTest(Centroid, Normal, Centroid + Radius * Direction) == TID)
+				if (FindHitTriangleTest(Centroid, Normal, Centroid + Radius * Direction, TID, UnusedBaryCoords) == TID)
 				{
-					TriStatusKnown[TID] = true;
+					TriStatusKnown[TID] = EVisStatus::Visible;
 					return;
 				}
 			}
 		}
-		if (FindHitTriangleTest(Centroid, Normal, Centroid + Radius * Normal) == TID)
+		if (FindHitTriangleTest(Centroid, Normal, Centroid + Radius * Normal, TID, UnusedBaryCoords) == TID)
 		{
-			TriStatusKnown[TID] = true;
+			TriStatusKnown[TID] = EVisStatus::Visible;
 			return;
 		}
 
@@ -114,20 +152,28 @@ void FExteriorVisibilitySampling::ComputePerTriangleOcclusion(const FDynamicMesh
 	});
 
 	// Set up PolyGroup visibility if requested
-	TArray<bool> GroupStatusKnown;
+	TArray<EVisStatus> GroupStatusKnown;
 	TArray<std::atomic<bool>> ThreadSafeGroupVisible;
 	if (bPerPolyGroup)
 	{
 		// Init initial group status from initial tri status
-		GroupStatusKnown.SetNumZeroed(TargetMesh.MaxGroupID());
+		GroupStatusKnown.Init(EVisStatus::DegenerateOrInvalid, TargetMesh.MaxGroupID());
 		for (int32 TID : TargetMesh.TriangleIndicesItr())
 		{
-			if (TriStatusKnown[TID])
+			if (TriStatusKnown[TID] == EVisStatus::Visible) // group has known-visible tri
 			{
 				int32 GID = TargetMesh.GetTriangleGroup(TID);
 				if (GroupStatusKnown.IsValidIndex(GID))
 				{
-					GroupStatusKnown[GID] = true;
+					GroupStatusKnown[GID] = EVisStatus::Visible;
+				}
+			}
+			else if (TriStatusKnown[TID] == EVisStatus::Unknown) // group has non-degenerate tri
+			{
+				int32 GID = TargetMesh.GetTriangleGroup(TID);
+				if (GroupStatusKnown.IsValidIndex(GID) && GroupStatusKnown[GID] == EVisStatus::DegenerateOrInvalid)
+				{
+					GroupStatusKnown[GID] = EVisStatus::Unknown;
 				}
 			}
 		}
@@ -135,15 +181,15 @@ void FExteriorVisibilitySampling::ComputePerTriangleOcclusion(const FDynamicMesh
 		ThreadSafeGroupVisible.SetNum(TargetMesh.MaxGroupID());
 		for (int32 GroupIdx = 0; GroupIdx < GroupStatusKnown.Num(); ++GroupIdx)
 		{
-			ThreadSafeGroupVisible[GroupIdx] = GroupStatusKnown[GroupIdx];
+			ThreadSafeGroupVisible[GroupIdx] = GroupStatusKnown[GroupIdx] == EVisStatus::Visible;
 		}
 		// Copy group statuses back to tri statuses
 		for (int32 TID : TargetMesh.TriangleIndicesItr())
 		{
 			int32 GID = TargetMesh.GetTriangleGroup(TID);
-			if (GroupStatusKnown.IsValidIndex(GID) && GroupStatusKnown[GID])
+			if (GroupStatusKnown.IsValidIndex(GID) && GroupStatusKnown[GID] == EVisStatus::Visible)
 			{
-				TriStatusKnown[TID] = true;
+				TriStatusKnown[TID] = EVisStatus::Visible;
 			}
 		}
 	}
@@ -154,9 +200,9 @@ void FExteriorVisibilitySampling::ComputePerTriangleOcclusion(const FDynamicMesh
 	MaxUnknownTID = -1; // re-compute max unknown
 	for (int32 TID : TargetMesh.TriangleIndicesItr())
 	{
-		bool bKnown = TriStatusKnown[TID];
-		ThreadSafeTriVisible[TID] = bKnown;
-		if (!bKnown)
+		EVisStatus Status = TriStatusKnown[TID];
+		ThreadSafeTriVisible[TID] = Status == EVisStatus::Visible;
+		if (Status == EVisStatus::Unknown)
 		{
 			MaxUnknownTID = TID;
 		}
@@ -202,13 +248,16 @@ void FExteriorVisibilitySampling::ComputePerTriangleOcclusion(const FDynamicMesh
 		{
 			GID = TargetMesh.GetTriangleGroup(TID);
 			bHasValidGroup = GroupStatusKnown.IsValidIndex(GID);
-			if (bHasValidGroup && (GroupStatusKnown[GID] || ThreadSafeGroupVisible[GID]))
+			if (bHasValidGroup && (GroupStatusKnown[GID] != EVisStatus::Unknown || ThreadSafeGroupVisible[GID]))
 			{
 				return;
 			}
 		}
 		// if we already found out this triangle is visible or hidden, we can skip it
-		if (TriStatusKnown[TID] || ThreadSafeTriVisible[TID]) return;
+		if (TriStatusKnown[TID] != EVisStatus::Unknown || ThreadSafeTriVisible[TID])
+		{
+			return;
+		}
 
 		FVector3d A, B, C;
 		TargetMesh.GetTriVertices(TID, A, B, C);
@@ -240,6 +289,11 @@ void FExteriorVisibilitySampling::ComputePerTriangleOcclusion(const FDynamicMesh
 		{
 			double a1 = RandomStream.GetFraction();
 			double a2 = RandomStream.GetFraction();
+			if (a1 + a2 > 1) // if we would sample outside the triangle, reflect the sample
+			{
+				a1 = 1 - a1;
+				a2 = 1 - a2;
+			}
 			FVector2d PointUV = UVTriangle.V[0] + a1 * V1 + a2 * V2;
 			if (UVTriangle.IsInside(PointUV))
 			{
@@ -257,23 +311,22 @@ void FExteriorVisibilitySampling::ComputePerTriangleOcclusion(const FDynamicMesh
 					}
 
 					FVector3d RayFrom = Position + 2.0 * Radius * VisibilityDirections[k];
-					int32 HitTriID = FindHitTriangleTest(Position, TriNormal, RayFrom);
-					if (HitTriID != IndexConstants::InvalidID && TriStatusKnown[HitTriID] == false)
+					FVector3d BaryCoords;
+					int32 HitTriID = FindHitTriangleTest(Position, TriNormal, RayFrom, TID, BaryCoords);
+					if (HitTriID != IndexConstants::InvalidID && TriStatusKnown[HitTriID] == EVisStatus::Unknown)
 					{
-						// Want to filter out on-edge triangle hits, as they are generally spurious and
+						if (HitTriID == TID)
+						{
+							HitTris.Add(HitTriID);
+							break;
+						}
+						// For tris other than the one that was sampled to create the ray, we want
+						// to filter out on-edge triangle hits, as they are generally spurious and
 						// will result in interior triangles remaining visible
-						FRay3d Ray(RayFrom, Normalized(Position - RayFrom), true);
-						FIntrRay3Triangle3d RayHit = TMeshQueries<FDynamicMesh3>::RayTriangleIntersection(TargetMesh, HitTriID, Ray);
-						if (RayHit.IntersectionType == EIntersectionType::Point &&
-							RayHit.TriangleBaryCoords.GetMin() > BaryCoordsThreshold &&
-							RayHit.TriangleBaryCoords.GetMax() < (1.0 - BaryCoordsThreshold))
+						else if (BaryCoords.GetMin() > BaryCoordsThreshold &&
+							BaryCoords.GetMax() < (1.0 - BaryCoordsThreshold))
 						{
-
-							HitTris.AddUnique(HitTriID);		// we hit some triangle, whether or not it is the one we are testing...
-							if (HitTriID == TID)
-							{
-								break;
-							}
+							HitTris.AddUnique(HitTriID);
 						}
 					}
 				}
@@ -294,7 +347,7 @@ void FExteriorVisibilitySampling::ComputePerTriangleOcclusion(const FDynamicMesh
 					for (int32 HitTriID : HitTris)
 					{
 						int32 HitGID = TargetMesh.GetTriangleGroup(HitTriID);
-						if (GroupStatusKnown.IsValidIndex(HitGID) && !GroupStatusKnown[HitGID] && HitGroups.Num() == HitGroups.AddUnique(HitGID))
+						if (GroupStatusKnown.IsValidIndex(HitGID) && GroupStatusKnown[HitGID] == EVisStatus::Unknown && HitGroups.Num() == HitGroups.AddUnique(HitGID))
 						{
 							// Note: setting/checking the group visibility flag is optional and just allow threads to stop early
 							// so it should be fine to do so w/ the relaxed memory ordering
@@ -332,7 +385,19 @@ void FExteriorVisibilitySampling::ComputePerTriangleOcclusion(const FDynamicMesh
 		for (int32 TID : TargetMesh.TriangleIndicesItr())
 		{
 			int32 GID = TargetMesh.GetTriangleGroup(TID);
-			if (ThreadSafeGroupVisible.IsValidIndex(GID) && ThreadSafeGroupVisible[GID])
+			// Set the tris in the group to visible if the group is visible or invalid (where invalid corresponds to entirely degenerate / untestable)
+			if (ThreadSafeGroupVisible.IsValidIndex(GID) && 
+				(ThreadSafeGroupVisible[GID] || (SamplingParameters.bMarkDegenerateAsVisible && GroupStatusKnown[GID] == EVisStatus::DegenerateOrInvalid)) )
+			{
+				ThreadSafeTriVisible[TID] = true;
+			}
+		}
+	}
+	else if (SamplingParameters.bMarkDegenerateAsVisible)
+	{
+		for (int32 TID : TargetMesh.TriangleIndicesItr())
+		{
+			if (TriStatusKnown[TID] == EVisStatus::DegenerateOrInvalid)
 			{
 				ThreadSafeTriVisible[TID] = true;
 			}
@@ -346,14 +411,110 @@ void FExteriorVisibilitySampling::ComputePerTriangleOcclusion(const FDynamicMesh
 		OutTriOccluded[Idx] = !ThreadSafeTriVisible[Idx];
 	}
 }
+} // namespace UE::Private::DetectExteriorVisibilityInternal
+
+void FExteriorVisibilitySampling::ComputePerTriangleOcclusion(const FDynamicMesh3& TargetMesh, TArray<bool>& OutTriOccluded, bool bPerPolyGroup, const FExteriorVisibilitySampling& SamplingParameters, TArrayView<const bool> SkipTris)
+{
+	using namespace UE::Private::DetectExteriorVisibilityInternal;
+	TMeshAABBTree3<FDynamicMesh3> Spatial(&TargetMesh, true);
+	return ComputePerTriangleOcclusionHelper< TMeshAABBTree3<FDynamicMesh3>, false >(TargetMesh, Spatial, OutTriOccluded, bPerPolyGroup, SamplingParameters, SkipTris, [](int32)->bool { checkNoEntry(); return false; });
+}
+
+// simple wrapper that skips transparent triangles, used for building an aabb tree without them
+struct FSkipTransparentMeshWrapper
+{
+	const FDynamicMesh3* Mesh;
+	TFunctionRef<bool(int32)> IsTriTransparent;
+	const int32 TriCount;
+
+	static int32 CountNonTransparent(const FDynamicMesh3* Mesh, TFunctionRef<bool(int32)> IsTriTransparent)
+	{
+		int32 Count = 0;
+		for (int32 TID = 0; TID < Mesh->MaxTriangleID(); ++TID)
+		{
+			Count += int(Mesh->IsTriangle(TID) && !IsTriTransparent(TID));
+		}
+		return Count;
+	}
+
+	FSkipTransparentMeshWrapper(const FDynamicMesh3* Mesh, TFunctionRef<bool(int32)> IsTriTransparent) : Mesh(Mesh), IsTriTransparent(IsTriTransparent), TriCount(CountNonTransparent(Mesh, IsTriTransparent))
+	{
+	}
+
+	FORCEINLINE bool IsTriangle(int32 Index) const
+	{
+		return Mesh->IsTriangle(Index) && !IsTriTransparent(Index);
+	}
+
+	FORCEINLINE bool IsVertex(int32 Index) const
+	{
+		return Mesh->IsVertex(Index);
+	}
+
+	FORCEINLINE int32 MaxTriangleID() const
+	{
+		return Mesh->MaxTriangleID();
+	}
+
+	FORCEINLINE int32 MaxVertexID() const
+	{
+		return Mesh->MaxVertexID();
+	}
+
+	// Triangle count is pre-computed to subtract the transparent triangles
+	FORCEINLINE int32 TriangleCount() const
+	{
+		return TriCount;
+	}
+
+	FORCEINLINE int32 VertexCount() const
+	{
+		return Mesh->VertexCount();
+	}
+
+	FORCEINLINE uint64 GetChangeStamp() const
+	{
+		return Mesh->GetChangeStamp();
+	}
+
+	FORCEINLINE FIndex3i GetTriangle(int32 Index) const
+	{
+		return Mesh->GetTriangle(Index);
+	}
+
+	FORCEINLINE FVector3d GetVertex(int32 Index) const
+	{
+		return Mesh->GetVertex(Index);
+	}
+
+	FORCEINLINE void GetTriVertices(int32 TriIndex, FVector3d& V0, FVector3d& V1, FVector3d& V2) const
+	{
+		return Mesh->GetTriVertices(TriIndex, V0, V1, V2);
+	}
+};
+
+void FExteriorVisibilitySampling::ComputePerTriangleOcclusion(const FDynamicMesh3& TargetMesh, TFunctionRef<bool(int32)> IsTriTransparent, TArray<bool>& OutTriOccluded, bool bPerPolyGroup, const FExteriorVisibilitySampling& SamplingParameters, TArrayView<const bool> SkipTris)
+{
+	using namespace UE::Private::DetectExteriorVisibilityInternal;
+	FSkipTransparentMeshWrapper SkipTransparent(&TargetMesh, IsTriTransparent);
+	TMeshAABBTree3<FSkipTransparentMeshWrapper> Spatial(&SkipTransparent, true);
+	return ComputePerTriangleOcclusionHelper< TMeshAABBTree3<FSkipTransparentMeshWrapper>, true >(TargetMesh, Spatial, OutTriOccluded, bPerPolyGroup, SamplingParameters, SkipTris, IsTriTransparent);
+}
 
 
-void FDetectPerDynamicMeshExteriorVisibility::ComputeHidden(TArray<bool>& OutIsOccluded)
+void FDetectPerDynamicMeshExteriorVisibility::ComputeHidden(TArray<bool>& OutIsOccluded, TArray<bool>* OutTransparentIsOccluded)
 {
 	OutIsOccluded.SetNumZeroed(Instances.Num());
-	if (Instances.Num() < 2)
+
+	bool bTestTransparent = OutTransparentIsOccluded && TransparentInstances.Num();
+	if (OutTransparentIsOccluded)
+	{
+		OutTransparentIsOccluded->SetNumZeroed(TransparentInstances.Num());
+	}
+
+	if (Instances.Num() == 0 && (!bTestTransparent || TransparentInstances.Num() == 0))
 	{
-		// 0 or 1 meshes can't be fully occluded
+		// no meshes to test
 		return;
 	}
 
@@ -376,16 +537,33 @@ void FDetectPerDynamicMeshExteriorVisibility::ComputeHidden(TArray<bool>& OutIsO
 		}
 	};
 	AddInstances(Instances);
+	int32 TransparentStartID = AllMesh.MaxTriangleID();
+	if (bTestTransparent)
+	{
+		AddInstances(TransparentInstances);
+	}
 	int32 OccludeStartID = AllMesh.MaxTriangleID();
 	AddInstances(OccludeInstances);
 
-	// Should be safe to assume contiguous triangle IDs because we've constructed this mesh ourselves
-	// if this assumption breaks, will need to change the logic below which assigns PolyGroup IDs per instance
+	// Should be safe to assume contiguous triangle IDs because we've constructed this mesh ourselves.
+	// If this assumption breaks, will need to change the logic below which assigns PolyGroup IDs per instance
 	if (!ensure(AllMesh.IsCompactT()))
 	{
 		return;
 	}
 
+	// Initialize output arrays and track the boundary between normal, transparent and occlude instances
+	OutIsOccluded.Init(true, Instances.Num());
+	int32 TestInstancesCount = OutIsOccluded.Num();
+	int32 TransparentInstanceStart = TestInstancesCount;
+	int32 OccludeInstancesStart = TestInstancesCount;
+	if (bTestTransparent)
+	{
+		OutTransparentIsOccluded->Init(true, TransparentInstances.Num());
+		TestInstancesCount += TransparentInstances.Num();
+		OccludeInstancesStart = TestInstancesCount;
+	}
+
 	// Enable and set triangle groups *after* adding all meshes, so the groups initialize to zero and we are not affected by groups on the source meshes
 	AllMesh.EnableTriangleGroups(0);
 	for (int32 TID = InstanceEndTID[0], CurInstance = 1; TID < AllMesh.MaxTriangleID(); ++TID)
@@ -395,7 +573,7 @@ void FDetectPerDynamicMeshExteriorVisibility::ComputeHidden(TArray<bool>& OutIsO
 			CurInstance++;
 		}
 		// Set triangle group by instance ID, but give all OccludeInstances the same ID
-		AllMesh.SetTriangleGroup(TID, CurInstance < Instances.Num() ? CurInstance : Instances.Num());
+		AllMesh.SetTriangleGroup(TID, FMath::Min(CurInstance, TestInstancesCount));
 	}
 
 	// Mark any triangles from Occlude Instances as SkipTris
@@ -408,24 +586,40 @@ void FDetectPerDynamicMeshExteriorVisibility::ComputeHidden(TArray<bool>& OutIsO
 
 	// Compute per-tri occlusion
 	TArray<bool> TriOccluded;
-	FExteriorVisibilitySampling::ComputePerTriangleOcclusion(AllMesh, TriOccluded, true /*bPerPolyGroup*/, SamplingParameters, SkipTris);
+	if (bTestTransparent)
+	{
+		FExteriorVisibilitySampling::ComputePerTriangleOcclusion(AllMesh, 
+			[TransparentStartID, OccludeStartID](int32 TID) 
+			{
+				return TID >= TransparentStartID && TID < OccludeStartID;
+			}, TriOccluded, true /*bPerPolyGroup*/, SamplingParameters, SkipTris);
+	}
+	else
+	{
+		FExteriorVisibilitySampling::ComputePerTriangleOcclusion(AllMesh, TriOccluded, true /*bPerPolyGroup*/, SamplingParameters, SkipTris);
+	}
 
 	// Transfer triangle occlusion back to source mesh occlusion via the triangle groups
-	OutIsOccluded.Init(true, Instances.Num());
 	for (int32 TID : AllMesh.TriangleIndicesItr())
 	{
 		int32 GID = AllMesh.GetTriangleGroup(TID);
-		if (GID < 0 || GID >= Instances.Num()) // tri is from an OccludeInstances mesh
+		if (GID < 0 || GID >= OccludeInstancesStart) // tri is from an OccludeInstances mesh
 		{
 			break; // OccludeInstances meshes are at the end; we can stop when we reach them
 		}
-		if (!OutIsOccluded[GID]) // already found visible
+		TArray<bool>* UseIsOccluded = &OutIsOccluded;
+		if (OutTransparentIsOccluded && GID >= TransparentInstanceStart)
+		{
+			UseIsOccluded = OutTransparentIsOccluded;
+		}
+		int32 InstanceIdx = GID < TransparentInstanceStart ? GID : GID - TransparentInstanceStart;
+		if (!(*UseIsOccluded)[InstanceIdx]) // already found visible
 		{
 			continue;
 		}
 		if (!TriOccluded[TID]) // found visible tri -> source is visible
 		{
-			OutIsOccluded[GID] = false;
+			(*UseIsOccluded)[InstanceIdx] = false;
 		}
 	}
 }
\ No newline at end of file
diff --git a/Engine/Plugins/Runtime/GeometryProcessing/Source/DynamicMesh/Public/Operations/DetectExteriorVisibility.h b/Engine/Plugins/Runtime/GeometryProcessing/Source/DynamicMesh/Public/Operations/DetectExteriorVisibility.h
--- a/Engine/Plugins/Runtime/GeometryProcessing/Source/DynamicMesh/Public/Operations/DetectExteriorVisibility.h
+++ b/Engine/Plugins/Runtime/GeometryProcessing/Source/DynamicMesh/Public/Operations/DetectExteriorVisibility.h
@@ -23,13 +23,25 @@ struct FExteriorVisibilitySampling
 	// Number of directions to test for visibility
 	int32 NumSearchDirections = 128;
 
+	// Whether to mark degenerate tris as visible. If bPerPolyGroup is true, this setting controls the classification of PolyGroups that *only* have degenerate tris.
+	bool bMarkDegenerateAsVisible = true;
+
 	// Compute per-triangle visibility array for a triangle mesh
-	// @param OutTriOccluded	For each valid triangle ID, OutTriOccluded[ID] will be true if that triangle is hidden, false if it is visible
-	// @param SkipTris			True for triangles do not need visibility testing (will be marked as not occluded). If empty, no triangles will be skipped.
-	// @param bPerPolyGroup		If true, visiblity will be determined on a per group basis: If any triangle in a group is visible, the whole group will be marked as not-occluded
-	// @param bDoubleSided		Whether to treat faces a double-sided
-	// @param SamplingDensity	Approximate spacing between triangle samples used for visibility tests
-	DYNAMICMESH_API static void ComputePerTriangleOcclusion(const FDynamicMesh3& Mesh, TArray<bool>& OutTriOccluded, bool bPerPolyGroup, const FExteriorVisibilitySampling& SamplingParameters, TArrayView<const bool> SkipTris = TArrayView<const bool>());
+	// @param OutTriOccluded		For each valid triangle ID, OutTriOccluded[ID] will be true if that triangle is hidden, false if it is visible
+	// @param bPerPolyGroup			If true, visiblity will be determined on a per group basis: If any triangle in a group is visible, the whole group will be marked as not-occluded
+	// @param SamplingParameters	Parameters controlling how exterior visiblity is tested
+	// @param SkipTris				True for triangles do not need visibility testing (will be marked as not occluded). If empty, no triangles will be skipped.
+	DYNAMICMESH_API static void ComputePerTriangleOcclusion(const FDynamicMesh3& Mesh, TArray<bool>& OutTriOccluded, bool bPerPolyGroup, const FExteriorVisibilitySampling& SamplingParameters, 
+		TArrayView<const bool> SkipTris = TArrayView<const bool>());
+
+	// Compute per-triangle visibility array for a triangle mesh, with support for transparent (non-occluding) triangles
+	// @param IsTriTransparent		Indicator function returns true if a triangle is 'transparent' (should not occlude, but can still be occluded)
+	// @param OutTriOccluded		For each valid triangle ID, OutTriOccluded[ID] will be true if that triangle is hidden, false if it is visible
+	// @param bPerPolyGroup			If true, visiblity will be determined on a per group basis: If any triangle in a group is visible, the whole group will be marked as not-occluded
+	// @param SamplingParameters	Parameters controlling how exterior visiblity is tested
+	// @param SkipTris				True for triangles do not need visibility testing (will be marked as not occluded). If empty, no triangles will be skipped.
+	DYNAMICMESH_API static void ComputePerTriangleOcclusion(const FDynamicMesh3& Mesh, TFunctionRef<bool(int32)> IsTriTransparent, TArray<bool>& OutTriOccluded, bool bPerPolyGroup, const FExteriorVisibilitySampling& SamplingParameters,
+		TArrayView<const bool> SkipTris = TArrayView<const bool>());
 };
 
 // Determine which meshes are visible from exterior views of a list of meshes
@@ -58,12 +70,16 @@ public:
 	// Instances which occlude visibility, but for which we don't need to determine visibility
 	TArray<FDynamicMeshInstance> OccludeInstances;
 
+	// Instances which do not occlude visibility, but for which we do need to determine visibility
+	TArray<FDynamicMeshInstance> TransparentInstances;
+
 	// Parameters controlling how the triangles are sampled for visibility
 	FExteriorVisibilitySampling SamplingParameters;
 
 	// Compute which source meshes are hidden
-	// @param OutIsOccluded		Indicate if each mesh in Instances is hidden (true) or visible (false)
-	DYNAMICMESH_API void ComputeHidden(TArray<bool>& OutIsOccluded);
+	// @param OutIsOccluded				Indicate if each mesh in Instances is hidden (true) or visible (false)
+	// @param OutTransparentIsOccluded	If provided, indicate if each mesh in TransparentInstances is hidden (true) or visible (false). If not provided, TransparentInstances are not tested.
+	DYNAMICMESH_API void ComputeHidden(TArray<bool>& OutIsOccluded, TArray<bool>* OutTransparentIsOccluded = nullptr);
 
 };
 
