CIS fix for static analyzer warning in DetectExteriorVisibility.cpp (1/3 should be 1.0/3.0)

diff --git a/Engine/Plugins/Runtime/GeometryProcessing/Source/DynamicMesh/Private/Operations/DetectExteriorVisibility.cpp b/Engine/Plugins/Runtime/GeometryProcessing/Source/DynamicMesh/Private/Operations/DetectExteriorVisibility.cpp
--- a/Engine/Plugins/Runtime/GeometryProcessing/Source/DynamicMesh/Private/Operations/DetectExteriorVisibility.cpp
+++ b/Engine/Plugins/Runtime/GeometryProcessing/Source/DynamicMesh/Private/Operations/DetectExteriorVisibility.cpp
@@ -53,7 +53,7 @@ void ComputePerTriangleOcclusionHelper(const FDynamicMesh3& TargetMesh, const Sp
 				{
 					// Note bary coords are only used to test if the sample is inside the triangle, so we don't accurate bary coords for this case
 					// If you copy/adapt this lambda to some other use case, you will probably need to fix this!
-					OutBaryCoords = FVector3d(1 / 3, 1 / 3, 1 / 3);
+					OutBaryCoords = FVector3d(1.0 / 3.0, 1.0 / 3.0, 1.0 / 3.0);
 					return TID; // assume we reached the source (transparent) triangle, which isn't in the AABB tree
 				}
 				else
