GeometryProcessing: various utility functions to assist with rounded bevels implementation
- add 2D version of VectorTanHalfAngle to VectorUtil.h
- add FDynamicMeshUVEditor::SetToPerVertexUVs, initializes overlay to have a single UV for each vertex
- add FQuadGridPatch::InitializeFromQuadPatch(), initializes the quad grid from external rows of quads and associated vertices
- add FQuadGridPatch::GetVertexColumn(), returns the list of vertices in a specific grid column
- add FQuadGridPatch::FindColumnIndex(), finds the index of the grid column that contains a specific Vertex ID

diff --git a/Engine/Plugins/Runtime/GeometryProcessing/Source/DynamicMesh/Private/Parameterization/DynamicMeshUVEditor.cpp b/Engine/Plugins/Runtime/GeometryProcessing/Source/DynamicMesh/Private/Parameterization/DynamicMeshUVEditor.cpp
--- a/Engine/Plugins/Runtime/GeometryProcessing/Source/DynamicMesh/Private/Parameterization/DynamicMeshUVEditor.cpp
+++ b/Engine/Plugins/Runtime/GeometryProcessing/Source/DynamicMesh/Private/Parameterization/DynamicMeshUVEditor.cpp
@@ -194,6 +194,36 @@ void FDynamicMeshUVEditor::TransformUVElements(const TArray<int32>& ElementIDs,
 }
 
 
+void FDynamicMeshUVEditor::SetToPerVertexUVs(TArray<int32>& VertexToUVOut, bool& bIsIdentityMapOut, FUVEditResult* Result)
+{
+	bIsIdentityMapOut = true;
+	VertexToUVOut.Init(IndexConstants::InvalidID, Mesh->MaxVertexID());
+
+	UVOverlay->ClearElements();
+	for (int32 VertexID : Mesh->VertexIndicesItr())
+	{
+		int32 UVID = UVOverlay->AppendElement(FVector2f::Zero());
+		VertexToUVOut[VertexID] = UVID;
+		bIsIdentityMapOut &= (UVID == VertexID);
+
+		if (Result)
+		{
+			Result->NewUVElements.Add(UVID);
+		}
+	}
+
+	for (int32 TriangleID : Mesh->TriangleIndicesItr())
+	{
+		FIndex3i Tri = Mesh->GetTriangle(TriangleID);
+		Tri.A = VertexToUVOut[Tri.A];
+		Tri.B = VertexToUVOut[Tri.B];
+		Tri.C = VertexToUVOut[Tri.C];
+		UVOverlay->SetTriangle(TriangleID, Tri);
+	}
+
+}
+
+
 template<typename EnumerableType>
 void InternalSetPerTriangleUVs(EnumerableType TriangleIDs, const FDynamicMesh3* Mesh, FDynamicMeshUVOverlay* UVOverlay, double ScaleFactor, FUVEditResult* Result)
 {
diff --git a/Engine/Plugins/Runtime/GeometryProcessing/Source/DynamicMesh/Private/Selections/QuadGridPatch.cpp b/Engine/Plugins/Runtime/GeometryProcessing/Source/DynamicMesh/Private/Selections/QuadGridPatch.cpp
--- a/Engine/Plugins/Runtime/GeometryProcessing/Source/DynamicMesh/Private/Selections/QuadGridPatch.cpp
+++ b/Engine/Plugins/Runtime/GeometryProcessing/Source/DynamicMesh/Private/Selections/QuadGridPatch.cpp
@@ -10,27 +10,29 @@ using namespace UE::Geometry;
 namespace UELocal
 {
 
-static bool FindTriEdgeIndices(
-	const FIndex3i& Triangle, 
-	int32 EdgeVert0, int32 EdgeVert1,
-	int32& EdgeIndex0, int32& EdgeIndex1, int32& OtherIndex2 )
-{
-	for ( int32 j = 0; j < 3; ++j )
+	static bool FindTriEdgeIndices(
+		const FIndex3i& Triangle,
+		int32 EdgeVert0, int32 EdgeVert1,
+		int32& EdgeIndex0, int32& EdgeIndex1, int32& OtherIndex2)
 	{
-		if (Triangle[j] == EdgeVert0)
+		for (int32 j = 0; j < 3; ++j)
 		{
-			if ( Triangle[(j+1)%3] == EdgeVert1 )
-			{
-				EdgeIndex0 = j; EdgeIndex1 = (j+1)%3; OtherIndex2 = (j+2)%3; return true;
-			}
-			else if ( Triangle[(j+2)%3] == EdgeVert1 )
+			if (Triangle[j] == EdgeVert0)
 			{
-				EdgeIndex0 = j; EdgeIndex1 = (j+2)%3; OtherIndex2 = (j+1)%3; return true;
+				if (Triangle[(j + 1) % 3] == EdgeVert1)
+				{
+					EdgeIndex0 = j; EdgeIndex1 = (j + 1) % 3; OtherIndex2 = (j + 2) % 3; 
+					return true;
+				}
+				else if (Triangle[(j + 2) % 3] == EdgeVert1)
+				{
+					EdgeIndex0 = j; EdgeIndex1 = (j + 2) % 3; OtherIndex2 = (j + 1) % 3; 
+					return true;
+				}
 			}
 		}
+		return false;
 	}
-	return false;
-}
 
 }
 
@@ -81,6 +83,90 @@ void FQuadGridPatch::InitializeFromQuadStrip(const FDynamicMesh3& Mesh, const TA
 
 
 
+bool FQuadGridPatch::InitializeFromQuadPatch(const FDynamicMesh3& Mesh, const TArray<TArray<FIndex2i>>& QuadRowsIn, const TArray<TArray<int32>>& VertexSpansIn)
+{
+	int32 NumV = VertexSpansIn[0].Num();
+	int32 NumQ = QuadRowsIn[0].Num();
+	if (NumQ != NumV - 1 || QuadRowsIn.Num() != VertexSpansIn.Num()-1 )
+	{
+		ensure(false);
+		return false;
+	}
+
+	NumVertexColsU = NumV;
+	NumVertexRowsV = VertexSpansIn.Num();
+	for (int32 j = 1; j < NumVertexRowsV; ++j)
+	{
+		if (VertexSpansIn[j].Num() != VertexSpansIn[0].Num())
+		{
+			ensure(false);
+			return false;
+		}
+	}
+	int32 NumQuadRows = QuadRowsIn.Num();
+	for (int32 j = 1; j < NumQuadRows; ++j)
+	{
+		if (QuadRowsIn[j].Num() != QuadRowsIn[0].Num())
+		{
+			ensure(false);
+			return false;
+		}
+	}
+
+	VertexSpans = VertexSpansIn;
+	QuadTriangles = QuadRowsIn;
+
+	bool bAllOK = true;
+	for (int32 j = 0; j < NumQuadRows && bAllOK; ++j)
+	{
+		for (int32 k = 0; k < NumQ && bAllOK; ++k)
+		{
+			// these should be the four vertices of the quad
+			int VertexA = VertexSpans[j][k];
+			int VertexB = VertexSpans[j][k+1];
+			int VertexC = VertexSpans[j+1][k+1];
+			int VertexD = VertexSpans[j+1][k];
+
+			// figure out unique vertices from the quad info, and verify that there are 4 and they are A/B/C/D
+			FIndex2i& QuadTris = QuadTriangles[j][k];
+			FIndex3i TriA = Mesh.GetTriangle(QuadTris.A);
+			FIndex3i TriB = Mesh.GetTriangle(QuadTris.B);
+			TArray<int32, TInlineAllocator<6>> TriVerts({ TriA.A, TriA.B, TriA.C });
+			TriVerts.AddUnique(TriB.A); TriVerts.AddUnique(TriB.B); TriVerts.AddUnique(TriB.C);
+			if (TriVerts.Num() != 4) 
+			{
+				bAllOK = false;
+				break;
+			}
+			if (TriVerts.Contains(VertexA) == false || TriVerts.Contains(VertexB) == false || TriVerts.Contains(VertexC) == false || TriVerts.Contains(VertexD) == false)
+			{
+				bAllOK = false;
+				break;
+			}
+
+			// we want TriA to be the one containing the edge (Vertex0,Vertex1), so if it's TriB, swap them, and the quad
+			if (TriB.Contains(VertexA) && TriB.Contains(VertexB))
+			{
+				Swap(TriA, TriB);
+				Swap(QuadTris.A, QuadTris.B);
+			}
+		}
+	}
+
+	if (!bAllOK)
+	{
+		ensure(false);
+		VertexSpans.Reset();
+		QuadTriangles.Reset();
+		NumVertexColsU = NumVertexRowsV = 0;
+		return false;
+	}
+
+	return true;
+}
+
+
+
 void FQuadGridPatch::ReverseRows()
 {
 	Algo::Reverse(QuadTriangles);
@@ -175,6 +261,36 @@ FIndex3i FQuadGridPatch::GetQuadTriMappedToSpanIndices(const FDynamicMesh3& Mesh
 
 
 
+bool FQuadGridPatch::GetVertexColumn(int32 ColumnIndex, TArray<int32>& VerticesOut) const
+{
+	if (ColumnIndex < 0 || ColumnIndex >= VertexSpans[0].Num())
+	{
+		ensure(false);
+		return false;
+	}
+
+	VerticesOut.Reset();
+	for (int32 k = 0; k < VertexSpans.Num(); ++k)
+	{
+		VerticesOut.Add(VertexSpans[k][ColumnIndex]);
+	}
+	return true;
+}
+
+int32 FQuadGridPatch::FindColumnIndex(int32 VertexID) const
+{
+	for (const TArray<int32>& Span : VertexSpans)
+	{
+		int32 Index = Span.IndexOfByKey(VertexID);
+		if (Index != INDEX_NONE)
+		{
+			return Index;
+		}
+	}
+	return IndexConstants::InvalidID;
+}
+
+
 void FQuadGridPatch::GetSubPatchByQuadRange(int StartRow, int EndRow, int StartCol, int EndCol, FQuadGridPatch& PatchOut ) const
 {
 	int NumQuadRows = (EndRow-StartRow)+1;
diff --git a/Engine/Plugins/Runtime/GeometryProcessing/Source/DynamicMesh/Public/Parameterization/DynamicMeshUVEditor.h b/Engine/Plugins/Runtime/GeometryProcessing/Source/DynamicMesh/Public/Parameterization/DynamicMeshUVEditor.h
--- a/Engine/Plugins/Runtime/GeometryProcessing/Source/DynamicMesh/Public/Parameterization/DynamicMeshUVEditor.h
+++ b/Engine/Plugins/Runtime/GeometryProcessing/Source/DynamicMesh/Public/Parameterization/DynamicMeshUVEditor.h
@@ -87,6 +87,14 @@ public:
 	bool CopyUVLayer(FDynamicMeshUVOverlay* FromUVOverlay);
 
 
+	/**
+	 * Initialize UV topology to be one UV per vertex. UVs are initialized to (0,0)
+	 * @param VertexToUVOut returned new overlay UV ID for each mesh Vertex ID (array size is mesh MaxVertexID)
+	 * @param bIsIdentityMapOut returned as true if UV ID == Vertex ID for every vertex  (even if there are gaps)
+	 */
+	void SetToPerVertexUVs(TArray<int32>& VertexToUVOut, bool& bIsIdentityMapOut, FUVEditResult* Result = nullptr);
+
+
 	/**
 	 * Create new UV island for each Triangle, by planar projection onto plane of Triangle. No transforms/etc are applied.
 	 */
diff --git a/Engine/Plugins/Runtime/GeometryProcessing/Source/DynamicMesh/Public/Selections/QuadGridPatch.h b/Engine/Plugins/Runtime/GeometryProcessing/Source/DynamicMesh/Public/Selections/QuadGridPatch.h
--- a/Engine/Plugins/Runtime/GeometryProcessing/Source/DynamicMesh/Public/Selections/QuadGridPatch.h
+++ b/Engine/Plugins/Runtime/GeometryProcessing/Source/DynamicMesh/Public/Selections/QuadGridPatch.h
@@ -27,8 +27,8 @@ class FDynamicMesh3;
 class DYNAMICMESH_API FQuadGridPatch
 {
 public:
-	int NumVertexColsU;
-	int NumVertexRowsV;
+	int NumVertexColsU = 0;
+	int NumVertexRowsV = 0;
 
 	/** NumVertexRowsV rows of NumVertexColsU VertexIDs, may contain repeated element if the patch forms a loop */
 	TArray<TArray<int32>> VertexSpans;
@@ -42,12 +42,20 @@ public:
 	int NumQuadCols() const { return NumVertexColsU-1; }
 	int NumQuadRows() const { return NumVertexRowsV-1; }
 
+	bool IsEmpty() const { return NumVertexColsU == 0 || NumVertexRowsV == 0; }
 
 	/**
-	 * Initialize an Nx1 grid from a list of quads and the row of vertices on the "bottom" (ie first span)
+	 * Initialize an Nx1 grid from a list of quads (represented as a pair of triangles) and the row of vertices on the "bottom" (ie first span)
 	 */
 	void InitializeFromQuadStrip(const FDynamicMesh3& Mesh, const TArray<FIndex2i>& SequentialQuadsIn, const TArray<int32>& FirstVertexSpan );
 
+	/**
+	 * Initialize an NxM grid from a list of quads (represented as a pair of triangles) and corresponding rows of vertices. 
+	 * For N quad-rows, there must be N+1 vertex spans. 
+	 * @return false if the input data is detected to be invalid, in this case the patch is also cleared
+	 */
+	bool InitializeFromQuadPatch(const FDynamicMesh3& Mesh, const TArray<TArray<FIndex2i>>& QuadRowsIn, const TArray<TArray<int32>>& VertexSpansIn);
+
 	/**
 	 * Reverse the order of the vertex/quad rows, ie "flip" vertically
 	 */
@@ -115,6 +123,15 @@ public:
 		}
 	}
 
+	/**
+	 * Return vertices in specified Column
+	 */
+	bool GetVertexColumn(int32 ColumnIndex, TArray<int32>& VerticesOut) const;
+
+	/**
+	 * Find column that contains vertex, or InvalidID
+	 */
+	int32 FindColumnIndex(int32 VertexID) const;
 
 	/**
 	 * Construct a new QuadGridPatch that is a subset of the current QuadGridPatch
diff --git a/Engine/Source/Runtime/GeometryCore/Public/VectorUtil.h b/Engine/Source/Runtime/GeometryCore/Public/VectorUtil.h
--- a/Engine/Source/Runtime/GeometryCore/Public/VectorUtil.h
+++ b/Engine/Source/Runtime/GeometryCore/Public/VectorUtil.h
@@ -316,6 +316,19 @@ namespace VectorUtil
 		return TMathUtil<RealType>::Sqrt(sqr);
 	}
 
+	/**
+	 * tan(theta/2) = +/- sqrt( (1-cos(theta)) / (1+cos(theta)) )
+	 * @return positive value of tan(theta/2) where theta is angle between normalized vectors A and B
+	 */
+	template <typename RealType>
+	RealType VectorTanHalfAngle(const TVector2<RealType>& A, const TVector2<RealType>& B)
+	{
+		RealType cosAngle = A.Dot(B);
+		RealType sqr = ((RealType)1 - cosAngle) / ((RealType)1 + cosAngle);
+		sqr = Clamp(sqr, (RealType)0, TMathUtil<RealType>::MaxReal);
+		return TMathUtil<RealType>::Sqrt(sqr);
+	}
+
 	/**
 	 * Fast cotangent of angle between two vectors (*do not have to be normalized unit vectors*).
 	 * cot = cos/sin, both of which can be computed from vector identities
