add collision shape merging to the 'mesh to collision' tool

diff --git a/Engine/Plugins/Experimental/MeshModelingToolsetExp/Source/MeshModelingToolsExp/Private/Physics/SetCollisionGeometryTool.cpp b/Engine/Plugins/Experimental/MeshModelingToolsetExp/Source/MeshModelingToolsExp/Private/Physics/SetCollisionGeometryTool.cpp
--- a/Engine/Plugins/Experimental/MeshModelingToolsetExp/Source/MeshModelingToolsExp/Private/Physics/SetCollisionGeometryTool.cpp
+++ b/Engine/Plugins/Experimental/MeshModelingToolsetExp/Source/MeshModelingToolsExp/Private/Physics/SetCollisionGeometryTool.cpp
@@ -20,6 +20,7 @@
 #include "ShapeApproximation/MeshSimpleShapeApproximation.h"
 
 #include "Physics/CollisionGeometryVisualization.h"
+#include "Physics/ComponentCollisionUtil.h"
 
 // physics data
 #include "Components/StaticMeshComponent.h"
@@ -64,10 +65,14 @@ public:
 	bool bUseMaxCount;
 	bool bRemoveContained;
 	bool bAppendToExisting;
+	bool bMergeCollisionShapes;
+	bool bUseNegativeSpaceInMerge;
 
 	EProjectedHullAxis SweepAxis;
 	int32 MaxCount;
 
+	int32 MergeAboveCount;
+
 	// Begin TGenericDataOperator interface
 	virtual void CalculateResult(FProgressCancel* Progress) override
 	{
@@ -159,6 +164,16 @@ public:
 			NewCollision->Geometry.RemoveContainedGeometry();
 		}
 
+		if (bMergeCollisionShapes)
+		{
+			FSimpleShapeSet3d::FMergeShapesSettings Settings;
+			Settings.bMergeShapesProtectNegativeSpace = bUseNegativeSpaceInMerge;
+			// Take the UseShapeGenerator negative space settings (same as we use for convex decomposition negative space)
+			Settings.bIgnoreInternalNegativeSpace = UseShapeGenerator->bIgnoreInternalNegativeSpace;
+			Settings.NegativeSpaceMinRadius = UseShapeGenerator->NegativeSpaceMinRadius;
+			Settings.NegativeSpaceTolerance = UseShapeGenerator->NegativeSpaceTolerance;
+			NewCollision->Geometry.MergeShapes(MergeAboveCount, Settings);
+		}
 		
 		if (bUseMaxCount)
 		{
@@ -302,19 +317,22 @@ void USetCollisionGeometryTool::Setup()
 	Settings->WatchProperty(Settings->bEnableMaxCount, [this](bool) { InvalidateCompute(); });
 	Settings->WatchProperty(Settings->MaxCount, [this](int32) { InvalidateCompute(); });
 	Settings->WatchProperty(Settings->MinThickness, [this](float) { InvalidateCompute(); });
-	Settings->WatchProperty(Settings->bDetectBoxes, [this](int32) { InvalidateCompute(); });
-	Settings->WatchProperty(Settings->bDetectSpheres, [this](int32) { InvalidateCompute(); });
-	Settings->WatchProperty(Settings->bDetectCapsules, [this](int32) { InvalidateCompute(); });
+	Settings->WatchProperty(Settings->bDetectBoxes, [this](bool) { InvalidateCompute(); });
+	Settings->WatchProperty(Settings->bDetectSpheres, [this](bool) { InvalidateCompute(); });
+	Settings->WatchProperty(Settings->bDetectCapsules, [this](bool) { InvalidateCompute(); });
 	Settings->WatchProperty(Settings->bSimplifyHulls, [this](bool) { InvalidateCompute(); });
 	Settings->WatchProperty(Settings->HullTargetFaceCount, [this](int32) { InvalidateCompute(); });
 	Settings->WatchProperty(Settings->MaxHullsPerMesh, [this](int32) { InvalidateCompute(); });
 	Settings->WatchProperty(Settings->ConvexDecompositionSearchFactor, [this](int32) { InvalidateCompute(); });
 	Settings->WatchProperty(Settings->AddHullsErrorTolerance, [this](int32) { InvalidateCompute(); });
 	Settings->WatchProperty(Settings->MinPartThickness, [this](int32) { InvalidateCompute(); });
-	Settings->WatchProperty(Settings->bUseNegativeSpaceInDecomposition, [this](int32) { InvalidateCompute(); });
+	Settings->WatchProperty(Settings->bUseNegativeSpaceInDecomposition, [this](bool) { InvalidateCompute(); });
 	Settings->WatchProperty(Settings->NegativeSpaceMinRadius, [this](int32) { InvalidateCompute(); });
 	Settings->WatchProperty(Settings->NegativeSpaceTolerance, [this](int32) { InvalidateCompute(); });
-	Settings->WatchProperty(Settings->bIgnoreInternalNegativeSpace, [this](int32) { InvalidateCompute(); });
+	Settings->WatchProperty(Settings->bIgnoreInternalNegativeSpace, [this](bool) { InvalidateCompute(); });
+	Settings->WatchProperty(Settings->bMergeCollisionShapes, [this](bool) { InvalidateCompute(); });
+	Settings->WatchProperty(Settings->bUseNegativeSpaceInMerge, [this](bool) { InvalidateCompute(); });
+	Settings->WatchProperty(Settings->MergeAboveCount, [this](int32) { InvalidateCompute(); });
 	Settings->WatchProperty(Settings->HullTolerance, [this](float) { InvalidateCompute(); });
 	Settings->WatchProperty(Settings->SweepAxis, [this](EProjectedHullAxis) { InvalidateCompute(); });
 	Settings->WatchProperty(Settings->LevelSetResolution, [this](int32) { InvalidateCompute(); });
@@ -437,6 +455,9 @@ TUniquePtr<UE::Geometry::TGenericDataOperator<FPhysicsDataCollection>> USetColli
 	Op->UseShapeGenerator->bIgnoreInternalNegativeSpace = Settings->bIgnoreInternalNegativeSpace;
 	Op->UseShapeGenerator->HullSimplifyTolerance = Settings->HullTolerance;
 	Op->UseShapeGenerator->LevelSetGridResolution = Settings->LevelSetResolution;
+	Op->bMergeCollisionShapes = Settings->bMergeCollisionShapes;
+	Op->bUseNegativeSpaceInMerge = Settings->bUseNegativeSpaceInMerge;
+	Op->MergeAboveCount = Settings->MergeAboveCount;
 
 	Op->ComputeType = Settings->GeometryType;
 	Op->bAppendToExisting = Settings->bAppendToExisting;
diff --git a/Engine/Plugins/Experimental/MeshModelingToolsetExp/Source/MeshModelingToolsExp/Public/Physics/SetCollisionGeometryTool.h b/Engine/Plugins/Experimental/MeshModelingToolsetExp/Source/MeshModelingToolsExp/Public/Physics/SetCollisionGeometryTool.h
--- a/Engine/Plugins/Experimental/MeshModelingToolsetExp/Source/MeshModelingToolsExp/Public/Physics/SetCollisionGeometryTool.h
+++ b/Engine/Plugins/Experimental/MeshModelingToolsetExp/Source/MeshModelingToolsExp/Public/Physics/SetCollisionGeometryTool.h
@@ -158,6 +158,18 @@ public:
 	UPROPERTY(EditAnywhere, Category = AutoDetectGeometryOverrides)
 	bool bDetectCapsules = true;
 
+	/** Whether to attempt to merge the generated collision shapes, when there are more than MergeAboveCount */
+	UPROPERTY(EditAnywhere, Category = MergeCollisionShapes)
+	bool bMergeCollisionShapes = false;
+
+	/** Attempt to merge generated collision shapes until there are at most this many */
+	UPROPERTY(EditAnywhere, Category = MergeCollisionShapes, meta = (EditConditionHides, EditCondition = "bMergeCollisionShapes", ClampMin = "1"))
+	int32 MergeAboveCount = 1;
+
+	/** Whether to protect negative space while merging the generated collision shapes, using the negative space settings */
+	UPROPERTY(EditAnywhere, Category = MergeCollisionShapes, meta = (EditConditionHides, EditCondition = "bMergeCollisionShapes"))
+	bool bUseNegativeSpaceInMerge = false;
+
 	/** Whether to simplify the convex hull */
 	UPROPERTY(EditAnywhere, Category = ConvexHulls, meta = (EditConditionHides, EditCondition = "GeometryType == ECollisionGeometryType::ConvexHulls"))
 	bool bSimplifyHulls = true;
@@ -194,15 +206,15 @@ public:
 
 	/** Negative space closer to the input than this tolerance distance can be filled in */
 	UPROPERTY(EditAnywhere, Category = NegativeSpace, meta = (UIMin = ".001", UIMax = "100", ClampMin = "0",
-		EditConditionHides, EditCondition = "GeometryType == ECollisionGeometryType::ConvexHulls && MaxHullsPerMesh > 1 && bUseNegativeSpaceInDecomposition"))
+		EditConditionHides, EditCondition = "(GeometryType == ECollisionGeometryType::ConvexHulls && MaxHullsPerMesh > 1 && bUseNegativeSpaceInDecomposition) || (bMergeCollisionShapes && bUseNegativeSpaceInMerge)"))
 	double NegativeSpaceTolerance = 3;
 	/** Minimum radius of negative space to protect; tunnels with radius smaller than this could be filled in */
 	UPROPERTY(EditAnywhere, Category = NegativeSpace, meta = (UIMin = ".001", UIMax = "100", ClampMin = "0",
-		EditConditionHides, EditCondition = "GeometryType == ECollisionGeometryType::ConvexHulls && MaxHullsPerMesh > 1 && bUseNegativeSpaceInDecomposition"))
+		EditConditionHides, EditCondition = "(GeometryType == ECollisionGeometryType::ConvexHulls && MaxHullsPerMesh > 1 && bUseNegativeSpaceInDecomposition) || (bMergeCollisionShapes && bUseNegativeSpaceInMerge)"))
 	double NegativeSpaceMinRadius = 10;
 	/** Whether to ignore negative space that is not accessible by traversing from the convex hull (via paths w/ radius of at least Negative Space Tolerance) */
 	UPROPERTY(EditAnywhere, Category = NegativeSpace, meta = (
-		EditConditionHides, EditCondition = "GeometryType == ECollisionGeometryType::ConvexHulls && MaxHullsPerMesh > 1 && bUseNegativeSpaceInDecomposition"))
+		EditConditionHides, EditCondition = "(GeometryType == ECollisionGeometryType::ConvexHulls && MaxHullsPerMesh > 1 && bUseNegativeSpaceInDecomposition) || (bMergeCollisionShapes && bUseNegativeSpaceInMerge)"))
 	bool bIgnoreInternalNegativeSpace = true;
 
 	/** If > 0, the polygon used to generate the swept hull will be simplified up to this distance tolerance, in cm */
diff --git a/Engine/Plugins/Runtime/GeometryProcessing/Source/DynamicMesh/Private/ShapeApproximation/SimpleShapeSet3.cpp b/Engine/Plugins/Runtime/GeometryProcessing/Source/DynamicMesh/Private/ShapeApproximation/SimpleShapeSet3.cpp
--- a/Engine/Plugins/Runtime/GeometryProcessing/Source/DynamicMesh/Private/ShapeApproximation/SimpleShapeSet3.cpp
+++ b/Engine/Plugins/Runtime/GeometryProcessing/Source/DynamicMesh/Private/ShapeApproximation/SimpleShapeSet3.cpp
@@ -6,6 +6,15 @@
 #include "DynamicMesh/MeshTransforms.h"
 #include "Intersection/ContainmentQueries3.h"
 
+#include "CompGeom/ConvexDecomposition3.h"
+#include "DynamicMeshEditor.h"
+#include "Generators/MeshShapeGenerator.h"
+#include "Generators/GridBoxMeshGenerator.h"
+#include "Generators/BoxSphereGenerator.h"
+#include "Generators/CapsuleGenerator.h"
+#include "DynamicMesh/DynamicMeshAABBTree3.h"
+#include "Spatial/FastWinding.h"
+
 using namespace UE::Geometry;
 
 
@@ -706,3 +715,259 @@ void FSimpleShapeSet3d::ApplyTransform(const FTransform3d& Transform)
 		TransformLevelSetShape(LevelSetShape, Transform);
 	}
 }
+
+bool FSimpleShapeSet3d::MergeShapes(int32 MergeAboveCount, const FMergeShapesSettings& MergeSettings)
+{
+	int32 NumShapes = TotalElementsNum();
+	//Note: We currently do not support merging level sets, because the merge logic assumes we can treat each simple shape as a convex hull
+	int32 NonLevelSet = NumShapes - LevelSets.Num();
+	if (NumShapes <= MergeAboveCount || NonLevelSet < 2)
+	{
+		// No merges needed / possible
+		return false;
+	}
+
+	struct FSourceShapeRef
+	{
+		int32 SourceIdx = -1;
+		ESimpleShapeType ShapeType;
+	};
+
+	TArray<FVector> HullVertices;
+	TArray<int32> HullVertexCounts;
+	TArray<double> HullVolumes;
+	TArray<FSourceShapeRef> HullToShapeSource;
+	TUniquePtr<FDynamicMesh3> CollisionMesh;
+	if (MergeSettings.bMergeShapesProtectNegativeSpace)
+	{
+		CollisionMesh = MakeUnique<FDynamicMesh3>();
+	}
+
+	auto AppendGeneratorToCollisionMesh = [&CollisionMesh](const FMeshShapeGenerator& Generator)
+	{
+		if (!CollisionMesh)
+		{
+			return;
+		}
+		checkSlow(CollisionMesh->IsCompact());
+		int32 NewVertStart = CollisionMesh->MaxVertexID();
+		for (FVector3d V : Generator.Vertices)
+		{
+			CollisionMesh->AppendVertex(V);
+		}
+		for (FIndex3i T : Generator.Triangles)
+		{
+			CollisionMesh->AppendTriangle(T.A + NewVertStart, T.B + NewVertStart, T.C + NewVertStart);
+		}
+	};
+	auto TransformVertices = [](TArrayView<FVector3d> Vertices, const FTransform& Transform)
+	{
+		for (FVector3d& Vertex : Vertices)
+		{
+			Vertex = Transform.TransformPosition(Vertex);
+		}
+	};
+	auto AppendHullVertices = [&HullToShapeSource, &HullVolumes, &HullVertices, &HullVertexCounts]
+	(TArrayView<const FVector3d> Vertices, double Volume, const FSourceShapeRef ShapeSource, const FTransform* ShapeTransform = nullptr)
+	{
+		check(HullToShapeSource.Num() == HullVolumes.Num());
+		HullToShapeSource.Add(ShapeSource);
+		int32 HullIdxStart = HullVertices.Num();
+		HullVertices.Append(Vertices);
+		if (ShapeTransform)
+		{
+			for (int32 Idx = HullIdxStart; Idx < HullVertices.Num(); ++Idx)
+			{
+				ShapeTransform->TransformPosition(HullVertices[Idx]);
+			}
+		}
+		HullVertexCounts.Add(Vertices.Num());
+		HullVolumes.Add(Volume);
+	};
+	auto GeneratorVolume = [](FMeshShapeGenerator* Generator) -> double
+	{
+		TIndexVectorMeshArrayAdapter<FIndex3i, double, FVector3d> GenMeshAdapter(&Generator->Vertices, &Generator->Triangles);
+		FVector2d VolArea = TMeshQueries<TIndexVectorMeshArrayAdapter<FIndex3i, double, FVector3d>>::GetVolumeArea(GenMeshAdapter);
+		return VolArea.X;
+	};
+
+
+	// Simple shape approximation options, for conversion to convex hull for merges
+	// TODO: consider exposing as parameters
+	constexpr int32 SphereStepsPerSide = 6;
+	constexpr int32 CapsuleCircleSteps = 16;
+	constexpr int32 CapsuleHemisphereSteps = 8;
+
+	for (int32 BoxIdx = 0; BoxIdx < Boxes.Num(); ++BoxIdx)
+	{
+		FSourceShapeRef Source{ BoxIdx, ESimpleShapeType::Box };
+		FOrientedBox3d OrientedBox = Boxes[BoxIdx].Box;
+		if (CollisionMesh)
+		{
+			FGridBoxMeshGenerator BoxGenerator;
+			BoxGenerator.EdgeVertices = FIndex3i(1, 1, 1);
+			BoxGenerator.Box = OrientedBox;
+			BoxGenerator.Generate();
+			AppendHullVertices(BoxGenerator.Vertices, OrientedBox.Volume(), Source);
+			AppendGeneratorToCollisionMesh(BoxGenerator);
+		}
+		else
+		{
+			TArray<FVector3d, TFixedAllocator<8>> BoxVertices;
+			OrientedBox.EnumerateCorners([&](FVector3d Corner) { BoxVertices.Add(Corner); });
+			AppendHullVertices(BoxVertices, OrientedBox.Volume(), Source);
+		}
+	}
+	for (int32 SphereIdx = 0; SphereIdx < Spheres.Num(); ++SphereIdx)
+	{
+		FSourceShapeRef Source{ SphereIdx, ESimpleShapeType::Sphere };
+		const FSphere3d& Sphere = Spheres[SphereIdx].Sphere;
+		FBoxSphereGenerator SphereGenerator;
+		SphereGenerator.Box.Frame.Origin = Sphere.Center;
+		SphereGenerator.Radius = (float)FMath::Max(FMathf::ZeroTolerance, Sphere.Radius);
+		int32 StepsPerSide = FMath::Max(1, SphereStepsPerSide);
+		SphereGenerator.EdgeVertices = FIndex3i(StepsPerSide, StepsPerSide, StepsPerSide);
+		SphereGenerator.Generate();
+		double Volume = GeneratorVolume(&SphereGenerator);
+		AppendHullVertices(SphereGenerator.Vertices, Volume, Source);
+		AppendGeneratorToCollisionMesh(SphereGenerator);
+	}
+	for (int32 CapsuleIdx = 0; CapsuleIdx < Capsules.Num(); ++CapsuleIdx)
+	{
+		FSourceShapeRef Source{ CapsuleIdx, ESimpleShapeType::Capsule };
+		const FCapsule3d& Capsule = Capsules[CapsuleIdx].Capsule;
+		FCapsuleGenerator CapsuleGenerator;
+		CapsuleGenerator.Radius = Capsule.Radius;
+		CapsuleGenerator.SegmentLength = Capsule.Segment.Length();
+		CapsuleGenerator.NumHemisphereArcSteps = FMath::Max(2, CapsuleHemisphereSteps);
+		CapsuleGenerator.NumCircleSteps = FMath::Max(3, CapsuleCircleSteps);
+		CapsuleGenerator.Generate();
+		Geometry::FQuaterniond CapsuleRot(FVector3d::ZAxisVector, Capsule.Segment.Direction);
+		FTransform CapsuleTransform(FRotator(CapsuleRot), Capsule.Center() - Capsule.Segment.Direction * Capsule.Segment.Length() * .5);
+		TransformVertices(CapsuleGenerator.Vertices, CapsuleTransform);
+		double Volume = GeneratorVolume(&CapsuleGenerator);
+		AppendHullVertices(CapsuleGenerator.Vertices, Volume, Source);
+		AppendGeneratorToCollisionMesh(CapsuleGenerator);
+	}
+	for (int32 ConvexIdx = 0; ConvexIdx < Convexes.Num(); ++ConvexIdx)
+	{
+		FSourceShapeRef Source{ ConvexIdx, ESimpleShapeType::Convex };
+		const FDynamicMesh3& Mesh = Convexes[ConvexIdx].Mesh;
+		double Volume = TMeshQueries<FDynamicMesh3>::GetVolumeArea(Mesh).X;
+		for (FVector3d V : Mesh.VerticesItr())
+		{
+			HullVertices.Add(V);
+		}
+		HullVertexCounts.Add(Mesh.VertexCount());
+		HullVolumes.Add(Volume);
+		HullToShapeSource.Add(Source);
+		if (CollisionMesh)
+		{
+			FDynamicMeshEditor Editor(CollisionMesh.Get());
+			FMeshIndexMappings UnusedMappings;
+			Editor.AppendMesh(&Mesh, UnusedMappings);
+		}
+	}
+
+	const int32 InitialNumConvex = HullVertexCounts.Num();
+	// Nothing we are able to merge
+	if (InitialNumConvex <= 1)
+	{
+		return false;
+	}
+
+	TArray<int32> HullVertexStarts;
+	HullVertexStarts.SetNumUninitialized(InitialNumConvex);
+	HullVertexStarts[0] = 0; // Note InitialNumConvex is > 1 due to above test
+	for (int32 HullIdx = 1, LastEnd = HullVertexCounts[0]; HullIdx < InitialNumConvex; LastEnd += HullVertexCounts[HullIdx++])
+	{
+		HullVertexStarts[HullIdx] = LastEnd;
+	}
+
+	TArray<TPair<int32, int32>> HullProximity;
+
+	FConvexDecomposition3 Decomposition;
+	Decomposition.InitializeFromHulls(HullVertexStarts.Num(),
+		[&HullVolumes](int32 HullIdx) { return HullVolumes[HullIdx]; }, [&HullVertexCounts](int32 HullIdx) { return HullVertexCounts[HullIdx]; },
+		[&HullVertexStarts, &HullVertices](int32 HullIdx, int32 VertIdx) { return HullVertices[HullVertexStarts[HullIdx] + VertIdx]; }, HullProximity);
+	double MinProximityOverlapTolerance = 0;
+	FSphereCovering NegativeSpace;
+	// Build the negative space of the collision shapes, if requested
+	if (MergeSettings.bMergeShapesProtectNegativeSpace)
+	{
+		FNegativeSpaceSampleSettings SampleSettings;
+		SampleSettings.bOnlyConnectedToHull = MergeSettings.bIgnoreInternalNegativeSpace;
+		SampleSettings.MinRadius = MergeSettings.NegativeSpaceMinRadius;
+		SampleSettings.ReduceRadiusMargin = MergeSettings.NegativeSpaceTolerance;
+		SampleSettings.MinRadius = FMath::Max(1, (MergeSettings.NegativeSpaceMinRadius + MergeSettings.NegativeSpaceTolerance) * .5);
+		SampleSettings.SampleMethod = FNegativeSpaceSampleSettings::ESampleMethod::VoxelSearch;
+		SampleSettings.bRequireSearchSampleCoverage = true;
+		SampleSettings.TargetNumSamples = 1; // let the sample coverage determine the number of spheres to place
+
+		MinProximityOverlapTolerance = FMath::Max(SampleSettings.ReduceRadiusMargin * .5, MinProximityOverlapTolerance);
+		FDynamicMeshAABBTree3 CollisionAABBTree(CollisionMesh.Get(), true);
+		TFastWindingTree<FDynamicMesh3> CollisionFastWinding(&CollisionAABBTree, true);
+		NegativeSpace.AddNegativeSpace(CollisionFastWinding, SampleSettings, false);
+	}
+	FSphereCovering* UseNegativeSpace = NegativeSpace.Num() > 0 ? &NegativeSpace : nullptr;
+
+	// Find possible shape merges based on the shape bounding box overlaps,
+	// where bounds are expanded by: Max(a quarter their min dimension, a tenth their max dimension, the half reduce radius margin if negative space is computed)
+	Decomposition.InitializeProximityFromDecompositionBoundingBoxOverlaps(.25, .1, MinProximityOverlapTolerance);
+
+	int32 MaxShapeCount = FMath::Max(0, MergeAboveCount - LevelSets.Num());
+	Decomposition.RestrictMergeSearchToLocalAfterTestNumConnections = 1000 + MaxShapeCount * MaxShapeCount; // Restrict searches in very large search cases, when not close to max shape count, to avoid excessive search time
+	constexpr double ErrorTolerance = 0, MinThicknessTolerance = 0;
+	Decomposition.MergeBest(MaxShapeCount, ErrorTolerance, MinThicknessTolerance, true, false, MaxShapeCount, UseNegativeSpace, nullptr /*optional FTransform for negative space*/);
+
+	// Algorithm decided not to merge
+	if (Decomposition.NumHulls() == InitialNumConvex)
+	{
+		return false;
+	}
+
+
+	TArray<FSphereShape3d> SourceSpheres = MoveTemp(Spheres);
+	TArray<FBoxShape3d> SourceBoxes = MoveTemp(Boxes);
+	TArray<UE::Geometry::FCapsuleShape3d> SourceCapsules = MoveTemp(Capsules);
+	TArray<FConvexShape3d> SourceConvexes = MoveTemp(Convexes);
+
+	for (int32 HullIdx = 0; HullIdx < Decomposition.Decomposition.Num(); ++HullIdx)
+	{
+		const FConvexDecomposition3::FConvexPart& Part = Decomposition.Decomposition[HullIdx];
+		// If part was not merged, use the source ID to map it back to the original collision primitive
+		if (Part.HullSourceID >= 0)
+		{
+			const FSourceShapeRef& SourceRef = HullToShapeSource[Part.HullSourceID];
+			bool bHandledShape = true;
+			switch (SourceRef.ShapeType)
+			{
+			case ESimpleShapeType::Box:
+				Boxes.Add(SourceBoxes[SourceRef.SourceIdx]);
+				break;
+			case ESimpleShapeType::Sphere:
+				Spheres.Add(SourceSpheres[SourceRef.SourceIdx]);
+				break;
+			case ESimpleShapeType::Capsule:
+				Capsules.Add(SourceCapsules[SourceRef.SourceIdx]);
+				break;
+			case ESimpleShapeType::Convex:
+				Convexes.Add(MoveTemp(SourceConvexes[SourceRef.SourceIdx]));
+				break;
+			default:
+				// Note: All shapes that we add to the HullToShapeSource array should be handled above, so we should not reach here
+				ensureMsgf(false, TEXT("Unhandled shape element type could not be restored from source shapes"));
+				bHandledShape = false;
+			}
+			if (bHandledShape)
+			{
+				continue;
+			}
+		}
+		// Add the merged part
+		Convexes.Emplace(Decomposition.GetHullMesh(HullIdx));
+	}
+
+	return true;
+}
+
diff --git a/Engine/Plugins/Runtime/GeometryProcessing/Source/DynamicMesh/Public/ShapeApproximation/SimpleShapeSet3.h b/Engine/Plugins/Runtime/GeometryProcessing/Source/DynamicMesh/Public/ShapeApproximation/SimpleShapeSet3.h
--- a/Engine/Plugins/Runtime/GeometryProcessing/Source/DynamicMesh/Public/ShapeApproximation/SimpleShapeSet3.h
+++ b/Engine/Plugins/Runtime/GeometryProcessing/Source/DynamicMesh/Public/ShapeApproximation/SimpleShapeSet3.h
@@ -194,6 +194,20 @@ struct DYNAMICMESH_API FSimpleShapeSet3d
 	 */
 	void FilterByVolume(int32 MaximumCount);
 
+	struct FMergeShapesSettings
+	{
+		/** Whether to use the negative space protection settings when merging collision shapes */
+		bool bMergeShapesProtectNegativeSpace = false;
+
+		/** Negative space closer to the input than this tolerance distance can be filled in */
+		double NegativeSpaceTolerance = 3;
+		/** Minimum radius of negative space to protect; tunnels with radius smaller than this could be filled in */
+		double NegativeSpaceMinRadius = 10;
+		/** Whether to ignore negative space that is not accessible by traversing from the convex hull (via paths w/ radius of at least Negative Space Tolerance) */
+		bool bIgnoreInternalNegativeSpace = true;
+	};
+	// Attempt to reduce the number of simple collision shapes by merging neighboring shapes, while still optionally protecting negative space
+	bool MergeShapes(int32 MergeAboveCount, const FMergeShapesSettings& MergeSettings);
 
 	/**
 	 * Transform shape elements. This will be a best-effort as if there is non-uniform scaling only Convexes can be transformed correctly
