Add support for multi-segment bevels with optional round profile in MeshBevel. The new code path is completely separated from the previous code. Works OK but handling of valence-3 bevel corners (filled with a rounded triangle) could be improved.
Expose in Geometry Script and PolyEdit functions.

diff --git a/Engine/Plugins/Experimental/GeometryScripting/Source/GeometryScriptingCore/Private/MeshModelingFunctions.cpp b/Engine/Plugins/Experimental/GeometryScripting/Source/GeometryScriptingCore/Private/MeshModelingFunctions.cpp
--- a/Engine/Plugins/Experimental/GeometryScripting/Source/GeometryScriptingCore/Private/MeshModelingFunctions.cpp
+++ b/Engine/Plugins/Experimental/GeometryScripting/Source/GeometryScriptingCore/Private/MeshModelingFunctions.cpp
@@ -858,6 +858,8 @@ UDynamicMesh* UGeometryScriptLibrary_MeshModelingFunctions::ApplyMeshBevelSelect
 		Bevel.InsetDistance = Options.BevelDistance;
 		Bevel.MaterialIDMode = (Options.bInferMaterialID) ? FMeshBevel::EMaterialIDMode::InferMaterialID : FMeshBevel::EMaterialIDMode::ConstantMaterialID;
 		Bevel.SetConstantMaterialID = Options.SetMaterialID;
+		Bevel.NumSubdivisions = FMath::Clamp(Options.Subdivisions, 0, 9999);
+		Bevel.RoundWeight = FMath::Clamp(Options.RoundWeight, -10.0, 10.0);
 
 		TUniquePtr<FGroupTopology> GroupTopology;
 
@@ -997,6 +999,8 @@ UDynamicMesh* UGeometryScriptLibrary_MeshModelingFunctions::ApplyMeshPolygroupBe
 		Bevel.InsetDistance = Options.BevelDistance;
 		Bevel.MaterialIDMode = (Options.bInferMaterialID) ? FMeshBevel::EMaterialIDMode::InferMaterialID : FMeshBevel::EMaterialIDMode::ConstantMaterialID;
 		Bevel.SetConstantMaterialID = Options.SetMaterialID;
+		Bevel.NumSubdivisions = FMath::Clamp(Options.Subdivisions, 0, 9999);
+		Bevel.RoundWeight = FMath::Clamp(Options.RoundWeight, -10.0, 10.0);
 		if (BevelGroupEdges.Num() > 0)
 		{
 			Bevel.InitializeFromGroupTopologyEdges(EditMesh, Topology, BevelGroupEdges);
diff --git a/Engine/Plugins/Experimental/GeometryScripting/Source/GeometryScriptingCore/Public/GeometryScript/MeshModelingFunctions.h b/Engine/Plugins/Experimental/GeometryScripting/Source/GeometryScriptingCore/Public/GeometryScript/MeshModelingFunctions.h
--- a/Engine/Plugins/Experimental/GeometryScripting/Source/GeometryScriptingCore/Public/GeometryScript/MeshModelingFunctions.h
+++ b/Engine/Plugins/Experimental/GeometryScripting/Source/GeometryScriptingCore/Public/GeometryScript/MeshModelingFunctions.h
@@ -193,16 +193,25 @@ struct GEOMETRYSCRIPTINGCORE_API FGeometryScriptMeshBevelOptions
 {
 	GENERATED_BODY()
 public:
+	/** Distance that each beveled mesh edge is inset from it's initial position*/
 	UPROPERTY(EditAnywhere, BlueprintReadWrite, Category = Options)
 	float BevelDistance = 1.0;
 
+	/** If true, when faces on either side of a beveled mesh edges have the same Material ID, beveled edge will be set to that Material ID. Otherwise SetMaterialID is used. */
 	UPROPERTY(EditAnywhere, BlueprintReadWrite, Category = Options)
 	bool bInferMaterialID = false;
 
+	/** Material ID to set on the new faces introduced by bevel operation, unless bInferMaterialID=true and non-ambiguous MaterialID can be inferred from adjacent faces */
 	UPROPERTY(EditAnywhere, BlueprintReadWrite, Category = Options)
 	int SetMaterialID = 0;
 
+	/** Number of edge loops added along the bevel faces */
+	UPROPERTY(EditAnywhere, BlueprintReadWrite, Category = Options)
+	int Subdivisions = 0;
 
+	/** Roundness of the bevel. Ignored if Subdivisions = 0. */
+	UPROPERTY(EditAnywhere, BlueprintReadWrite, Category = Options)
+	float RoundWeight = 1.0;
 
 	/**
 	 * If true the set of beveled PolyGroup edges is limited to those that 
@@ -253,14 +262,25 @@ struct GEOMETRYSCRIPTINGCORE_API FGeometryScriptMeshBevelSelectionOptions
 {
 	GENERATED_BODY()
 public:
+	/** Distance that each beveled mesh edge is inset from it's initial position*/
 	UPROPERTY(EditAnywhere, BlueprintReadWrite, Category = Options)
 	float BevelDistance = 1.0;
 
+	/** If true, when faces on either side of a beveled mesh edges have the same Material ID, beveled edge will be set to that Material ID. Otherwise SetMaterialID is used. */
 	UPROPERTY(EditAnywhere, BlueprintReadWrite, Category = Options)
 	bool bInferMaterialID = false;
 
+	/** Material ID to set on the new faces introduced by bevel operation, unless bInferMaterialID=true and non-ambiguous MaterialID can be inferred from adjacent faces */
 	UPROPERTY(EditAnywhere, BlueprintReadWrite, Category = Options)
 	int SetMaterialID = 0;
+
+	/** Number of edge loops added along the bevel faces */
+	UPROPERTY(EditAnywhere, BlueprintReadWrite, Category = Options)
+	int Subdivisions = 0;
+
+	/** Roundness of the bevel. Ignored if Subdivisions = 0. */
+	UPROPERTY(EditAnywhere, BlueprintReadWrite, Category = Options)
+	float RoundWeight = 1.0;
 };
 
 
diff --git a/Engine/Plugins/Runtime/GeometryProcessing/Source/DynamicMesh/Private/Operations/MeshBevel.cpp b/Engine/Plugins/Runtime/GeometryProcessing/Source/DynamicMesh/Private/Operations/MeshBevel.cpp
--- a/Engine/Plugins/Runtime/GeometryProcessing/Source/DynamicMesh/Private/Operations/MeshBevel.cpp
+++ b/Engine/Plugins/Runtime/GeometryProcessing/Source/DynamicMesh/Private/Operations/MeshBevel.cpp
@@ -10,10 +10,21 @@
 #include "DynamicMesh/MeshNormals.h"
 #include "DynamicMesh/MeshIndexUtil.h"
 #include "MeshRegionBoundaryLoops.h"
+#include "MeshBoundaryLoops.h"
 #include "Operations/PolyEditingEdgeUtil.h"
 #include "Operations/PolyEditingUVUtil.h"
 #include "Algo/Count.h"
 #include "Distance/DistLine3Line3.h"
+#include "Operations/UniformTessellate.h"
+#include "DynamicSubmesh3.h"
+#include "Solvers/ConstrainedMeshSmoother.h"
+#include "Solvers/ConstrainedMeshDeformer.h"
+#include "Selections/MeshFaceSelection.h"
+#include "Parameterization/DynamicMeshUVEditor.h"
+#include "Polygon2.h"
+#include "VectorUtil.h"
+#include "Spatial/DenseGrid2.h"
+
 
 using namespace UE::Geometry;
 
@@ -197,7 +208,15 @@ bool FMeshBevel::Apply(FDynamicMesh3& Mesh, FDynamicMeshChangeTracker* ChangeTra
 	}
 
 	// mesh the bevel corners and edges
-	CreateBevelMeshing(Mesh);
+	if (NumSubdivisions <= 0)
+	{
+		CreateBevelMeshing(Mesh);
+	}
+	else
+	{
+		CreateBevelMeshing_Multi(Mesh);
+	}
+
 	if (ResultInfo.CheckAndSetCancelled(Progress))
 	{
 		return false;
@@ -247,13 +266,17 @@ bool FMeshBevel::Apply(FDynamicMesh3& Mesh, FDynamicMeshChangeTracker* ChangeTra
 
 
 
-FMeshBevel::FBevelVertex* FMeshBevel::GetBevelVertexFromVertexID(int32 VertexID)
+FMeshBevel::FBevelVertex* FMeshBevel::GetBevelVertexFromVertexID(int32 VertexID, int32* IndexOut)
 {
 	int32* FoundIndex = VertexIDToIndexMap.Find(VertexID);
 	if (FoundIndex == nullptr)
 	{
 		return nullptr;
 	}
+	if (IndexOut != nullptr)
+	{
+		*IndexOut = *FoundIndex;
+	}
 	return &Vertices[*FoundIndex];
 }
 
@@ -272,6 +295,7 @@ void FMeshBevel::AddBevelGroupEdge(const FDynamicMesh3& Mesh, const FGroupTopolo
 	FIndex2i EdgeCornerIDs = Topology.Edges[GroupEdgeID].EndpointCorners;
 
 	FBevelEdge Edge;
+	int32 NewBevelEdgeIndex = Edges.Num();
 
 	for (int32 ci = 0; ci < 2; ++ci)
 	{
@@ -281,19 +305,22 @@ void FMeshBevel::AddBevelGroupEdge(const FDynamicMesh3& Mesh, const FGroupTopolo
 		Edge.bEndpointBoundaryFlag[ci] = Mesh.IsBoundaryVertex(VertexID);
 		int32 IncomingEdgeID = (ci == 0) ? MeshEdgeList[0] : MeshEdgeList.Last();
 
-		FBevelVertex* VertInfo = GetBevelVertexFromVertexID(VertexID);
+		int32 BevelVertexIndex = -1;
+		FBevelVertex* VertInfo = GetBevelVertexFromVertexID(VertexID, &BevelVertexIndex);
 		if (VertInfo == nullptr)
 		{
 			FBevelVertex NewVertex;
 			NewVertex.CornerID = CornerID;
 			NewVertex.VertexID = VertexID;
-			int32 NewIndex = Vertices.Num();
+			BevelVertexIndex = Vertices.Num();
 			Vertices.Add(NewVertex);
-			VertexIDToIndexMap.Add(VertexID, NewIndex);
-			VertInfo = &Vertices[NewIndex];
+			VertexIDToIndexMap.Add(VertexID, BevelVertexIndex);
+			VertInfo = &Vertices[BevelVertexIndex];
 		}
 		VertInfo->IncomingBevelMeshEdges.Add(IncomingEdgeID);
 		VertInfo->IncomingBevelTopoEdges.Add(GroupEdgeID);
+		VertInfo->IncomingBevelEdgeIndices.Add(NewBevelEdgeIndex);
+		Edge.BevelVertices[ci] = BevelVertexIndex;
 	}
 
 	Edge.MeshEdges.Append(MeshEdgeList);
@@ -307,6 +334,13 @@ void FMeshBevel::AddBevelGroupEdge(const FDynamicMesh3& Mesh, const FGroupTopolo
 		Edge.MeshEdgeTris.Add(Mesh.GetEdgeT(eid));
 	}
 
+	Edge.InitialPositions.Reserve(Edge.MeshVertices.Num());
+	for (int32 vid : Edge.MeshVertices)
+	{
+		Edge.InitialPositions.Add(Mesh.GetVertex(vid));
+	}
+
+	Edge.EdgeIndex = Edges.Num();
 	Edges.Add(MoveTemp(Edge));
 }
 
@@ -332,6 +366,12 @@ void FMeshBevel::AddBevelEdgeLoop(const FDynamicMesh3& Mesh, const FEdgeLoop& Me
 		Loop.MeshEdgeTris.Add(Mesh.GetEdgeT(eid));
 	}
 
+	Loop.InitialPositions.Reserve(Loop.MeshVertices.Num());
+	for (int32 vid : Loop.MeshVertices)
+	{
+		Loop.InitialPositions.Add(Mesh.GetVertex(vid));
+	}
+
 	Loops.Add(Loop);
 }
 
@@ -1510,6 +1550,1594 @@ void FMeshBevel::CreateBevelMeshing(FDynamicMesh3& Mesh)
 
 
 
+
+
+
+
+
+
+
+void FMeshBevel::AppendEdgeQuads_Multi(FDynamicMesh3& Mesh, FBevelEdge& Edge)
+{
+	int32 NumEdges = Edge.MeshEdges.Num();
+	if (NumEdges != Edge.NewMeshEdges.Num())
+	{
+		return;
+	}
+
+	Edge.NewGroupID = Mesh.AllocateTriangleGroup();
+
+	struct FEdgePair
+	{
+		FIndex2i EdgeV0;
+		FIndex2i EdgeV1;
+	};
+	TArray<FEdgePair> SequentialQuadEdges;
+
+	bool bFoundInvalidCase = false;
+
+	// At this point each edge-span should be fully disconnected into a set of paired edges, 
+	// so we can trivially join each edge pair with a quad. (In the multi-case we will further
+	// subdivide this quad instead of adding it directly)
+	for (int32 k = 0; k < NumEdges && bFoundInvalidCase == false; ++k)
+	{
+		int32 EdgeID0 = Edge.MeshEdges[k];
+		int32 EdgeID1 = Edge.NewMeshEdges[k];
+
+		// in certain cases, like bevel topo-edges with a single mesh-edge, we would not
+		// have been able to construct the "other" mesh edge when processing the topo-edge
+		// (where .NewMeshEdges is computed), it would only have been created when processing the
+		// junction vertex. Currently we do not go back and update .NewMeshEdges in that case, but
+		// we do store the edge-pair-correspondence in the MeshEdgePairs map. 
+		if (EdgeID0 == EdgeID1)
+		{
+			int32* FoundEdgeID1 = MeshEdgePairs.Find(EdgeID0);
+			if (FoundEdgeID1 != nullptr)
+			{
+				EdgeID1 = *FoundEdgeID1;
+			}
+		}
+
+		if (EdgeID0 != EdgeID1 && Mesh.IsEdge(EdgeID1))
+		{
+			FIndex2i EdgeV0 = Mesh.GetOrientedBoundaryEdgeV(EdgeID0);
+			FIndex2i EdgeV1 = Mesh.GetOrientedBoundaryEdgeV(EdgeID1);
+			if (EdgeV0.Contains(EdgeV1.A) || EdgeV0.Contains(EdgeV1.B))
+			{
+				// If we hit this case, it means that Edge0 and Edge1 are still connected at one end, and
+				// so cannot be connected by a Quad, they can only be connected by a single triangle.
+				// It is unclear how we end up in this situation, it does occur somewhat regularly in complex
+				// geometry scripts though (eg see UE-157531 for a potential repro).
+				// In this case we cannot add subdivisions currently, and will fall back to simpler geometry
+				bFoundInvalidCase = true;
+			}
+			else
+			{
+				SequentialQuadEdges.Add({ EdgeV0, EdgeV1 });
+			}
+		}
+		else
+		{
+			bFoundInvalidCase = true;
+		}
+	}
+
+
+	// if something went wrong above, fall back to adding a single quad to avoid holes/etc
+	// (it's possible that this path may still result in catastrophic failure...)
+	if (bFoundInvalidCase || SequentialQuadEdges.Num() != NumEdges)
+	{
+		AppendEdgeQuads(Mesh, Edge);
+		return;
+	}
+
+	// Will need the normals for each exterior "side" of the bevel edge along the vertices path to
+	// define the arcs along the bevel edge. It ought to be possible to compute this after
+	// adding the bevel geometry, by filtering out the bevel tris, but for now it's simpler to just 
+	// do it here before we add the bevel tris and cache it...
+	Edge.NormalsA.SetNum(Edge.MeshVertices.Num());
+	Edge.NormalsB.SetNum(Edge.MeshVertices.Num());
+	for (int32 k = 0; k < Edge.MeshVertices.Num(); ++k)
+	{
+		Edge.NormalsA[k] = FMeshNormals::ComputeVertexNormal(Mesh, Edge.MeshVertices[k]);
+		Edge.NormalsB[k] = FMeshNormals::ComputeVertexNormal(Mesh, Edge.NewMeshVertices[k]);
+	}
+
+	// all the code below is going to generate the bevel edge geometry and populate this list of 
+	// vertex-rows, which will be used to assemble the final FQuadGridPatch for the bevel edge-strip
+	int32 N = NumSubdivisions;
+	TArray<TArray<int32>> VertexSpans;
+	VertexSpans.SetNum(N + 2);
+
+	// this function appends a new column to the rows in VertexSpans, by generating new
+	// vertices on the interior of the edge between StartVID and EndVID
+	auto AppendNewVertColumn = [&Mesh, &VertexSpans, N](int32 StartVID, int32 EndVID)
+	{
+		FVector3d StartPos = Mesh.GetVertex(StartVID);
+		FVector3d EndPos = Mesh.GetVertex(EndVID);
+		VertexSpans[0].Add(StartVID);
+		for (int32 j = 0; j < N; ++j)
+		{
+			double T = (double)(j+1) / (double)(N + 1);
+			int32 NewVertID = Mesh.AppendVertex(Lerp(StartPos, EndPos, T));
+			VertexSpans[j+1].Add(NewVertID);
+		}
+		VertexSpans[N+1].Add(EndVID);
+	};
+
+	// this function appends a new column to the rows in VertexSpans by copying
+	// the column from an existing adjacent FQuadGridPatch. This happens if we are 
+	// meshing a bevel-edge that is connected to another bevel-edge that was already
+	// generated, and the connecting bevel-vertex only has 2 bevel-edges (ie no polygon will be inserted)
+	auto AppendExistingVertColumn = [&VertexSpans, N](const FQuadGridPatch* AdjacentQuadPatch, int32 CornerVertexID) -> bool
+	{
+		int32 ColumnIdx = AdjacentQuadPatch->FindColumnIndex(CornerVertexID);
+		if (ColumnIdx >= 0)
+		{
+			TArray<int32> ColumnVerts;
+			AdjacentQuadPatch->GetVertexColumn(ColumnIdx, ColumnVerts);
+
+			// Our column should start w/ CornerVertexID, but it may be reversed due to different orientations when building the adjacent patch...
+			// (could that be fixed further upstream? possibly it should be!)
+			if (ColumnVerts.Last() == CornerVertexID)
+			{
+				Algo::Reverse(ColumnVerts);
+			}
+
+			for (int32 j = 0; j <= (N + 1); ++j)
+			{
+				VertexSpans[j].Add(ColumnVerts[j]);
+			}
+
+			return true;
+		}
+		return false;
+	};
+
+	// start and end vertex columns may already exist in some other FBevelEdge that has already been generated. 
+	// In that case, instead of appending new vertices, we want to look up the existing vertices and stitch to them. 
+	// This case only happens for BevelVertices connected to exactly 2 BevelEdges (<2 => Terminator, >2 => Polygon at vertex). 
+	auto GetConnectedQuadStripRef = [this](int BevelVertexIdx, int CurBevelEdgeIdx)
+	{
+		if (BevelVertexIdx == -1) return (const FQuadGridPatch*)nullptr;
+		const FBevelVertex& Vtx = Vertices[BevelVertexIdx];
+		if (Vtx.VertexType != EBevelVertexType::JunctionVertex || Vtx.Wedges.Num() != 2) return (const FQuadGridPatch*)nullptr;
+
+		for (int EdgeIdx : Vtx.IncomingBevelEdgeIndices)
+		{
+			if (EdgeIdx != CurBevelEdgeIdx)
+			{
+				const FBevelEdge& OtherEdge = Edges[EdgeIdx];
+				if (OtherEdge.StripQuadPatch.IsEmpty() == false)
+				{
+					return (const FQuadGridPatch*)&OtherEdge.StripQuadPatch;
+				}
+				return (const FQuadGridPatch*)nullptr;
+			}
+		}
+		return (const FQuadGridPatch*)nullptr;
+	};
+
+
+	// somehow SequentialQuadEdges ordering may not be in agreement with the [BevelVertices.A, BevelVertices.B] 
+	// ordering of the BevelEdge....unclear how or why this would be the case. But if so, then the Prev/Next QuadGridPatch
+	// search below would return the Prev & Next flipped from where we need them, causing the later AppendExistingVertColumn
+	// to fail. The simplest fix for this here is to swap PrevQuadPatch/NextQuadPatch...
+	int OriginalPrevVtxID = (Edge.BevelVertices.A >= 0) ? Vertices[Edge.BevelVertices.A].VertexID : -1;
+	int OriginalNextVtxID = (Edge.BevelVertices.B >= 0) ? Vertices[Edge.BevelVertices.B].VertexID : -1;
+	FIndex2i EdgeStartQuadVerts(SequentialQuadEdges[0].EdgeV0.B, SequentialQuadEdges[0].EdgeV1.A);
+	FIndex2i EdgeEndQuadVerts(SequentialQuadEdges[NumEdges-1].EdgeV0.A, SequentialQuadEdges[NumEdges-1].EdgeV1.B);
+	FIndex2i BevelEdgeVerts(Edge.BevelVertices.A, Edge.BevelVertices.B);
+	if (EdgeStartQuadVerts.Contains(OriginalNextVtxID) || EdgeEndQuadVerts.Contains(OriginalPrevVtxID))
+	{
+		Swap(BevelEdgeVerts.A, BevelEdgeVerts.B);
+	}
+
+	// find quadgrid patches connected to start and end of current bevel edge, if the yexist
+	const FQuadGridPatch* PrevQuadPatch = GetConnectedQuadStripRef(BevelEdgeVerts.A, Edge.EdgeIndex);
+	const FQuadGridPatch* NextQuadPatch = GetConnectedQuadStripRef(BevelEdgeVerts.B, Edge.EdgeIndex);
+
+	// now add columns of vertices for each inital vertex along the bevel-edge (by iterating along mesh-edges) 
+	for (int32 k = 0; k < NumEdges; ++k)
+	{
+		FIndex2i EdgeV0 = SequentialQuadEdges[k].EdgeV0;
+		FIndex2i EdgeV1 = SequentialQuadEdges[k].EdgeV1;
+
+		int32 QuadA = EdgeV0.B;
+		int32 QuadB = EdgeV0.A;
+		int32 QuadC = EdgeV1.B;
+		int32 QuadD = EdgeV1.A;
+
+		// for first and last edges, we have to add an extra row of vertices
+		if (k == 0)
+		{
+			if (PrevQuadPatch == nullptr || AppendExistingVertColumn(PrevQuadPatch, QuadA) == false)
+			{
+				AppendNewVertColumn(QuadA, QuadD);
+			}
+		}
+
+		if (k != NumEdges-1 || NextQuadPatch == nullptr || AppendExistingVertColumn(NextQuadPatch, QuadB) == false)
+		{
+			AppendNewVertColumn(QuadB, QuadC);
+		}
+	
+	}
+
+	// now append the quads that connect up the vertices
+	TArray<TArray<FIndex2i>> QuadSpans;
+	QuadSpans.SetNum(N + 1);
+
+	int32 NumStrips = VertexSpans.Num() - 1;
+	for (int32 k = 0; k < NumEdges; ++k)
+	{
+		for ( int32 j = 0; j < NumStrips; ++j)
+		{
+			FIndex2i QuadTris(IndexConstants::InvalidID, IndexConstants::InvalidID);
+			QuadTris.A = Mesh.AppendTriangle(VertexSpans[j][k], VertexSpans[j][k+1], VertexSpans[j+1][k], Edge.NewGroupID);
+			QuadTris.B = Mesh.AppendTriangle(VertexSpans[j+1][k+1], VertexSpans[j+1][k], VertexSpans[j][k+1], Edge.NewGroupID);
+			QuadSpans[j].Add(QuadTris);
+
+			Edge.StripQuads.Add(QuadTris);
+		}
+	}
+
+	// finally save the quad-patch we generated in the FBevelEdge
+	Edge.StripQuadPatch.InitializeFromQuadPatch(Mesh, QuadSpans, VertexSpans);
+}
+
+
+
+void FMeshBevel::AppendLoopQuads_Multi(FDynamicMesh3& Mesh, FBevelLoop& Loop)
+{
+	// As in other places, the Loop case is a simplified version of the Edge case, where
+	// many special cases do not have to be handled. See the comments in 
+	// AppendEdgeQuads_Multi() for anything non-obvious below
+
+	int32 NumEdges = Loop.MeshEdges.Num();
+	if (NumEdges != Loop.NewMeshEdges.Num())
+	{
+		return;
+	}
+
+	auto GetGroupKey = [&Mesh, &Loop](int32 k)
+	{
+		FIndex2i EdgeTris = Loop.MeshEdgeTris[k];
+		int32 Group0 = Mesh.GetTriangleGroup(EdgeTris.A);
+		int32 Group1 = Mesh.IsTriangle(EdgeTris.B) ? Mesh.GetTriangleGroup(EdgeTris.B) : -1;
+		return FIndex2i(FMath::Max(Group0, Group1), FMath::Min(Group0, Group1));
+	};
+
+	TMap<FIndex2i, int> NewGroupIDs;
+	for (int32 k = 0; k < NumEdges; ++k)
+	{
+		FIndex2i GroupKey = GetGroupKey(k);
+		if (NewGroupIDs.Contains(GroupKey) == false)
+		{
+			NewGroupIDs.Add(GroupKey, Mesh.AllocateTriangleGroup());
+			Loop.NewGroupIDs.Add(NewGroupIDs[GroupKey]);
+		}
+	}
+
+	struct FEdgePair
+	{
+		FIndex2i EdgeV0;
+		FIndex2i EdgeV1;
+	};
+	TArray<FEdgePair> SequentialQuadEdges;
+
+	bool bFoundInvalidCase = false;
+
+	// At this point each edge-span should be fully disconnected into a set of paired edges, 
+	// so we can trivially join each edge pair with a quad.
+	for (int32 k = 0; k < NumEdges && bFoundInvalidCase == false; ++k)
+	{
+		int32 EdgeID0 = Loop.MeshEdges[k];
+		int32 EdgeID1 = Loop.NewMeshEdges[k];
+
+		// case that happens in AppendEdgeQuads() should never happen for loops...
+
+		FIndex2i QuadTris(IndexConstants::InvalidID, IndexConstants::InvalidID);
+		if (EdgeID0 != EdgeID1 && Mesh.IsEdge(EdgeID1))
+		{
+			FIndex2i GroupKey = GetGroupKey(k);
+			int32 NewGroupID = NewGroupIDs[GroupKey];
+
+			FIndex2i EdgeV0 = Mesh.GetOrientedBoundaryEdgeV(EdgeID0);
+			FIndex2i EdgeV1 = Mesh.GetOrientedBoundaryEdgeV(EdgeID1);
+			SequentialQuadEdges.Add({ EdgeV0, EdgeV1 });
+		}
+		else
+		{
+			bFoundInvalidCase = true;
+		}
+	}
+
+	if (bFoundInvalidCase || SequentialQuadEdges.Num() != NumEdges)
+	{
+		AppendLoopQuads(Mesh, Loop);
+		return;
+	}
+
+	Loop.NormalsA.SetNum(Loop.MeshVertices.Num());
+	Loop.NormalsB.SetNum(Loop.MeshVertices.Num());
+	for (int32 k = 0; k < Loop.MeshVertices.Num(); ++k)
+	{
+		Loop.NormalsA[k] = FMeshNormals::ComputeVertexNormal(Mesh, Loop.MeshVertices[k]);
+		Loop.NormalsB[k] = FMeshNormals::ComputeVertexNormal(Mesh, Loop.NewMeshVertices[k]);
+	}
+
+	int32 N = NumSubdivisions;
+	TArray<TArray<int32>> VertexSpans;
+	VertexSpans.SetNum(N + 2);
+
+	auto AppendVertColumn = [&Mesh, &VertexSpans, &N](int32 StartVID, int32 EndVID)
+	{
+		FVector3d StartPos = Mesh.GetVertex(StartVID);
+		FVector3d EndPos = Mesh.GetVertex(EndVID);
+		VertexSpans[0].Add(StartVID);
+		for (int32 j = 0; j < N; ++j)
+		{
+			double T = (double)(j + 1) / (double)(N + 1);
+			int32 NewVertID = Mesh.AppendVertex(Lerp(StartPos, EndPos, T));
+			VertexSpans[j + 1].Add(NewVertID);
+		}
+		VertexSpans[N + 1].Add(EndVID);
+	};
+
+	for (int32 k = 0; k < NumEdges; ++k)
+	{
+		FIndex2i EdgeV0 = SequentialQuadEdges[k].EdgeV0;
+		FIndex2i EdgeV1 = SequentialQuadEdges[k].EdgeV1;
+
+		int32 QuadA = EdgeV0.B;
+		int32 QuadB = EdgeV0.A;
+		int32 QuadC = EdgeV1.B;
+		int32 QuadD = EdgeV1.A;
+
+		if (k == 0)
+		{
+			AppendVertColumn(QuadA, QuadD);
+		}
+
+		if (k != NumEdges - 1)
+		{
+			AppendVertColumn(QuadB, QuadC);
+		}
+		else
+		{
+			for (int32 j = 0; j <= (N+1); ++j)
+			{
+				int FirstColVal = VertexSpans[j][0];
+				VertexSpans[j].Add(FirstColVal);
+			}
+		}
+
+	}
+
+	TArray<TArray<FIndex2i>> QuadSpans;
+	QuadSpans.SetNum(N + 1);
+
+	int32 NumStrips = VertexSpans.Num() - 1;
+	for (int32 k = 0; k < NumEdges; ++k)
+	{
+		FIndex2i GroupKey = GetGroupKey(k);
+		int32 NewGroupID = NewGroupIDs[GroupKey];
+
+		for (int32 j = 0; j < NumStrips; ++j)
+		{
+			FIndex2i QuadTris(IndexConstants::InvalidID, IndexConstants::InvalidID);
+			QuadTris.A = Mesh.AppendTriangle(VertexSpans[j][k], VertexSpans[j][k+1], VertexSpans[j+1][k], NewGroupID);
+			QuadTris.B = Mesh.AppendTriangle(VertexSpans[j+1][k+1], VertexSpans[j+1][k], VertexSpans[j][k+1], NewGroupID);
+			QuadSpans[j].Add(QuadTris);
+
+			Loop.StripQuads.Add(QuadTris);
+		}
+	}
+
+	Loop.StripQuadPatch.InitializeFromQuadPatch(Mesh, QuadSpans, VertexSpans);
+}
+
+
+
+
+void FMeshBevel::AppendJunctionVertexPolygon_Multi(FDynamicMesh3& Mesh, FBevelVertex& Vertex)
+{
+	// This function appends triangulations for junction vertices, ie where 3+ bevel-edges meet and
+	// so a polygonal face must be inserted. For a multi-segment bevel, we need to tessellate the
+	// polygon with interior vertices so that the interior vertices can be displaced later based
+	// on the bevel profile shape. And this is also where we will precompute certain information about
+	// the interior vertices, like (eg) barycentric coords of interior vertices, so that we can later
+	// deform them relative to the profile shapes on the boundary. 
+	//
+	// This is probably the most complex part of doing the bevel, and there are special cases for 
+	// valence-3 and valence-4, with lots of similar-but-slightly-different code.
+
+	MESH_BEVEL_DEBUG_CHECK(Vertex.VertexType == EBevelVertexType::JunctionVertex);
+
+	// TODO: all the setup here could be done in parallel, which might matter given that
+	// we are doing a tessellation, conformal unwrap, etc. Would need to cache temporary data.
+	// Junction vertices can never be directly connected so there should not be any order effects.
+
+	// UnlinkJunctionVertex() split the junction vertex into N vertices, one for each
+	// (now disconnected) triangle-wedge. The wedges are ordered such that their wedge-vertices
+	// define a polygon with correct winding. However in this case we have additional interior
+	// vertices along the polygon edges. So we need to construct a mesh for the junction vertex
+	// that has a compatible boundary and (in most cases) additional interior vertices, which
+	// is done in this function via a regular tessellation of the polygon
+
+	int32 NumEdgeVerts = 0;		// this should be determined by NumSubdivisions but we will compute from actual mesh here
+	bool bAllEdgesHaveSameNumVerts = true;	// current code only handles the case where all edges have same # of subdivisions.
+
+	// First step is to find the loop of vertices around the border of the junction vertex.
+	// Each incoming source-mesh edge became a "wedge", with start vertex A and end B, and
+	// then a quad-patch was appended connecting A to B. So we need to get the right "column"
+	// of the quad patch, which will give us the vertex span [A, ..., B]. These are accumulated
+	// in-wedge-sequence into PolygonPoints/Vertices. The start/end vertices A and B are accumulated 
+	// in-order in the PolygonCorners/Indices/VIDs lists. 
+	// (In both sets, the 'B' vertex is always skipped because it will be added as A in the next span)
+	TArray<FVector3d> PolygonPoints, PolygonCorners;
+	TArray<int32> PolygonVertices, PolygonCornerVIDs;
+	int32 NumWedges = Vertex.Wedges.Num();
+	for (int32 wi = 0; wi < NumWedges; ++wi)
+	{
+		FOneRingWedge& CurWedge = Vertex.Wedges[wi];
+		FOneRingWedge& NextWedge = Vertex.Wedges[(wi+1)%NumWedges];
+
+		int32 A = CurWedge.WedgeVertex;
+		int32 B = NextWedge.WedgeVertex;
+
+		// We don't know which Edge will contain the [A,B] vertices, and whether it is the 'start'
+		// or 'end' of the Edge. So currently just linear searching through all of them and checking
+		// each end. Maybe Vertex.IncomingBevelEdgeIndices would work here instead
+		TArray<int32> BevelEdgeQuadStripEndVertices;
+		bool bFound = false;
+		for (FBevelEdge& Span : Edges)
+		{
+			// this helper appends the vertex-column of an adjacent edge-quadstrip to the current vertex polygon loop
+			auto AppendStrip = [&](TArray<int32>& OrderedVertices)
+			{
+				PolygonCorners.Add(Mesh.GetVertex(OrderedVertices[0]));
+				PolygonCornerVIDs.Add(OrderedVertices[0]);
+				for (int32 j = 0; j < OrderedVertices.Num()-1; ++j)
+				{
+					int32 VertexID = OrderedVertices[j];
+					PolygonVertices.Add(VertexID);
+					PolygonPoints.Add(Mesh.GetVertex(VertexID));
+				}
+				if (NumEdgeVerts == 0)
+				{
+					NumEdgeVerts = OrderedVertices.Num();
+				}
+				else if (OrderedVertices.Num() != NumEdgeVerts)
+				{
+					bAllEdgesHaveSameNumVerts = false;
+				}
+				bFound = true;
+			};
+
+			// try start and end columns, and handle case where vertex ordering might be reversed  (should this be possible, because of consistent mesh winding??)
+			Span.StripQuadPatch.GetVertexColumn(0, BevelEdgeQuadStripEndVertices);
+			if (BevelEdgeQuadStripEndVertices[0] == A && BevelEdgeQuadStripEndVertices.Last() == B)
+			{
+				AppendStrip(BevelEdgeQuadStripEndVertices);
+				break;
+			}
+			if (BevelEdgeQuadStripEndVertices[0] == B && BevelEdgeQuadStripEndVertices.Last() == A)
+			{
+				Algo::Reverse(BevelEdgeQuadStripEndVertices);
+				AppendStrip(BevelEdgeQuadStripEndVertices);
+				break;
+			}
+
+			Span.StripQuadPatch.GetVertexColumn(Span.StripQuadPatch.NumVertexCols() - 1, BevelEdgeQuadStripEndVertices);
+			if (BevelEdgeQuadStripEndVertices[0] == A && BevelEdgeQuadStripEndVertices.Last() == B)
+			{
+				AppendStrip(BevelEdgeQuadStripEndVertices);
+				break;
+			}
+			if (BevelEdgeQuadStripEndVertices[0] == B && BevelEdgeQuadStripEndVertices.Last() == A)
+			{
+				Algo::Reverse(BevelEdgeQuadStripEndVertices);
+				AppendStrip(BevelEdgeQuadStripEndVertices);
+				break;
+			}
+		}
+		if (!bFound)
+		{
+			PolygonVertices.Add(CurWedge.WedgeVertex);
+			PolygonPoints.Add(Mesh.GetVertex(CurWedge.WedgeVertex));
+			bAllEdgesHaveSameNumVerts = false;
+		}
+	}
+
+	// code below won't work unless each edge has same vert count, in that case just triangulate the polygon (ie fail)
+	if (bAllEdgesHaveSameNumVerts == false || NumEdgeVerts == 0)
+	{
+		TArray<FIndex3i> Triangles;
+		PolygonTriangulation::TriangulateSimplePolygon<double>(PolygonPoints, Triangles);
+		Vertex.NewGroupID = Mesh.AllocateTriangleGroup();
+		for (FIndex3i Tri : Triangles)
+		{
+			int32 tid = Mesh.AppendTriangle(PolygonVertices[Tri.A], PolygonVertices[Tri.B], PolygonVertices[Tri.C], Vertex.NewGroupID);
+			if (Mesh.IsTriangle(tid))
+			{
+				Vertex.NewTriangles.Add(tid);
+			}
+		}
+		return;
+	}
+
+	// find centroid of polygon
+	FVector3d Centroid = FVector3d::Zero();
+	for (FVector3d Pos : PolygonPoints)
+	{
+		Centroid += Pos;
+	}
+	Centroid /= (double)PolygonPoints.Num();
+
+
+	// Now we have enough information to tessellate and fill the hole. 
+	// Apologies for the convoluted code here. The general case was implemented
+	// first, and then the 3/4 special-cases were added. Probably they should
+	// be refactored out into 3 separate functions. Sorry! -rms
+
+	// Append corners of the polygon into a temporary mesh, and then triangles
+	// For the case of > 4 vertices, add the centroid and triangulate as a fan
+	FDynamicMesh3 TmpMesh;
+	int32 NCorners = PolygonCorners.Num();
+	for (int32 k = 0; k < NCorners; ++k)
+	{
+		int NewVID = TmpMesh.AppendVertex(PolygonCorners[k]);
+		ensure(NewVID == k);
+	}
+	if (NCorners == 3)	// single triangle
+	{
+		TmpMesh.AppendTriangle(0, 1, 2);
+	}
+	else if (NCorners == 4)	// quad
+	{
+		TmpMesh.AppendTriangle(0, 1, 2);
+		TmpMesh.AppendTriangle(2, 3, 0);
+	}
+	else  // make triangle fan
+	{
+		int32 CentroidID = TmpMesh.AppendVertex(Centroid);
+		for (int32 i = 0; i < NCorners; ++i)
+		{
+			int NewTID = TmpMesh.AppendTriangle(i, ((i + 1) % NCorners), CentroidID);
+		}
+	}
+
+	// Ok, for the case of a triangle, we will use barycentric coordinates of the initial
+	// triangle rather than trying to do general polygon barycentrics. To avoid writing custom
+	// triangle-tess code here, FUniformTessellate is still used, but it doesn't return the barycentrics.
+	// So instead we store the 3 vertex-coordinate-functions in vertex colors and allow them to be lerp'd
+	if (NCorners == 3)
+	{
+		// note that this construction of TriTess will be the same for every valence-3 corner, and
+		// probably should be cached unless NumEdgeVerts changes... 
+		FDynamicMesh3 TriTess;
+		TriTess = TmpMesh;
+		TriTess.SetVertex(0, FVector3d(1, 0, 0));
+		TriTess.SetVertex(1, FVector3d(0, 1, 0));
+		TriTess.SetVertex(2, FVector3d(0.5, 0.75, 0));
+		TriTess.EnableVertexColors(FVector3f::Zero());
+		TriTess.SetVertexColor(0, FVector3f(1, 0, 0));
+		TriTess.SetVertexColor(1, FVector3f(0, 1, 0));
+		TriTess.SetVertexColor(2, FVector3f(0, 0, 1));
+
+		FUniformTessellate Tesselator(&TriTess);
+		Tesselator.TessellationNum = NumEdgeVerts - 2;
+		bool bTesselateOK = Tesselator.Compute();
+		checkSlow(bTesselateOK);
+
+		// need mapping between the boundary of this tessellated mesh and the triangle-shaped polygon
+		// surrounding the hole in the original mesh. Derive this from the known correspondence
+		// between the first vertex in each
+		FMeshBoundaryLoops BoundaryLoops(&TriTess, true);
+		TArray<int32> TriLoopVerts = BoundaryLoops.Loops[0].Vertices;
+		int32 LoopN = TriLoopVerts.Num();
+		checkSlow(PolygonVertices[0] == PolygonCornerVIDs[0]);
+		checkSlow(PolygonVertices.Num() == LoopN);
+		int32 StartIndex = TriLoopVerts.IndexOfByKey(0);		// index of first vertex in loop
+
+		TArray<int32> TriTessVertIDToPolygonIndexMap;
+		TriTessVertIDToPolygonIndexMap.SetNum(TriTess.MaxVertexID());
+		for (int32 k = 0; k < LoopN; ++k)
+		{
+			int32 LoopVertID = TriLoopVerts[(k + StartIndex) % LoopN];
+			TriTessVertIDToPolygonIndexMap[LoopVertID] = k;
+		}
+
+		FVector3d PosA = TmpMesh.GetVertex(0);
+		FVector3d PosB = TmpMesh.GetVertex(1);
+		FVector3d PosC = TmpMesh.GetVertex(2);
+
+		// for this special case we only store the 3 corners in the InteriorBorderLoop, this will be detected in
+		// the profile code and used correctly there
+		Vertex.InteriorVertices.Reserve(FMath::Max((NumEdgeVerts-1)*(NumEdgeVerts-1), 0));		// this is for quad so it's more than needed...
+		Vertex.InteriorBorderLoop = TArray<int32>({ PolygonCornerVIDs[0], PolygonCornerVIDs[1], PolygonCornerVIDs[2] });
+
+		TArray<int32> VertexMap;
+		VertexMap.SetNum(TriTess.MaxVertexID());
+
+		// Append the new interior vertices, boundary vertices already exist and are just updated in the map.
+		// For interior vertices the barycentric coords are stored in the BorderFrameWeights
+		for (int32 vid : TriTess.VertexIndicesItr())
+		{
+			if (TriTess.IsBoundaryVertex(vid))
+			{
+				int32 LoopIndex = TriTessVertIDToPolygonIndexMap[vid];
+				int32 ExistingVertID = PolygonVertices[LoopIndex];
+				VertexMap[vid] = ExistingVertID;
+			}
+			else
+			{
+				FVector3f BaryCoords = TriTess.GetVertexColor(vid);
+				BaryCoords /= (BaryCoords.X + BaryCoords.Y + BaryCoords.Z);
+				FVector3d InterpPos = BaryCoords.X * PosA + BaryCoords.Y * PosB + BaryCoords.Z * PosC;
+				int32 NewVID = Mesh.AppendVertex(InterpPos);
+
+				FBevelVertex_InteriorVertex& InteriorVertex = Vertex.InteriorVertices.Emplace_GetRef();
+				InteriorVertex.VertexID = NewVID;
+				InteriorVertex.BorderFrameWeight.Add(FVector3d(BaryCoords.X, BaryCoords.Y, BaryCoords.Z));
+				VertexMap[vid] = NewVID;
+			}
+		}
+
+		// finally append the triangles
+		TArray<FIndex3i> Triangles;
+		Vertex.NewGroupID = Mesh.AllocateTriangleGroup();
+		for (FIndex3i Triangle : TriTess.TrianglesItr())
+		{
+			// NOTE FLIP HERE! This is because outer loop is oriented for outside tris, so we need to flip
+			int T0 = Mesh.AppendTriangle(VertexMap[Triangle.A], VertexMap[Triangle.C], VertexMap[Triangle.B], Vertex.NewGroupID);
+			if (Mesh.IsTriangle(T0))
+			{
+				Vertex.NewTriangles.Add(T0);
+			}
+		}
+
+		// we are done the valnce=3 case, exit this function
+		return;
+	}
+
+	// For a quad, the tessellation is simple and is just done directly here
+	if (NCorners == 4)
+	{
+		FDenseGrid2i VertexGrid;
+		VertexGrid.Resize(NumEdgeVerts, NumEdgeVerts);
+		VertexGrid.AssignAll(-1);
+
+		// transfer vertex polygon into border of grid
+		int32 PolygonIdx = 0;
+		for (int32 xi = 0; xi < NumEdgeVerts; ++xi)
+		{
+			VertexGrid.At(xi,0) = PolygonVertices[PolygonIdx++];
+		}
+		for ( int32 yi = 1; yi < NumEdgeVerts-1; ++yi )
+		{
+			VertexGrid.At(NumEdgeVerts-1, yi) = PolygonVertices[PolygonIdx++];
+		}
+		for (int32 xi = NumEdgeVerts-1; xi >= 0; --xi)
+		{
+			VertexGrid.At(xi, NumEdgeVerts-1) = PolygonVertices[PolygonIdx++];
+		}
+		for (int32 yi = NumEdgeVerts - 2; yi >= 1; --yi)
+		{
+			VertexGrid.At(0, yi) = PolygonVertices[PolygonIdx++];
+		}
+		ensure(PolygonIdx == PolygonVertices.Num());
+
+		int c00 = VertexGrid.At(0, 0);
+		int c10 = VertexGrid.At(NumEdgeVerts-1, 0);
+		int c01 = VertexGrid.At(0, NumEdgeVerts-1);
+		int c11 = VertexGrid.At(NumEdgeVerts-1, NumEdgeVerts-1);
+		FVector3d V00 = Mesh.GetVertex(c00);
+		FVector3d V10 = Mesh.GetVertex(c10);
+		FVector3d V01 = Mesh.GetVertex(c01);
+		FVector3d V11 = Mesh.GetVertex(c11);
+
+		// for this special case we only store the 4 corners in the InteriorBorderLoop, this will be detected in
+		// the profile code and used correctly there
+		Vertex.InteriorVertices.Reserve(FMath::Max((NumEdgeVerts - 1)* (NumEdgeVerts - 1), 0));
+		Vertex.InteriorBorderLoop = TArray<int32>({ c00, c10, c01, c11 });
+
+		// append new vertices for the interior, while storing their U/V coords in the BorderFrameWeights
+		for (int yi = 1; yi < NumEdgeVerts-1; ++yi)
+		{
+			double ty = (double)yi / (double)(NumEdgeVerts - 1);
+			FVector3d A = Lerp(V00, V01, ty);
+			FVector3d B = Lerp(V10, V11, ty);
+			for (int xi = 1; xi < NumEdgeVerts-1; ++xi)
+			{
+				ensure(VertexGrid.At(xi, yi) == -1);
+				double tx = (double)xi / (double)(NumEdgeVerts - 1);
+				FVector3d InterpPos = Lerp(A, B, tx);
+				int32 NewVID = Mesh.AppendVertex(InterpPos);
+				VertexGrid.At(xi, yi) = NewVID;
+
+				FBevelVertex_InteriorVertex& InteriorVertex = Vertex.InteriorVertices.Emplace_GetRef();
+				InteriorVertex.VertexID = NewVID;
+				InteriorVertex.BorderFrameWeight.Add(FVector3d(tx, ty, 0));
+			}
+		}
+
+		// create new vertices for any that don't already have a mapping back to boundary ring
+		TArray<FIndex3i> Triangles;
+		Vertex.NewGroupID = Mesh.AllocateTriangleGroup();
+
+		for (int y0 = 0; y0 < NumEdgeVerts - 1; ++y0)
+		{
+			for (int x0 = 0; x0 < NumEdgeVerts - 1; ++x0)
+			{
+				int i00 = VertexGrid.At(x0, y0);
+				int i10 = VertexGrid.At(x0+1, y0);
+				int i01 = VertexGrid.At(x0, y0+1);
+				int i11 = VertexGrid.At(x0+1, y0+1);
+				int T0 = Mesh.AppendTriangle(i10, i00, i01, Vertex.NewGroupID);
+				if (Mesh.IsTriangle(T0))
+				{
+					Vertex.NewTriangles.Add(T0);
+				}
+				int T1 = Mesh.AppendTriangle(i01, i11, i10, Vertex.NewGroupID);
+				if (Mesh.IsTriangle(T1))
+				{
+					Vertex.NewTriangles.Add(T1);
+				}
+
+			}
+		}
+
+		return;
+	}
+
+	// Ok, now we are in the general case (currently used for >= 5-vertex polygons)
+	// We will use Uniform Tessellation to subdivide the initial polygon triangulation up to
+	// the target polycount. This is a bit annoying because we have to reconstruct the boundary
+	// correspondence, perhaps FUniformTessellate could be improved to return this information
+	FUniformTessellate Tesselator(&TmpMesh);
+	Tesselator.TessellationNum = NumEdgeVerts - 2;
+	bool bTesselateOK = Tesselator.Compute();
+
+	// create flattened 2D version of this tessellated mesh  (maybe should optimize in 2D for inner fairness??)
+	FDynamicMeshUVEditor UVEditor(&TmpMesh, 0, true);
+	TArray<int32> AllTriangles;
+	for (int32 tid : TmpMesh.TriangleIndicesItr()) AllTriangles.Add(tid);
+	TArray<int32> TmpMeshToUVMap; bool bMapIsCompact = false;
+	UVEditor.SetToPerVertexUVs(TmpMeshToUVMap, bMapIsCompact);
+	check(bMapIsCompact == true);
+	UVEditor.SetTriangleUVsFromFreeBoundarySpectralConformal(AllTriangles, true, true);
+	UVEditor.ScaleUVAreaTo3DArea(AllTriangles, true);
+	FDynamicMeshUVOverlay* UVOverlay = UVEditor.GetOverlay();
+
+	// Now we need to append this triangulated patch back into the main mesh, ie "fill the hole"
+	// that exists for this junction vertex. Ideally we will stitch at the compatible patch border, via this map
+	TArray<int32> VertexMap;
+	VertexMap.Init(-1, TmpMesh.MaxVertexID());
+
+	// We should have an exact match between the existing loop of boundary vertices (in PolygonVertices) and 
+	// the boundary loop of the patch. And we know the correspondences at the Corners as we constructed them 
+	// above in the TmpMesh.AppendVertex calls. So we should be able to find Corner 0 in the boundary loop,
+	// and walk around the two loops from that initial correspondence. This block will do that and store it in VertexMap, if possible.
+	int32 NV = PolygonVertices.Num();
+	FMeshBoundaryLoops BoundaryLoops(&TmpMesh, true);	// should only ever be one loop...
+	bool bFoundLoop = false;
+	TArray<int32> BoundaryLoopVerts;
+	for (int32 k = 0; k < BoundaryLoops.GetLoopCount() && bFoundLoop == false; ++k)
+	{
+		if (BoundaryLoops[k].Vertices.Num() == NV)
+		{
+			BoundaryLoopVerts = BoundaryLoops[k].Vertices;
+			bFoundLoop = true;	// conceivably it's possible this isn't the correct loop, somehow. That case is not handled
+
+			// find our known-correspondence vertices, from the initial polygon corners
+			int32 CornerVID = PolygonCornerVIDs[0];
+			int32 FoundIdx = BoundaryLoopVerts.IndexOfByKey(0);
+			int32 SecondCornerVID = PolygonCornerVIDs[1];
+			int32 FoundSecondIdx = BoundaryLoopVerts.IndexOfByKey(1);
+
+			if (FoundIdx != INDEX_NONE && FoundSecondIdx != INDEX_NONE)
+			{
+				// detect if the boundary loop is reversed, and if so, flip it
+				if (BoundaryLoopVerts[(FoundIdx+(NumEdgeVerts-1)) % NV] != BoundaryLoopVerts[FoundSecondIdx] )
+				{
+					Algo::Reverse(BoundaryLoopVerts);
+					FoundIdx = BoundaryLoopVerts.IndexOfByKey(0);
+					FoundSecondIdx = BoundaryLoopVerts.IndexOfByKey(1);
+				}
+
+				if (ensure(BoundaryLoopVerts[(FoundIdx+(NumEdgeVerts-1)) % NV] == BoundaryLoopVerts[FoundSecondIdx]))
+				{
+					// walk around loop at corresponding indices and populate VertexMap
+					for (int32 j = 0; j < NV; ++j)
+					{
+						int32 ExistingID = PolygonVertices[j];
+						int32 NewID = BoundaryLoopVerts[(FoundIdx + j) % NV];
+						VertexMap[NewID] = ExistingID;
+					}
+				}
+			}
+		}
+	}
+
+	Vertex.InteriorVertices.Reserve(FMath::Max(TmpMesh.VertexCount() - NV, 0));
+
+	int32 NumLoopVerts = BoundaryLoopVerts.Num();
+	FPolygon2d BorderPolygon;
+	for (int32 vid : BoundaryLoopVerts)
+	{
+		BorderPolygon.AppendVertex( (FVector2d)UVOverlay->GetElement(vid) );
+	}
+
+	// Create new vertices for any that don't already have a mapping back to boundary ring.
+	// For those interior vertices we also want to compute 2D Mean Value Coordinates (MVC)
+	// relative to the boundary polygon. We will use these later to interpolate 3D positions from
+	// the 3D border.
+	TArray<FIndex3i> Triangles;
+	Vertex.NewGroupID = Mesh.AllocateTriangleGroup();
+	for (int32 vid : TmpMesh.VertexIndicesItr())
+	{
+		if (VertexMap[vid] == -1)
+		{
+			int32 NewVID = Mesh.AppendVertex(TmpMesh.GetVertex(vid));
+			VertexMap[vid] = NewVID;
+
+			// construct interior MVC weights from 2D mesh
+			FBevelVertex_InteriorVertex& InteriorVertex = Vertex.InteriorVertices.Emplace_GetRef();
+			InteriorVertex.VertexID = NewVID;
+			InteriorVertex.BorderFrameWeight.Reserve(BoundaryLoopVerts.Num());
+			FVector2d UVPosition = (FVector2d)UVOverlay->GetElement(vid);
+			double WeightSum = 0;
+			for (int32 k = 0; k < NumLoopVerts; ++k)
+			{
+				int32 BoundaryVID = BoundaryLoopVerts[k];
+				int32 NewBoundaryVID = VertexMap[BoundaryVID];
+				checkSlow(NewBoundaryVID >= 0);
+				FVector2d BoundaryUVPosition = BorderPolygon[k];
+				FVector2d Prev = BorderPolygon[(k-1+NumLoopVerts) % NumLoopVerts];
+				FVector2d Next = BorderPolygon[(k+1) % NumLoopVerts];
+
+				FVector2d BoundaryFrameX = Normalized(Next - Prev);
+				FVector2d BoundaryFrameY = PerpCW(BoundaryFrameX);
+
+				FVector2d DeltaUV = UVPosition - BoundaryUVPosition;
+				double DeltaX = BoundaryFrameX.Dot(DeltaUV);
+				double DeltaY = BoundaryFrameY.Dot(DeltaUV);
+
+				// compute Polygon MVC weight ( https://cgvr.cs.uni-bremen.de/teaching/cg_literatur/barycentric_floater.pdf )
+				double Dist = Distance(UVPosition, BoundaryUVPosition);
+				double Weight = 1.0;
+				if ( Dist > FMathd::ZeroTolerance )
+				{
+					FVector2d DeltaP = Normalized(BoundaryUVPosition - UVPosition);
+					double T1 = UE::Geometry::VectorUtil::VectorTanHalfAngle(Normalized(Prev - UVPosition), DeltaP);
+					double T2 = UE::Geometry::VectorUtil::VectorTanHalfAngle(Normalized(Next - UVPosition), DeltaP);
+					Weight = (T1 + T2) / Dist;
+				}
+
+				InteriorVertex.BorderFrameWeight.Add( FVector3d(DeltaX, DeltaY, Weight) );
+				WeightSum += Weight;
+			}
+			for (FVector3d& Weight : InteriorVertex.BorderFrameWeight)
+			{
+				Weight.Z /= WeightSum;		// normalize MVC weights
+			}
+		}
+	}
+
+	Vertex.InteriorBorderLoop.Reserve(NumLoopVerts);
+	for (int32 vid : BoundaryLoopVerts)
+	{
+		Vertex.InteriorBorderLoop.Add(VertexMap[vid]);
+	}
+
+	// PolygonVertices array is always constructed in reversed orientation, so patch will be flipped otherwise
+	TmpMesh.ReverseOrientation();
+
+	// append new triangles
+	for (FIndex3i Tri : TmpMesh.TrianglesItr())
+	{
+		int32 A = VertexMap[Tri.A];
+		int32 B = VertexMap[Tri.B];
+		int32 C = VertexMap[Tri.C];
+		int32 tid = Mesh.AppendTriangle(A, B, C, Vertex.NewGroupID);
+		if (Mesh.IsTriangle(tid))
+		{
+			Vertex.NewTriangles.Add(tid);
+		}
+	}
+
+}
+
+
+
+void FMeshBevel::AppendTerminatorVertexTriangles_Multi(FDynamicMesh3& Mesh, FBevelVertex& Vertex)
+{
+	// A Terminator Vertex occurs at the end of an non-loop bevel-edge, where the bevel "ends". 
+	// Disconnecting the mesh to insert the bevel edge-strip will have created a triangle-shaped hole
+	// in the mesh that must be filled. In the Multi-case, we have already meshed the adjacent bevel-edge,
+	// so it's not a simple triangle but one with subdivisions along the edge adajacent to the bevel-edge.
+	// So basically we have to find that already-meshed edge, and extract the column from it's quadpatch
+	// that we want to stitch to, and then add a triangle-fan to fill the hole polygon
+
+	// TODO: maybe could do delaunay triangulation if the polygon below is nearly coplanar...
+
+	MESH_BEVEL_DEBUG_CHECK(Vertex.VertexType == EBevelVertexType::TerminatorVertex);
+
+	check(Vertex.IncomingBevelEdgeIndices.Num() == 1);
+
+	int32 RingSplitEdgeID = Vertex.TerminatorInfo.A;
+	if (Mesh.IsEdge(RingSplitEdgeID))
+	{
+		FIndex2i SplitEdgeV = Mesh.GetEdgeV(RingSplitEdgeID);
+		int32 FarVertexID = SplitEdgeV.OtherElement(Vertex.VertexID);
+
+		int32 BevelEdgeIndex = Vertex.IncomingBevelEdgeIndices[0];
+		const FBevelEdge& IncomingEdge = Edges[BevelEdgeIndex];
+
+		int32 WedgeVertexA = Vertex.Wedges[0].WedgeVertex;
+		int32 WedgeVertexB = Vertex.Wedges[1].WedgeVertex;
+
+		// figure out which side of the incoming-edge quad-strip the terminator connects to
+		int32 ColumnIndex = IncomingEdge.StripQuadPatch.FindColumnIndex(WedgeVertexA);
+		check(ColumnIndex == 0 || ColumnIndex == IncomingEdge.StripQuadPatch.NumVertexCols() - 1);
+		TArray<int32> QuadStripEdgeVerts;
+		IncomingEdge.StripQuadPatch.GetVertexColumn(ColumnIndex, QuadStripEdgeVerts);
+		check(QuadStripEdgeVerts.Contains(WedgeVertexB));
+
+		// make sure it is oriented consistently
+		int32 FirstEdgeID = Mesh.FindEdge(QuadStripEdgeVerts[0], QuadStripEdgeVerts[1]);
+		FIndex2i QuadEdgeV = Mesh.GetOrientedBoundaryEdgeV(FirstEdgeID);
+		if (QuadEdgeV.A != QuadStripEdgeVerts[0])
+		{
+			Algo::Reverse(QuadStripEdgeVerts);
+		}
+
+		// should have computed this GroupID in initial setup
+		int32 UseGroupID = (Vertex.NewGroupID >= 0) ? Vertex.NewGroupID : Mesh.AllocateTriangleGroup();
+
+		int32 NumEdges = QuadStripEdgeVerts.Num()-1;
+		for (int32 k = 0; k < NumEdges; ++k)
+		{
+			int32 tid = Mesh.AppendTriangle(QuadStripEdgeVerts[k+1], QuadStripEdgeVerts[k], FarVertexID, UseGroupID);
+			MESH_BEVEL_DEBUG_CHECK(tid >= 0);
+			if (Mesh.IsTriangle(tid))
+			{
+				Vertex.NewTriangles.Add(tid);
+			}
+		}
+	}
+}
+
+
+
+void FMeshBevel::AppendTerminatorVertexPairQuad_Multi(FDynamicMesh3& Mesh, FBevelVertex& Vertex0, FBevelVertex& Vertex1)
+{
+	MESH_BEVEL_DEBUG_CHECK(Vertex0.VertexType == EBevelVertexType::TerminatorVertex);
+	MESH_BEVEL_DEBUG_CHECK(Vertex1.VertexType == EBevelVertexType::TerminatorVertex);
+
+	// This is a variant of AppendTerminatorVertexTriangle that handles the case where basically two
+	// Terminator Vertices are directly connected by a non-beveled mesh edge that was used as the ring-split-edge.
+	// Since both sides were opened, we have a quad-shaped hole instead of a triangle-shaped hole, with a 
+	// (subdivided, in the multi-case) quad-edge at each end. So this hole can be filled with a quad-patch that
+	// stitches together the existing edges (which must have already been meshed)
+
+	auto CollectTerminatorVtxInfo = [this](FDynamicMesh3& Mesh, FBevelVertex& Vertex, TArray<int32>& QuadStripEdgeVertsOut)
+	{
+		int32 BevelEdgeIndex = Vertex.IncomingBevelEdgeIndices[0];
+		const FBevelEdge& IncomingEdge = Edges[BevelEdgeIndex];
+
+		int32 ColumnIndex = IncomingEdge.StripQuadPatch.FindColumnIndex(Vertex.Wedges[0].WedgeVertex);
+		check(ColumnIndex == 0 || ColumnIndex == IncomingEdge.StripQuadPatch.NumVertexCols() - 1);
+		IncomingEdge.StripQuadPatch.GetVertexColumn(ColumnIndex, QuadStripEdgeVertsOut);
+		check(QuadStripEdgeVertsOut.Contains(Vertex.Wedges[1].WedgeVertex));
+
+		int32 FirstEdgeID = Mesh.FindEdge(QuadStripEdgeVertsOut[0], QuadStripEdgeVertsOut[1]);
+		FIndex2i QuadEdgeV = Mesh.GetOrientedBoundaryEdgeV(FirstEdgeID);
+		if (QuadEdgeV.A != QuadStripEdgeVertsOut[0])
+		{
+			Algo::Reverse(QuadStripEdgeVertsOut);
+		}
+	};
+
+	TArray<int32> QuadStripEdgeVerts0, QuadStripEdgeVerts1;
+	CollectTerminatorVtxInfo(Mesh, Vertex0, QuadStripEdgeVerts0);
+	CollectTerminatorVtxInfo(Mesh, Vertex1, QuadStripEdgeVerts1);
+	if (QuadStripEdgeVerts0.Num() != QuadStripEdgeVerts1.Num() || QuadStripEdgeVerts0.Num() == 0)
+	{
+		// this could happen if the quadstrip failed on one or both sides? will leave a hole.
+		MESH_BEVEL_DEBUG_CHECK(false);
+		return;
+	}
+
+	// BIASED? should have computed this GroupID in initial setup
+	int32 UseGroupID = (Vertex0.NewGroupID >= 0) ? Vertex0.NewGroupID : Mesh.AllocateTriangleGroup();
+
+	int32 NumEdges = QuadStripEdgeVerts0.Num() - 1;
+	for (int32 k = 0; k < NumEdges; ++k)
+	{
+		int32 A0 = QuadStripEdgeVerts0[k];
+		int32 B0 = QuadStripEdgeVerts0[k+1];
+		int32 A1 = QuadStripEdgeVerts1[NumEdges-(k)];
+		int32 B1 = QuadStripEdgeVerts1[NumEdges-(k+1)];
+
+		int32 tid0 = Mesh.AppendTriangle(B0, A0, A1, UseGroupID);
+		MESH_BEVEL_DEBUG_CHECK(tid0 >= 0);
+		if (Mesh.IsTriangle(tid0))
+		{
+			Vertex0.NewTriangles.Add(tid0);
+		}
+
+		int32 tid1 = Mesh.AppendTriangle(A1, B1, B0, UseGroupID);
+		MESH_BEVEL_DEBUG_CHECK(tid1 >= 0);
+		if (Mesh.IsTriangle(tid1))
+		{
+			Vertex0.NewTriangles.Add(tid1);
+		}
+	}
+
+}
+
+
+void FMeshBevel::CreateBevelMeshing_Multi(FDynamicMesh3& Mesh)
+{
+	// This is the top-level driver function that is called after the mesh has been pulled apart along
+	// the bevelled edges. There are four cases - open edge spans, edge loops, corners/vertices with bevel-valence > 3 (become polygons),
+	// and "terminator" vertices of bevel-valence 1 (become triangles, except if directly connected, then they become quads).
+	// First we figure out the mesh connectivity, ie the 'stitching' between the pulled-apart geometry,
+	// and then we (optionally) apply a profile-curve shape along the bevel strips
+
+	for (FBevelEdge& Edge : Edges)
+	{
+		AppendEdgeQuads_Multi(Mesh, Edge);
+	}
+
+	for (FBevelLoop& Loop : Loops)
+	{
+		AppendLoopQuads_Multi(Mesh, Loop);
+	}
+
+
+	for (FBevelVertex& Vertex : Vertices)
+	{
+		if (Vertex.VertexType == EBevelVertexType::JunctionVertex)
+		{
+			if (Vertex.Wedges.Num() > 2)
+			{
+				AppendJunctionVertexPolygon_Multi(Mesh, Vertex);
+			}
+		}
+	}
+
+	// easier to do terminators last so that we can use quad edge to orient the triangle
+	TSet<FIndex2i> HandledQuadVtxPairs;
+	for (FBevelVertex& Vertex : Vertices)
+	{
+		if (Vertex.VertexType == EBevelVertexType::TerminatorVertex)
+		{
+			if (Vertex.ConnectedBevelVertex >= 0)
+			{
+				FBevelVertex& OtherVertex = Vertices[Vertex.ConnectedBevelVertex];
+				FIndex2i VtxPair(Vertex.VertexID, OtherVertex.VertexID);
+				VtxPair.Sort();
+				if (HandledQuadVtxPairs.Contains(VtxPair) == false)
+				{
+					AppendTerminatorVertexPairQuad_Multi(Mesh, Vertex, OtherVertex);
+					HandledQuadVtxPairs.Add(VtxPair);
+				}
+			}
+			else
+			{
+				AppendTerminatorVertexTriangles_Multi(Mesh, Vertex);
+			}
+		}
+	}
+
+	
+	// At this point, all topology is determined. Now if desired a profile-curve shape can be applied as a postprocess.
+	// (Not entirely certain this is the right strategy for all profile types but it works OK for a simple round profile)
+	if (FMathd::Abs(RoundWeight) > FMathf::ZeroTolerance)
+	{
+		ApplyProfileShape_Round(Mesh);
+	}
+}
+
+
+
+
+FInterpCurveVector FMeshBevel::MakeArcSplineCurve(const FVector3d& PosA, FVector3d& NormalA, const FVector3d& PosB, FVector3d& NormalB) const
+{
+	FInterpCurveVector Curve;
+
+	// basic idea here is to approximate an arc w/ a bezier curve. We know A and B,
+	// and we need TangentA (T_A) and TangentB. We have a surface Normal at A and B, so
+	// for T_A we can project B onto the surface tangent plane at A (ie with NormalA),
+	// and then B-A gives us a direction for T_A, and the same works for for B
+	// 
+	//         T_A
+	//       <-----
+	//    ^    __-*A
+	//    |   /
+	// T_B|  / (this is the arc, sorry)
+	//    | |
+	//      *B
+	//
+	// The lengths of the Tangents in this approach end up being the lengths of the sides
+	// of the square, if this were a 2D situation. Which could effectively define a 2D arc.
+	// For now, we are use a FInterpCurve instead, which is not quite a Bezier. With those
+	// lengths, the curve will end up too flat, so it is scaled by a parameter (default sqrt 2).
+	// This could perhaps be exposed as an option...
+	//
+	// The advantage of using a curve here instead of an arc is that it allows something reasonable
+	// to happen for messy cases, like where the curve might not be planar. But likely some
+	// issues are going to arise with ugly projections...
+
+	// project the vector (B-A) onto the plane (A,NormalA). 
+	FFrame3d PlaneNormalA(PosA, NormalA);
+	FVector3d TangentA = PlaneNormalA.ToPlane(PosB) - PosA;
+
+	// same for A-B onto (B, NormalB)
+	FFrame3d PlaneNormalB(PosB, NormalB);
+	FVector3d TangentB = PlaneNormalB.ToPlane(PosA) - PosB;
+
+	Curve.AddPoint(0, PosA);
+	Curve.AddPoint(1, PosB);
+	Curve.Points[0].InterpMode = EInterpCurveMode::CIM_CurveUser;
+	Curve.Points[1].InterpMode = EInterpCurveMode::CIM_CurveUser;
+	double TangentScale = FMathd::Abs(RoundWeight) * FMathd::Sqrt2;
+	if (RoundWeight >= 0)
+	{
+		Curve.Points[0].ArriveTangent = Curve.Points[0].LeaveTangent = TangentScale * TangentA;
+		Curve.Points[1].ArriveTangent = Curve.Points[1].LeaveTangent = -TangentScale * TangentB;
+	}
+	else
+	{
+		Curve.Points[0].ArriveTangent = Curve.Points[0].LeaveTangent = -TangentScale * TangentB;
+		Curve.Points[1].ArriveTangent = Curve.Points[1].LeaveTangent = TangentScale * TangentA;
+	}
+	return Curve;
+}
+
+
+void FMeshBevel::PlanarizeArcNormals(const TArray<FVector3d>& InitialEdgeCurve, bool bIsLoop, int32 Index, FVector3d& NormalA, FVector3d& NormalB, FVector3d& TangentOut, bool& bTangentIsValid) const
+{
+	TangentOut = FVector3d::Zero();
+
+	// We have two Normals A and B, which are destined to be used with MakeArcSplineCurve, ie they are the surface normals
+	// at the endpoints of the curve we want to generate, and we are going to use those normals to define the curve.
+	// In nearly all cases, we want this curve to be planar, so the normals need to lie in the same plane.
+	//
+	// To achieve that, we are going to assume that the tangent direction around the InitialEdgeCurve (actually a polyline)
+	// defines the desired plane. This makes sense if you think of the bevel as essentially sweeping a circular arc around
+	// the bevel paths and using it to cut away from the surface. Then the cut plane at any point on the curve would be normal
+	// to the curve tangent. So we find that tangent direction, and then project both surface normals onto it.
+	//
+	// Note that we use a purely opening-angle-based polyline-tangent estimate here. This makes less sense for a smooth curve path,
+	// but for (eg) a rectangular box, if we use the polyine-tangent definition that is an approximation of a continuous tangent, then
+	// the cut-angle at the corners will vary with the box dimensions, which is clearly wrong. 
+
+	int32 NV = InitialEdgeCurve.Num();
+	if (bIsLoop || (Index != 0 && Index != NV-1) )
+	{
+		FVector3d CurPos = InitialEdgeCurve[Index];
+		FVector3d PrevDir = InitialEdgeCurve[(Index-1 + NV) % NV] - CurPos;
+		FVector3d NextDir = InitialEdgeCurve[(Index+1) % NV] - CurPos;
+		if (NextDir.Normalize() && PrevDir.Normalize())
+		{
+			// this direction is the angle-based tangent, ie ignores edge length
+			TangentOut = NextDir - PrevDir;
+		}
+	}
+	else
+	{
+		FVector3d CurPos = InitialEdgeCurve[Index];
+		if (Index == 0)
+		{
+			FVector3d EdgeDir = InitialEdgeCurve[1] - CurPos;
+			if (EdgeDir.Normalize())
+			{
+				TangentOut = EdgeDir;
+			}
+		}
+		else if (Index == NV - 1)
+		{
+			FVector3d EdgeDir = CurPos - InitialEdgeCurve[NV-2];
+			if (EdgeDir.Normalize())
+			{
+				TangentOut = EdgeDir;
+			}
+		}
+		else
+		{
+			check(false);
+		}
+	}
+
+	if (TangentOut.Normalize())
+	{
+		// project normals onto plane of tangent
+		NormalA = Normalized(NormalA - NormalA.Dot(TangentOut) * TangentOut);
+		NormalB = Normalized(NormalB - NormalB.Dot(TangentOut) * TangentOut);
+		bTangentIsValid = true;
+	}
+}
+
+
+void FMeshBevel::ApplyProfileShape_Round(FDynamicMesh3& Mesh)
+{
+	// This function applies round profile curves to the bevel edge quadstrips & 
+	// bevel vertex tessellated-polygons that were computed in the _Multi() functions.
+	// All the geometry already exists so the job here is just to deform it into rounder shapes.
+	//
+	// The current code is based on various different strategies but the basic idea is to
+	// compute a arc-shaped 3D splines between correspondeng edge-vertex pairs , and then derive
+	// the shape of the vertex patches using the now-curved borders. For 4-sided patches
+	// this is relatively simple, 5+-sided we use a sort of MVC-based hole-filling strategy, and
+	// for 3-sided currently using a 3-sided spline batch taken from PN-tessellation. 
+	// The 3-sided case currently does not work very well though.	
+
+	// Loops are an easy case as there are no junction vertices, so each vertex-pair gets a separate
+	// spline and everything is done. 
+	for (FBevelLoop& Loop : Loops)
+	{
+		FQuadGridPatch& Patch = Loop.StripQuadPatch;
+		int32 NumCols = Patch.NumVertexCols();
+		for (int32 Col = 0; Col < NumCols-1; ++Col)
+		{
+			TArray<int32> ColVerts;
+			Patch.GetVertexColumn(Col, ColVerts);
+			int32 NV = ColVerts.Num();
+
+			int32 A = ColVerts[0], B = ColVerts.Last();
+			FVector3d PosA = Mesh.GetVertex(A), PosB = Mesh.GetVertex(B);
+			FVector3d NormalA = Loop.NormalsA[Col], NormalB = Loop.NormalsB[Col];
+
+			FVector3d Tangent; bool bTangentIsValid = false;
+			PlanarizeArcNormals(Loop.InitialPositions, true, Col, NormalA, NormalB, Tangent, bTangentIsValid);
+
+			FInterpCurveVector Curve = MakeArcSplineCurve(PosA, NormalA, PosB, NormalB);
+
+			// TODO: should have special case here for flat patches? they get squished out a bit...
+
+			for (int32 k = 1; k < (NV - 1); ++k)
+			{
+				int32 VID = ColVerts[k];
+				FVector3d Pos = Mesh.GetVertex(VID);
+				double T = (double)k / (double)(NV-1);
+				FVector3d CurvePos = Curve.Eval((float)T, Lerp(PosA, PosB, T));
+				Mesh.SetVertex(VID, CurvePos);
+			}
+		}
+	}
+
+
+	// kinda gross approach to keeping track of the normals along the ends of edge spans,
+	// ie after we make them curved, we want the normals along those curves, ie what we 
+	// want to use as the normal/tangent constraint for the bevelvertex-polygon. At this point
+	// the mesh topo should be finalized so we can just use an array...
+	TArray<FVector3d> DeformNormals;
+	DeformNormals.Init(FVector3d::Zero(), Mesh.MaxVertexID());
+
+	// we need to keep track of the curve used between each split pair of beveledge-vertices,
+	// and we need to be able to loop them up based on the vertex-pair.
+	TMap<FIndex2i, FInterpCurveVector> BorderCurves;
+
+	// search function for BorderCurves list
+	auto FindCurve = [&BorderCurves](int i0, int i1, bool& bReversed)
+	{
+		FInterpCurveVector* FoundCurve = BorderCurves.Find(FIndex2i(i0, i1));
+		bReversed = false;
+		if (FoundCurve == nullptr)
+		{
+			FoundCurve = BorderCurves.Find(FIndex2i(i1, i0));
+			bReversed = true;
+		}
+		return FoundCurve;
+	};
+
+	// Each edge is processed similar to a loop, ie the column of vertices bewteen each vertex-pair on 
+	// either side of the quad-strip gets deformed by fitting a spline curve based on the endpoints and end-normals. 
+	// 
+	for (FBevelEdge& Edge : Edges)
+	{
+		FQuadGridPatch& Patch = Edge.StripQuadPatch;
+
+		int32 NumCols = Patch.NumVertexCols();
+		for (int32 Col = 0; Col < NumCols; ++Col)
+		{
+			TArray<int32> ColVerts;
+			Patch.GetVertexColumn(Col, ColVerts);
+			int32 NV = ColVerts.Num();
+
+			int32 A = ColVerts[0];
+			int32 B = ColVerts.Last();
+			FVector3d PosA = Mesh.GetVertex(A);
+			FVector3d PosB = Mesh.GetVertex(B);
+			FVector3d NormalA = Edge.NormalsA[Col];
+			FVector3d NormalB = Edge.NormalsB[Col];
+
+			bool bProjectToTangentPlane = true;
+			bool bIsEndpoint = (Col == 0) || (Col == NumCols - 1);
+			int32 BevelEdgeIndex = (Col == 0) ? Edge.BevelVertices.A : Edge.BevelVertices.B;
+			check(BevelEdgeIndex >= 0);
+			const FBevelVertex& BevelVtx = Vertices[BevelEdgeIndex];
+			if (BevelVtx.VertexType == EBevelVertexType::JunctionVertex && BevelVtx.IncomingBevelEdgeIndices.Num() == 2)
+			{
+				bProjectToTangentPlane = false;
+			}
+
+			FVector3d Tangent; bool bTangentIsValid = false;
+			if (bProjectToTangentPlane)
+			{
+				PlanarizeArcNormals(Edge.InitialPositions, false, Col, NormalA, NormalB, Tangent, bTangentIsValid);
+			}
+
+			// TODO: should have special case here for flat patches? they get squished out a bit...
+			// This may update NormalA/NormalB for 'inverted' bevels (ie w/ negative RoundWeight)
+			FInterpCurveVector Curve = MakeArcSplineCurve(PosA, NormalA, PosB, NormalB);
+
+			// accumulate normal at endpoints
+			DeformNormals[A] += NormalA;
+			DeformNormals[B] += NormalB;
+
+			// save this curve in the curves list
+			BorderCurves.Add(FIndex2i(A, B), Curve);
+
+			// map linear-interpolation to curve parameter and evaluate curve
+			TArray<FVector3d> CurveVerts;
+			CurveVerts.SetNum(NV);
+			CurveVerts[0] = PosA; CurveVerts[NV-1] = PosB;
+			for (int32 k = 1; k < (NV - 1); ++k)
+			{
+				int32 VID = ColVerts[k];
+				FVector3d Pos = Mesh.GetVertex(VID);
+				double T = (double)k / (double)(NV-1);
+				FVector3d CurvePos = Curve.Eval((float)T, Lerp(PosA, PosB, T));
+				Mesh.SetVertex(VID, CurvePos);
+				CurveVerts[k] = CurvePos;
+			}
+
+			// accumulate new interior normals (only triangles inside the edge are considered)
+			// (do we need these anywhere??)
+			for (int32 k = 1; k < (NV - 1); ++k)
+			{
+				int32 VID = ColVerts[k];
+				FVector3d CurveEdgeNormal = FMeshNormals::ComputeVertexNormal(Mesh, VID,
+						[&](int32 TriangleID) { return Mesh.GetTriangleGroup(TriangleID) == Edge.NewGroupID; });
+
+				DeformNormals[VID] += CurveEdgeNormal;
+			}
+		}
+	}
+
+	// re-normalize here because we accumulated normals above...
+	for (FVector3d& Normal : DeformNormals)
+	{
+		Normal.Normalize();
+	}
+
+	// Process vertices. This code has special-cases for valence 3 and 4, which should be factored out into separate functions.
+	for (FBevelVertex& Vertex : Vertices)
+	{
+		int32 LoopN = Vertex.InteriorBorderLoop.Num();
+		for (const FBevelVertex_InteriorVertex& InteriorVtx : Vertex.InteriorVertices)
+		{
+			// Special case for valence-3 vertex, ie triangular patch. In this case we use
+			// a triangular spline, the math is based on PN triangulation, where the vertex/normal
+			// pairs at the endpoints are used. However a center-point/normal is also necessary.
+			// This does an OK job but seems to come out too flat in many cases...ultimately it is
+			// not going to even be C0 with the boundary spline curves we already decided on.
+			// Perhaps a better option would be to do it similar to the valence-4 case, where basically
+			// we construct interpolated splines based on the barycentrics and border splines
+			if (InteriorVtx.BorderFrameWeight.Num() == 1 && Vertex.InteriorBorderLoop.Num() == 3)
+			{
+				// Nearly all of this work below is identical for every interior vertex!!! It should be done once and cached. 
+
+				double w = InteriorVtx.BorderFrameWeight[0].X, u = InteriorVtx.BorderFrameWeight[0].Y, v = InteriorVtx.BorderFrameWeight[0].Z;
+
+				int i300 = Vertex.InteriorBorderLoop[0];
+				FVector3d b300 = Mesh.GetVertex(i300);
+				int i030 = Vertex.InteriorBorderLoop[1];
+				FVector3d b030 = Mesh.GetVertex(i030);
+				int i003 = Vertex.InteriorBorderLoop[2];
+				FVector3d b003 = Mesh.GetVertex(i003);
+
+				FVector3d N300 = DeformNormals[i300];
+				FVector3d N030 = DeformNormals[i030];
+				FVector3d N003 = DeformNormals[i003];
+
+				bool bReversedA = false, bReversedB = false, bReversedC = false;
+				FInterpCurveVector* CurveA = FindCurve(i300, i030, bReversedA);
+				FInterpCurveVector* CurveB = FindCurve(i030, i003, bReversedB);
+				FInterpCurveVector* CurveC = FindCurve(i003, i300, bReversedC);
+				if ( !ensure(CurveA != nullptr && CurveB != nullptr && CurveC != nullptr )) continue;
+
+				// PN triangle sides are defined by bezier curves, however FInterpCurve is not actually a Bezier, it's missing a 3
+				// on the tangents in the math. So scaling by 1/3 here corrects for this (...ish?)
+				const double TangentScale = 1.0 / 3.0;
+
+				FVector3d b210 = TangentScale * ((bReversedA) ? -CurveA->Points[1].LeaveTangent : CurveA->Points[0].LeaveTangent);
+				b210 += b300;
+				FVector3d b120 = TangentScale * ((bReversedA) ? CurveA->Points[0].LeaveTangent : -CurveA->Points[1].LeaveTangent);
+				b120 += b030;
+
+				FVector3d b021 = TangentScale * ((bReversedB) ? -CurveB->Points[1].LeaveTangent : CurveB->Points[0].LeaveTangent);
+				b021 += b030;
+				FVector3d b012 = TangentScale * ((bReversedB) ? CurveB->Points[0].LeaveTangent : -CurveB->Points[1].LeaveTangent);
+				b012 += b003;
+
+				FVector3d b102 = TangentScale * ((bReversedC) ? -CurveC->Points[1].LeaveTangent : CurveC->Points[0].LeaveTangent);
+				b102 += b003;
+				FVector3d b201 = TangentScale * ((bReversedC) ? CurveC->Points[0].LeaveTangent : -CurveC->Points[1].LeaveTangent);
+				b201 += b300;
+
+				// this computed midpoint tends to be too flat...
+				FVector3d E = (b210 + b120 + b021 + b012 + b102 + b201) / 6.0;
+				FVector3d V = (b300 + b030 + b003) / 3.0;
+				FVector3d b111 = E + RoundWeight * (E - V) / 2.0;
+
+				FVector InterpPos = b300*w*w*w + b030*u*u*u + b003*v*v*v
+					+ b210*3*w*w*u + b120*3*w*u*u + b201*3*w*w*v
+					+ b021*3*u*u*v + b102*3*w*v*v + b012*3*u*v*v
+					+ b111*6*w*u*v;
+
+				Mesh.SetVertex(InteriorVtx.VertexID, InterpPos);
+
+				continue;
+			}
+
+
+			// this block handles special case of valence-4 vertex that can be done w/ a quad patch interpolating boundary curves
+			// TODO: refactor this to a separate first pass so that we can pull the curve setup out of each iteration
+			if (InteriorVtx.BorderFrameWeight.Num() == 1 && Vertex.InteriorBorderLoop.Num() == 4)
+			{
+				// Nearly all of this work below is identical for every interior vertex!!! It should be done once and cached. 
+
+				// here is the situation: We have a tessellated quad patch, when we tessellated it, it was planar (or
+				// at least the edges were straight). At tessellation time we computed barycentric coordinates for each interior vertex.
+				// Now we have curved the edges of the patch along spline curves (above). So basically we want to compute new interior
+				// positions by blending the four edge curve positions. If we think of the quad as having 2 axes X and Y, we are
+				// going to blend the two X curves along the two Y curves, and then evaluate the blended curve to get the vertex position.
+				// This is messy because (1) we have to look up the curves and (2) each one might have swapped end-vertices with a reversed parameter space.
+				// Most of the code below is about handling that weirdness.
+				// 
+				// Note that since these are (almost) Beziers, the 'curve blending' is strictly about the tangents, as the endpoints are fixed/known
+				//
+				//c01  X2            c11            (diagram corresponds to variable naming below)
+				//   *--------------*    ty=1
+				// Y1|              |Y2
+				//   |  XInterp     |
+				//  A*==============*B
+				//   |              |
+				//   *--------------*    ty=0
+				//c00  X1            c10
+				// 
+				// tx=0           tx=1
+
+				// these are the UV (XY) coords of the vertex inside the initial planar quad
+				double tx = InteriorVtx.BorderFrameWeight[0].X, ty = InteriorVtx.BorderFrameWeight[0].Y;
+
+				// four ordered corner vertices of the initial planar quad
+				int c00 = Vertex.InteriorBorderLoop[0];
+				int c10 = Vertex.InteriorBorderLoop[1];
+				int c01 = Vertex.InteriorBorderLoop[2];
+				int c11 = Vertex.InteriorBorderLoop[3];
+
+				// locate the two "Y" edge curves from their corner vertices, handling case of curve being reversed
+				FInterpCurveVector* CurveY1 = BorderCurves.Find(FIndex2i(c00, c01));
+				bool bReversed1 = false;
+				if (CurveY1 == nullptr)
+				{
+					CurveY1 = BorderCurves.Find(FIndex2i(c01, c00));
+					bReversed1 = true;
+				}
+				if ( ! ensure(CurveY1 != nullptr) ) continue;
+
+				FInterpCurveVector* CurveY2 = BorderCurves.Find(FIndex2i(c10, c11));
+				bool bReversed2 = false;
+				if (CurveY2 == nullptr)
+				{
+					CurveY2 = BorderCurves.Find(FIndex2i(c11, c10));
+					bReversed2 = true;
+				}
+				if (!ensure(CurveY2 != nullptr)) continue;
+
+				// evaluate each edge curve at the Y coord
+				double tA = bReversed1 ? (1.0 - ty) : ty;
+				FVector A = CurveY1->Eval((float)tA, FVector::Zero());		// TODO: should lerp here instead of using ::Zero() obviously...
+				double tB = bReversed2 ? (1.0 - ty) : ty;
+				FVector B = CurveY2->Eval((float)tB, FVector::Zero());		// TODO: should lerp here instead of using ::Zero() obviously...
+
+				// now find the two "X" edge curves from their corner vertices, again handling case of reversed edge
+				FInterpCurveVector* CurveX1 = BorderCurves.Find(FIndex2i(c00, c10));
+				bool bReversedX1 = false;
+				if (CurveX1 == nullptr)
+				{
+					CurveX1 = BorderCurves.Find(FIndex2i(c10, c00));
+					bReversedX1 = true;
+				}
+				if (!ensure(CurveX1 != nullptr)) continue;
+
+				FInterpCurveVector* CurveX2 = BorderCurves.Find(FIndex2i(c01, c11));
+				bool bReversedX2 = false;
+				if (CurveX2 == nullptr)
+				{
+					CurveX2 = BorderCurves.Find(FIndex2i(c11, c01));
+					bReversedX2 = true;
+				}
+				if (!ensure(CurveX2 != nullptr)) continue;
+
+				// extract tangents of X curves, need to blend these to compute interpolated tangents at Y-curve points A/B
+				FVector3d Tangent00 = (bReversedX1) ? CurveX1->Points[1].LeaveTangent : -CurveX1->Points[0].LeaveTangent;
+				FVector3d Tangent10 = (bReversedX1) ? -CurveX1->Points[0].LeaveTangent : CurveX1->Points[1].LeaveTangent;
+				FVector3d Tangent01 = (bReversedX2) ? CurveX2->Points[1].LeaveTangent : -CurveX2->Points[0].LeaveTangent;
+				FVector3d Tangent11 = (bReversedX2) ? -CurveX2->Points[0].LeaveTangent : CurveX2->Points[1].LeaveTangent;
+
+				// should perhaps be based on interpolating tangents along Y splines instead of just lerp? is it equivalent?
+				FVector TangentA = Lerp(Tangent00, Tangent01, ty);
+				FVector TangentB = Lerp(Tangent10, Tangent11, ty);
+
+				// construct interpolated X curve, has endpoints evaluated along Y curves, with X curve endpoint tangents
+				// blended along Y curves
+				FInterpCurveVector InterpolatedXCurve;
+				InterpolatedXCurve.AddPoint(0, A);
+				InterpolatedXCurve.AddPoint(1, B);
+				InterpolatedXCurve.Points[0].ArriveTangent = InterpolatedXCurve.Points[0].LeaveTangent = -TangentA;
+				InterpolatedXCurve.Points[1].ArriveTangent = InterpolatedXCurve.Points[1].LeaveTangent = TangentB;
+				InterpolatedXCurve.Points[0].InterpMode = InterpolatedXCurve.Points[1].InterpMode = EInterpCurveMode::CIM_CurveUser;
+
+				// evaluate this new X curve at the vertex X position
+				FVector InterpPos = InterpolatedXCurve.Eval((float)tx, Lerp(A, B, tx));
+				Mesh.SetVertex(InteriorVtx.VertexID, InterpPos);
+
+				continue;
+			}
+
+			// This should only happen for valence 3 and 4 that we already handled. But if it happens
+			// elsewhere, it means we will just leave that patch flat
+			if (InteriorVtx.BorderFrameWeight.Num() != LoopN) continue;
+
+			// ok the general-case solution is basically to compute the new position for this vertex in the rotated
+			// frame of each boundary-polygon vertex, and blend those positions using the MVC coordinates. This produces
+			// a smooth interpolating surface but there is no continuity at the border, and the 'far' vertices will tend
+			// to exert too much influence
+			FVector3d BlendedPos = FVector3d::Zero();
+			double WeightSum = 0;
+			for (int32 k = 0; k < LoopN; ++k)
+			{
+				int32 BorderVID = Vertex.InteriorBorderLoop[k];
+				FVector3d BorderPos = Mesh.GetVertex(BorderVID);
+				FVector3d BorderFrameX = Mesh.GetVertex(Vertex.InteriorBorderLoop[(k+1) % LoopN]) - 
+					Mesh.GetVertex(Vertex.InteriorBorderLoop[(k-1 + LoopN) % LoopN]);
+				BorderFrameX = Normalized(BorderFrameX);
+				FVector3d BorderFrameN = DeformNormals[BorderVID];
+
+				if (RoundWeight < 0)
+				{
+					FQuaterniond RotQuat(BorderFrameX, -45, true);
+					BorderFrameN = RotQuat * BorderFrameN;
+				}
+
+				FVector3d BorderFrameY = Cross(BorderFrameN, BorderFrameX);
+
+				FVector3d FrameDeltaWeight = InteriorVtx.BorderFrameWeight[k];
+				FVector3d ReconstructedPos = BorderPos + (FrameDeltaWeight.X*BorderFrameX) + (FrameDeltaWeight.Y*BorderFrameY);
+
+				BlendedPos += FrameDeltaWeight.Z * ReconstructedPos;
+				WeightSum += FrameDeltaWeight.Z;
+			}
+			BlendedPos *= (1.0 / WeightSum);
+			Mesh.SetVertex(InteriorVtx.VertexID, BlendedPos);
+		}
+	}
+}
+
+
+
+
+
 void FMeshBevel::ComputeNormals(FDynamicMesh3& Mesh)
 {
 	if (Mesh.HasAttributes() == false)
diff --git a/Engine/Plugins/Runtime/GeometryProcessing/Source/DynamicMesh/Public/Operations/MeshBevel.h b/Engine/Plugins/Runtime/GeometryProcessing/Source/DynamicMesh/Public/Operations/MeshBevel.h
--- a/Engine/Plugins/Runtime/GeometryProcessing/Source/DynamicMesh/Public/Operations/MeshBevel.h
+++ b/Engine/Plugins/Runtime/GeometryProcessing/Source/DynamicMesh/Public/Operations/MeshBevel.h
@@ -3,6 +3,8 @@
 
 #include "DynamicMesh/DynamicMesh3.h"
 #include "Util/ProgressCancel.h"
+#include "Selections/QuadGridPatch.h"
+#include "Math/InterpCurve.h"
 
 namespace UE
 {
@@ -53,6 +55,15 @@ public:
 	/** Distance that bevel edges/vertices are inset from their initial position. Not guaranteed to hold for all vertices, though. */
 	double InsetDistance = 5.0;
 	
+	/** Number of subdivisions inserted in each bevel strip */
+	int32 NumSubdivisions = 0;
+
+	/** 
+	 * "Roundness" of the bevel profile. Ignored if Subdivisions = 0. Default=1 means try to be circular-ish. 
+	 * Higher values pull towards a sharper crease. 0 is flat (ie no profile, a linear chamfer). If negative, bevel profile will be an inverted-arc 
+	 */
+	double RoundWeight = 1.0;
+
 	/** Options for MaterialID assignment on the new triangles generated for the bevel */
 	enum class EMaterialIDMode
 	{
@@ -134,6 +145,7 @@ public:
 		TArray<int32> MeshVertices;		// sequential list of mesh vertex IDs along edge loop
 		TArray<int32> MeshEdges;		// sequential list of mesh edge IDs along edge loop
 		TArray<FIndex2i> MeshEdgeTris;	// the one or two triangles associated w/ each MeshEdges element in the input mesh
+		TArray<FVector3d> InitialPositions;		// initial vertex positions
 
 		// new mesh topology computed during unlink step
 		TArray<int32> NewMeshVertices;		// list of new vertices on "other" side of unlinked edge, 1-1 with MeshVertices
@@ -146,6 +158,8 @@ public:
 		// new geometry computed during mesh step
 		TArray<int32> NewGroupIDs;
 		TArray<FIndex2i> StripQuads;		// triangle-ID-pairs for each new quad added along edge, 1-1 with MeshEdges
+		FQuadGridPatch StripQuadPatch;			// only initialized in multi-segment bevel
+		TArray<FVector3d> NormalsA, NormalsB;	// normals at NewPositions0 and NewPositions1, before internal meshing is added (ie the tangent-boundary condition)
 	};
 
 	// FBevelEdge is the accumulated data for an open span of mesh-edges, which possibly meets up with other bevel-edges
@@ -153,12 +167,15 @@ public:
 	struct FBevelEdge
 	{
 		// initial topological information that defines what happens in unlink/displace/mesh steps
+		int32 EdgeIndex;				// index of this BevelEdge in Edges array
 		TArray<int32> MeshVertices;		// sequential list of mesh vertex IDs along edge
 		TArray<int32> MeshEdges;		// sequential list of mesh edge IDs along edge
 		TArray<FIndex2i> MeshEdgeTris;	// the one or two triangles associated w/ each MeshEdges element in the input mesh
 		int32 GroupEdgeID;				// ID of this edge in external topology (eg FGroupTopology)
 		FIndex2i GroupIDs;				// topological IDs of groups on either side of topological edge
 		bool bEndpointBoundaryFlag[2];	// flag defining whether vertex at start/end of MeshVertices was a boundary vertex
+		TArray<FVector3d> InitialPositions;		// initial vertex positions
+		FIndex2i BevelVertices;			// indices of Bevel Vertices at either end of Bevel Edge
 
 		// new mesh topology computed during unlink step
 		TArray<int32> NewMeshVertices;		// list of new vertices on "other" side of unlinked edge, 1-1 with MeshVertices
@@ -171,6 +188,8 @@ public:
 		// new geometry computed during mesh step
 		int32 NewGroupID;
 		TArray<FIndex2i> StripQuads;		// triangle-ID-pairs for each new quad added along edge, 1-1 with MeshEdges
+		FQuadGridPatch StripQuadPatch;			// only initialized in multi-segment bevel
+		TArray<FVector3d> NormalsA, NormalsB;	// normals at NewPositions0 and NewPositions1, before internal meshing is added (ie the tangent-boundary condition)
 	};
 
 
@@ -202,6 +221,12 @@ public:
 		Unknown
 	};
 
+	struct FBevelVertex_InteriorVertex
+	{
+		int32 VertexID = -1;
+		TArray<FVector3d> BorderFrameWeight;
+	};
+
 	// A FBevelVertex repesents/stores the accumulated data at a "bevel vertex", which is the mesh vertex at the end of a FBevelEdge.
 	// A FBevelVertex may be expanded out into a polygon or just an edge, depending on its Type
 	struct FBevelVertex
@@ -212,6 +237,7 @@ public:
 
 		TArray<int32> IncomingBevelMeshEdges;						// Set of (unsorted) Mesh Edges that are destined to be Beveled, coming into the vertex
 		TArray<int32> IncomingBevelTopoEdges;						// Set of (unsorted) Group Topology Edges that are to be Beveled, coming into the vertex
+		TArray<int32> IncomingBevelEdgeIndices;						// Set of (unsorted) indices of FBevelEdge, coming into the vertex
 
 		TArray<int32> SortedTriangles;			// ordered triangle one-ring around VertexID
 		TArray<FOneRingWedge> Wedges;			// ordered decomposition of one-ring into "wedges" between incoming bevel edges (no correspondence w/ IncomingBevelMeshEdges list)
@@ -221,6 +247,10 @@ public:
 
 		FIndex2i TerminatorInfo;				// for TerminatorVertex type, store [EdgeID, FarVertexID] in one-ring, used to unlink/fill (see usage)
 		int32 ConnectedBevelVertex = -1;		// If set to another FBevelVertex index, then the TerminatorInfo.EdgeID directly connects to that vertex and special handling is needed
+
+		// these arrays are used for multi-segment bevels, and are initialized in different ways depending on the vertex valence - see cpp comments
+		TArray<FBevelVertex_InteriorVertex> InteriorVertices;
+		TArray<int32> InteriorBorderLoop;
 	};
 
 
@@ -239,7 +269,7 @@ public:
 
 protected:
 
-	FBevelVertex* GetBevelVertexFromVertexID(int32 VertexID);
+	FBevelVertex* GetBevelVertexFromVertexID(int32 VertexID, int32* IndexOut = nullptr);
 
 	// Setup phase: register Edges (spans) and (isolated) Loops that need to be beveled and precompute/store any mesh topology that must be tracked across the operation
 	// Required BevelVertex's are added by AddBevelGroupEdge()
@@ -273,6 +303,7 @@ protected:
 	// Meshing phase - append quad-strips between unlinked edge spans/loops, polygons at junction vertices where required,
 	// and triangles at terminator vertices
 
+	/* Meshing functions for chamfer bevel, ie no subdivisions */
 	void CreateBevelMeshing(FDynamicMesh3& Mesh);
 	void AppendJunctionVertexPolygon(FDynamicMesh3& Mesh, FBevelVertex& Vertex);
 	void AppendTerminatorVertexTriangle(FDynamicMesh3& Mesh, FBevelVertex& Vertex);
@@ -280,6 +311,18 @@ protected:
 	void AppendEdgeQuads(FDynamicMesh3& Mesh, FBevelEdge& Edge);
 	void AppendLoopQuads(FDynamicMesh3& Mesh, FBevelLoop& Loop);
 
+	/** Meshing functions for multi-segment bevel with optional round profile */
+	void CreateBevelMeshing_Multi(FDynamicMesh3& Mesh);
+	void AppendEdgeQuads_Multi(FDynamicMesh3& Mesh, FBevelEdge& Edge);
+	void AppendLoopQuads_Multi(FDynamicMesh3& Mesh, FBevelLoop& Loop);
+	void AppendJunctionVertexPolygon_Multi(FDynamicMesh3& Mesh, FBevelVertex& Vertex);
+	void AppendTerminatorVertexTriangles_Multi(FDynamicMesh3& Mesh, FBevelVertex& Vertex);
+	void AppendTerminatorVertexPairQuad_Multi(FDynamicMesh3& Mesh, FBevelVertex& Vertex0, FBevelVertex& Vertex1);
+	void ApplyProfileShape_Round(FDynamicMesh3& Mesh);
+	FInterpCurveVector MakeArcSplineCurve(const FVector3d& PosA, FVector3d& NormalA, const FVector3d& PosB, FVector3d& NormalB) const;
+	void PlanarizeArcNormals(const TArray<FVector3d>& InitialEdgeCurve, bool bIsLoop, int32 Index, FVector3d& NormalA, FVector3d& NormalB, FVector3d& TangentOut, bool& bTangentIsValid) const;
+
+
 	// Normals phase - calculate normals for new geometry
 
 	void ComputeNormals(FDynamicMesh3& Mesh);
diff --git a/Engine/Plugins/Runtime/MeshModelingToolset/Source/MeshModelingTools/Private/ToolActivities/PolyEditBevelEdgeActivity.cpp b/Engine/Plugins/Runtime/MeshModelingToolset/Source/MeshModelingTools/Private/ToolActivities/PolyEditBevelEdgeActivity.cpp
--- a/Engine/Plugins/Runtime/MeshModelingToolset/Source/MeshModelingTools/Private/ToolActivities/PolyEditBevelEdgeActivity.cpp
+++ b/Engine/Plugins/Runtime/MeshModelingToolset/Source/MeshModelingTools/Private/ToolActivities/PolyEditBevelEdgeActivity.cpp
@@ -29,6 +29,10 @@ public:
 	TArray<int32> BevelGroupEdges;
 	TArray<int32> BevelGroupFaces;
 	double BevelDistance = 1.0;
+	int32 Subdivisions = 0;
+	double RoundWeight = 1.0;
+	bool bInferMaterialID = true;
+	int32 SetMaterialID = 0;
 
 	void SetTransform(const FTransformSRT3d& Transform)
 	{
@@ -47,7 +51,10 @@ public:
 		UE::Geometry::FMeshBevel Bevel;
 		Bevel.InsetDistance = BevelDistance;
 		Bevel.MaterialIDMode = FMeshBevel::EMaterialIDMode::InferMaterialID;
-		Bevel.SetConstantMaterialID = 0;
+		Bevel.MaterialIDMode = (bInferMaterialID) ? FMeshBevel::EMaterialIDMode::InferMaterialID_ConstantIfAmbiguous : FMeshBevel::EMaterialIDMode::ConstantMaterialID;
+		Bevel.SetConstantMaterialID = SetMaterialID;
+		Bevel.NumSubdivisions = Subdivisions;
+		Bevel.RoundWeight = RoundWeight;
 		bool bBevelSetupValid = false;
 		if (BevelGroupEdges.Num() > 0)
 		{
@@ -84,6 +91,10 @@ TUniquePtr<FDynamicMeshOperator> UPolyEditBevelEdgeActivity::MakeNewOperator()
 	Op->OriginalMesh = ActivityContext->CurrentMesh;
 	Op->MeshTopology = ActivityContext->CurrentTopology;
 	Op->BevelDistance = BevelProperties->BevelDistance;
+	Op->Subdivisions = BevelProperties->Subdivisions;
+	Op->RoundWeight = BevelProperties->RoundWeight;
+	Op->bInferMaterialID = BevelProperties->bInferMaterialID;
+	Op->SetMaterialID = BevelProperties->SetMaterialID;
 
 	// copy edge selection
 	if (ActiveSelection.SelectedEdgeIDs.Num() > 0)
@@ -126,6 +137,16 @@ void UPolyEditBevelEdgeActivity::Setup(UInteractiveTool* ParentToolIn)
 			ActivityContext->Preview->InvalidateResult();
 		}
 	});
+	BevelProperties->WatchProperty(BevelProperties->Subdivisions, [this](int) {
+		if (bIsRunning) {
+			ActivityContext->Preview->InvalidateResult();
+		}
+	});
+	BevelProperties->WatchProperty(BevelProperties->RoundWeight, [this](float) {
+		if (bIsRunning) {
+			ActivityContext->Preview->InvalidateResult();
+		}
+	});
 }
 
 void UPolyEditBevelEdgeActivity::Shutdown(EToolShutdownType ShutdownType)
diff --git a/Engine/Plugins/Runtime/MeshModelingToolset/Source/MeshModelingTools/Private/ToolActivities/PolyEditBevelEdgeActivity.h b/Engine/Plugins/Runtime/MeshModelingToolset/Source/MeshModelingTools/Private/ToolActivities/PolyEditBevelEdgeActivity.h
--- a/Engine/Plugins/Runtime/MeshModelingToolset/Source/MeshModelingTools/Private/ToolActivities/PolyEditBevelEdgeActivity.h
+++ b/Engine/Plugins/Runtime/MeshModelingToolset/Source/MeshModelingTools/Private/ToolActivities/PolyEditBevelEdgeActivity.h
@@ -25,8 +25,25 @@ class MESHMODELINGTOOLS_API UPolyEditBevelEdgeProperties : public UInteractiveTo
 	GENERATED_BODY()
 
 public:
+	/** Distance that each beveled mesh edge is inset from it's initial position*/
 	UPROPERTY(EditAnywhere, Category = Bevel)
 	double BevelDistance = 4.0;
+
+	/** Number of edge loops added along the bevel faces */
+	UPROPERTY(EditAnywhere, Category = Bevel, meta=(UIMax="10", ClampMin="0", ClampMax="1000"))
+	int Subdivisions = 0;
+
+	/** Roundness of the bevel. Ignored if Subdivisions = 0. */
+	UPROPERTY(EditAnywhere, Category = Bevel, meta = (UIMin = "-2", UIMax="2"))
+	float RoundWeight = 1.0;
+
+	/** If true, when faces on either side of a beveled mesh edges have the same Material ID, beveled edge will be set to that Material ID. Otherwise SetMaterialID is used. */
+	UPROPERTY(EditAnywhere, Category = Bevel)
+	bool bInferMaterialID = true;
+
+	/** Material ID to set on the new faces introduced by bevel operation, unless bInferMaterialID=true and non-ambiguous MaterialID can be inferred from adjacent faces */
+	UPROPERTY(EditAnywhere, Category = Bevel, meta = (ClampMin="0", NoSpinbox))
+	int SetMaterialID = 0;
 };
 
 
