fix Mac-specific build issue

diff --git a/Engine/Plugins/Runtime/GeometryProcessing/Source/DynamicMesh/Private/Operations/MeshBevel.cpp b/Engine/Plugins/Runtime/GeometryProcessing/Source/DynamicMesh/Private/Operations/MeshBevel.cpp
--- a/Engine/Plugins/Runtime/GeometryProcessing/Source/DynamicMesh/Private/Operations/MeshBevel.cpp
+++ b/Engine/Plugins/Runtime/GeometryProcessing/Source/DynamicMesh/Private/Operations/MeshBevel.cpp
@@ -2905,7 +2905,7 @@ void FMeshBevel::ApplyProfileShape_Round(FDynamicMesh3& Mesh)
 			{
 				int32 VID = ColVerts[k];
 				FVector3d CurveEdgeNormal = FMeshNormals::ComputeVertexNormal(Mesh, VID,
-						[&](int32 TriangleID) { return Mesh.GetTriangleGroup(TriangleID) == Edge.NewGroupID; });
+						[&](int32 TriangleID) { return Mesh.GetTriangleGroup(TriangleID) == Edge.NewGroupID; }, true, true);
 
 				DeformNormals[VID] += CurveEdgeNormal;
 			}
