fix various issues with mesh bevel. Boundary-condition normals are now computed before any meshing. DisplaceVertices() now tries to keep terminator vertex positions on existing mesh edges (this should be improved to stay on group edges, in most cases a group edge can be selected topologically, but for now just snapping to mesh edges to handle the simple lowpoly-boxy cases). Simplify/improve computation of section planes in ApplyProfileShape_Round() - in all cases except corners, the section plane is well-defined by the two inset edge points and the original vertex position. This allowed removal of PlanarizeArcNormals function. Also now handling flipped orientation of quad patch relative to edge vertex sequence.

Added more error handling inFQuadGridPatch::InitializeFromQuadPatch. Fixed property invalidating in PolyEdit bevel activity.

diff --git a/Engine/Plugins/Runtime/GeometryProcessing/Source/DynamicMesh/Private/Operations/MeshBevel.cpp b/Engine/Plugins/Runtime/GeometryProcessing/Source/DynamicMesh/Private/Operations/MeshBevel.cpp
--- a/Engine/Plugins/Runtime/GeometryProcessing/Source/DynamicMesh/Private/Operations/MeshBevel.cpp
+++ b/Engine/Plugins/Runtime/GeometryProcessing/Source/DynamicMesh/Private/Operations/MeshBevel.cpp
@@ -405,6 +405,9 @@ void FMeshBevel::BuildVertexSets(const FDynamicMesh3& Mesh)
 		if (Mesh.IsBoundaryVertex(Vertex.VertexID))
 		{
 			Vertex.VertexType = EBevelVertexType::BoundaryVertex;
+			// TODO: we should have a BuildBoundaryVertex function here that correctly populates the 
+			// Wedges for the boundary vertex. The currently BuildJunctionVertex will not be able to do
+			// this because it assumes it can just walk forward from any edge
 			continue;
 		}
 
@@ -1175,58 +1178,95 @@ void FMeshBevel::DisplaceVertices(FDynamicMesh3& Mesh, double Distance)
 
 
 	// Now solve corners. For corners, we want to find the 1 or 2 inset-lines corresponding
-	// to the outgoing bevel-edges at each bevel-vertex-wedge. Unfortunately we do not have
-	// a precomputed mapping for this so we currently linear-search over the full edge set for 
-	// each wedge. Could do in parallel (eg make list of valid wedges first)
+	// to the outgoing bevel-edges at each bevel-vertex-wedge. 
 	for (FBevelVertex& Vertex : Vertices)
 	{
-		if ( (Vertex.VertexType == EBevelVertexType::JunctionVertex)
-			|| (Vertex.VertexType == EBevelVertexType::TerminatorVertex) )
+		if (Vertex.VertexType == EBevelVertexType::Unknown)
 		{
-			int32 NumWedges = Vertex.Wedges.Num();
-			for (int32 k = 0; k < NumWedges; ++k)
+			continue;
+		}
+
+		int32 NumWedges = Vertex.Wedges.Num();
+		for (int32 k = 0; k < NumWedges; ++k)
+		{
+			FOneRingWedge& Wedge = Vertex.Wedges[k];
+			FVector3d CurPos = Mesh.GetVertex(Wedge.WedgeVertex);
+
+			// collect up set of inset lines relevant to this vertex
+			TArray<FLine3d> SolveLines;
+			for (int32 j : Vertex.IncomingBevelEdgeIndices)
 			{
-				FOneRingWedge& Wedge = Vertex.Wedges[k];
+				if (Edges[j].MeshVertices[0] == Wedge.WedgeVertex)
+				{
+					SolveLines.Add(AllInsetLines[j].InsetLines0[0]);
+				}
+				else if (Edges[j].MeshVertices.Last() == Wedge.WedgeVertex)
+				{
+					SolveLines.Add(AllInsetLines[j].InsetLines0.Last());
+				}
+				else if (Edges[j].NewMeshVertices[0] == Wedge.WedgeVertex)
+				{
+					SolveLines.Add(AllInsetLines[j].InsetLines1[0]);
+				}
+				else if (Edges[j].NewMeshVertices.Last() == Wedge.WedgeVertex)
+				{
+					SolveLines.Add(AllInsetLines[j].InsetLines1.Last());
+				}
+			}
 
-				// collect up set of inset lines relevant to this vertex
-				TArray<FLine3d> SolveLines;
-				for (int32 j = 0; j < Edges.Num(); ++j)
+			// todo: BoundaryVertex case never actually gets here because currently we do not initialize wedges of BoundaryVertex!
+			bool bIsSimpleBoundary = (Vertex.VertexType == EBevelVertexType::BoundaryVertex && SolveLines.Num() == 1);
+
+			if (Vertex.VertexType == EBevelVertexType::TerminatorVertex || bIsSimpleBoundary)
+			{
+				if (ensure(SolveLines.Num() == 1))
 				{
-					if (Edges[j].MeshVertices[0] == Wedge.WedgeVertex)
-					{
-						SolveLines.Add(AllInsetLines[j].InsetLines0[0]);
-					}
-					else if (Edges[j].MeshVertices.Last() == Wedge.WedgeVertex)
-					{
-						SolveLines.Add(AllInsetLines[j].InsetLines0.Last());
-					}
-					else if (Edges[j].NewMeshVertices[0] == Wedge.WedgeVertex)
+					// This will be on the inset edge-line but possibly pulled away from the face the incoming terminating edge is 'hitting'
+					// It's fine on right-angles but you can see the problem on the front edge of a cube shaped like:
+					// 
+					//         *---*
+					//          *-*
+					FVector3d InsetLinePosition = SolveLines[0].NearestPoint(CurPos);
+					Wedge.NewPosition = InsetLinePosition;
+
+					// What we ought to do is determine which group topology edges each wedge vertex should 'slide along'. 
+					// However this is a bit complex to figure out and so for the shorter term we are just going to
+					// do a hack by finding the wedge-mesh-edge most aligned w/ the line inset edge.
+					// This will obviously fail if there are sliver triangles in the wedge that it can get confused by...
+
+					FVector3d BaseInsetDir = Normalized(InsetLinePosition - CurPos);
+					double MaxDot = -1; 
+					FLine3d MaxDotEdgeLine;
+					Mesh.EnumerateVertexVertices(Wedge.WedgeVertex, [&](int32 othervid)
 					{
-						SolveLines.Add(AllInsetLines[j].InsetLines1[0]);
-					}
-					else if (Edges[j].NewMeshVertices.Last() == Wedge.WedgeVertex)
+						FLine3d EdgeLine = FLine3d::FromPoints(CurPos, Mesh.GetVertex(othervid));
+						double DirDot = EdgeLine.Direction.Dot(BaseInsetDir);
+						if (DirDot > MaxDot )
+						{
+							MaxDot = DirDot;
+							MaxDotEdgeLine = EdgeLine;
+						}
+					});
+
+					if (MaxDot > -1)
 					{
-						SolveLines.Add(AllInsetLines[j].InsetLines1.Last());
+						FDistLine3Line3d LineIntersection(SolveLines[0], MaxDotEdgeLine);
+						LineIntersection.Get();
+						Wedge.NewPosition = LineIntersection.Line2ClosestPoint;
 					}
-				}
 
-				// now that we have our line set, use it to solve inset position
-				FVector3d CurPos = Mesh.GetVertex(Wedge.WedgeVertex);
-				if (SolveLines.Num() == 1)
-				{
-					Wedge.NewPosition = SolveLines[0].NearestPoint(CurPos);
+					Wedge.bHaveNewPosition = true;
 				}
-				else if (SolveLines.Num() == 2)
+			}
+			else 
+			{
+				if (ensure(SolveLines.Num() >= 2))
 				{
 					Wedge.NewPosition = UE::Geometry::SolveInsetVertexPositionFromLinePair(CurPos, SolveLines[0], SolveLines[1]);
-				}
-				else
-				{
-					// Is this even possible? #SolveLines should equal #BoundaryEdges, how can we have more than 2 at a vertex??
-					// fall back to not-very-good inset technique
-					Wedge.NewPosition = GetDisplacedVertexPos(Mesh, Wedge.WedgeVertex);
+					Wedge.bHaveNewPosition = true;
 				}
 			}
+
 		}
 	}
 
@@ -1258,10 +1298,9 @@ void FMeshBevel::DisplaceVertices(FDynamicMesh3& Mesh, double Distance)
 	}
 	for (FBevelVertex& Vertex : Vertices)
 	{
-		if ( (Vertex.VertexType == EBevelVertexType::JunctionVertex)
-			|| (Vertex.VertexType == EBevelVertexType::TerminatorVertex) )
+		for (FOneRingWedge& Wedge : Vertex.Wedges)
 		{
-			for (FOneRingWedge& Wedge : Vertex.Wedges)
+			if (Wedge.bHaveNewPosition)
 			{
 				Mesh.SetVertex(Wedge.WedgeVertex, Wedge.NewPosition);
 			}
@@ -1631,18 +1670,6 @@ void FMeshBevel::AppendEdgeQuads_Multi(FDynamicMesh3& Mesh, FBevelEdge& Edge)
 		return;
 	}
 
-	// Will need the normals for each exterior "side" of the bevel edge along the vertices path to
-	// define the arcs along the bevel edge. It ought to be possible to compute this after
-	// adding the bevel geometry, by filtering out the bevel tris, but for now it's simpler to just 
-	// do it here before we add the bevel tris and cache it...
-	Edge.NormalsA.SetNum(Edge.MeshVertices.Num());
-	Edge.NormalsB.SetNum(Edge.MeshVertices.Num());
-	for (int32 k = 0; k < Edge.MeshVertices.Num(); ++k)
-	{
-		Edge.NormalsA[k] = FMeshNormals::ComputeVertexNormal(Mesh, Edge.MeshVertices[k]);
-		Edge.NormalsB[k] = FMeshNormals::ComputeVertexNormal(Mesh, Edge.NewMeshVertices[k]);
-	}
-
 	// all the code below is going to generate the bevel edge geometry and populate this list of 
 	// vertex-rows, which will be used to assemble the final FQuadGridPatch for the bevel edge-strip
 	int32 N = NumSubdivisions;
@@ -1859,14 +1886,6 @@ void FMeshBevel::AppendLoopQuads_Multi(FDynamicMesh3& Mesh, FBevelLoop& Loop)
 		return;
 	}
 
-	Loop.NormalsA.SetNum(Loop.MeshVertices.Num());
-	Loop.NormalsB.SetNum(Loop.MeshVertices.Num());
-	for (int32 k = 0; k < Loop.MeshVertices.Num(); ++k)
-	{
-		Loop.NormalsA[k] = FMeshNormals::ComputeVertexNormal(Mesh, Loop.MeshVertices[k]);
-		Loop.NormalsB[k] = FMeshNormals::ComputeVertexNormal(Mesh, Loop.NewMeshVertices[k]);
-	}
-
 	int32 N = NumSubdivisions;
 	TArray<TArray<int32>> VertexSpans;
 	VertexSpans.SetNum(N + 2);
@@ -2589,6 +2608,31 @@ void FMeshBevel::CreateBevelMeshing_Multi(FDynamicMesh3& Mesh)
 	// First we figure out the mesh connectivity, ie the 'stitching' between the pulled-apart geometry,
 	// and then we (optionally) apply a profile-curve shape along the bevel strips
 
+	// We will later need normals along each exterior "side" of the bevel edge to define the arcs along 
+	// rounded bevel edges, ie basically these are the smooth boundary conditions.
+	// It ought to be possible to compute this after adding the bevel geometry, by filtering out the bevel-edge tris, 
+	// but for now it's simpler to just do it here before we add the bevel tris and cache it...
+	for (FBevelEdge& Edge : Edges)
+	{
+		Edge.NormalsA.SetNum(Edge.MeshVertices.Num());
+		Edge.NormalsB.SetNum(Edge.MeshVertices.Num());
+		for (int32 k = 0; k < Edge.MeshVertices.Num(); ++k)
+		{
+			Edge.NormalsA[k] = FMeshNormals::ComputeVertexNormal(Mesh, Edge.MeshVertices[k]);
+			Edge.NormalsB[k] = FMeshNormals::ComputeVertexNormal(Mesh, Edge.NewMeshVertices[k]);
+		}
+	}
+	for (FBevelLoop& Loop : Loops)
+	{
+		Loop.NormalsA.SetNum(Loop.MeshVertices.Num());
+		Loop.NormalsB.SetNum(Loop.MeshVertices.Num());
+		for (int32 k = 0; k < Loop.MeshVertices.Num(); ++k)
+		{
+			Loop.NormalsA[k] = FMeshNormals::ComputeVertexNormal(Mesh, Loop.MeshVertices[k]);
+			Loop.NormalsB[k] = FMeshNormals::ComputeVertexNormal(Mesh, Loop.NewMeshVertices[k]);
+		}
+	}
+
 	for (FBevelEdge& Edge : Edges)
 	{
 		AppendEdgeQuads_Multi(Mesh, Edge);
@@ -2701,70 +2745,6 @@ FInterpCurveVector FMeshBevel::MakeArcSplineCurve(const FVector3d& PosA, FVector
 }
 
 
-void FMeshBevel::PlanarizeArcNormals(const TArray<FVector3d>& InitialEdgeCurve, bool bIsLoop, int32 Index, FVector3d& NormalA, FVector3d& NormalB, FVector3d& TangentOut, bool& bTangentIsValid) const
-{
-	TangentOut = FVector3d::Zero();
-
-	// We have two Normals A and B, which are destined to be used with MakeArcSplineCurve, ie they are the surface normals
-	// at the endpoints of the curve we want to generate, and we are going to use those normals to define the curve.
-	// In nearly all cases, we want this curve to be planar, so the normals need to lie in the same plane.
-	//
-	// To achieve that, we are going to assume that the tangent direction around the InitialEdgeCurve (actually a polyline)
-	// defines the desired plane. This makes sense if you think of the bevel as essentially sweeping a circular arc around
-	// the bevel paths and using it to cut away from the surface. Then the cut plane at any point on the curve would be normal
-	// to the curve tangent. So we find that tangent direction, and then project both surface normals onto it.
-	//
-	// Note that we use a purely opening-angle-based polyline-tangent estimate here. This makes less sense for a smooth curve path,
-	// but for (eg) a rectangular box, if we use the polyine-tangent definition that is an approximation of a continuous tangent, then
-	// the cut-angle at the corners will vary with the box dimensions, which is clearly wrong. 
-
-	int32 NV = InitialEdgeCurve.Num();
-	if (bIsLoop || (Index != 0 && Index != NV-1) )
-	{
-		FVector3d CurPos = InitialEdgeCurve[Index];
-		FVector3d PrevDir = InitialEdgeCurve[(Index-1 + NV) % NV] - CurPos;
-		FVector3d NextDir = InitialEdgeCurve[(Index+1) % NV] - CurPos;
-		if (NextDir.Normalize() && PrevDir.Normalize())
-		{
-			// this direction is the angle-based tangent, ie ignores edge length
-			TangentOut = NextDir - PrevDir;
-		}
-	}
-	else
-	{
-		FVector3d CurPos = InitialEdgeCurve[Index];
-		if (Index == 0)
-		{
-			FVector3d EdgeDir = InitialEdgeCurve[1] - CurPos;
-			if (EdgeDir.Normalize())
-			{
-				TangentOut = EdgeDir;
-			}
-		}
-		else if (Index == NV - 1)
-		{
-			FVector3d EdgeDir = CurPos - InitialEdgeCurve[NV-2];
-			if (EdgeDir.Normalize())
-			{
-				TangentOut = EdgeDir;
-			}
-		}
-		else
-		{
-			check(false);
-		}
-	}
-
-	if (TangentOut.Normalize())
-	{
-		// project normals onto plane of tangent
-		NormalA = Normalized(NormalA - NormalA.Dot(TangentOut) * TangentOut);
-		NormalB = Normalized(NormalB - NormalB.Dot(TangentOut) * TangentOut);
-		bTangentIsValid = true;
-	}
-}
-
-
 void FMeshBevel::ApplyProfileShape_Round(FDynamicMesh3& Mesh)
 {
 	// This function applies round profile curves to the bevel edge quadstrips & 
@@ -2794,8 +2774,14 @@ void FMeshBevel::ApplyProfileShape_Round(FDynamicMesh3& Mesh)
 			FVector3d PosA = Mesh.GetVertex(A), PosB = Mesh.GetVertex(B);
 			FVector3d NormalA = Loop.NormalsA[Col], NormalB = Loop.NormalsB[Col];
 
-			FVector3d Tangent; bool bTangentIsValid = false;
-			PlanarizeArcNormals(Loop.InitialPositions, true, Col, NormalA, NormalB, Tangent, bTangentIsValid);
+			// project normals onto section plane defined by original and inset vertex positions.
+			// (do we even need the normals anymore? could they be defined by the 2D perp operator in the section plane?)
+			FVector3d InitialPosition = Loop.InitialPositions[Col];
+			FVector3d Direction1 = Normalized(PosA - InitialPosition);
+			FVector3d Direction2 = Normalized(PosB - InitialPosition);
+			FVector3d PlaneNormal = Normalized(Cross(Direction1, Direction2));
+			NormalA = Normalized(NormalA - NormalA.Dot(PlaneNormal) * PlaneNormal);
+			NormalB = Normalized(NormalB - NormalB.Dot(PlaneNormal) * PlaneNormal);
 
 			FInterpCurveVector Curve = MakeArcSplineCurve(PosA, NormalA, PosB, NormalB);
 
@@ -2844,36 +2830,78 @@ void FMeshBevel::ApplyProfileShape_Round(FDynamicMesh3& Mesh)
 	{
 		FQuadGridPatch& Patch = Edge.StripQuadPatch;
 
+		bool bPatchIsFlippedX = false;
+
 		int32 NumCols = Patch.NumVertexCols();
 		for (int32 Col = 0; Col < NumCols; ++Col)
 		{
 			TArray<int32> ColVerts;
 			Patch.GetVertexColumn(Col, ColVerts);
-			int32 NV = ColVerts.Num();
+			int32 NumColVerts = ColVerts.Num();
+
+			// The Edge contains various vertex lists like Edge.MeshVertices that have a consistent ordering.
+			// However because of how the patches are constructed, the columns may be reversed  (unclear why...)
+			// Detect that case so that the indexing Col/Array indexing can be inverted where necessary
+			if (Col == 0 && ColVerts.Contains(Edge.MeshVertices[0]) == false)
+			{
+				bPatchIsFlippedX = true;
+			}
 
 			int32 A = ColVerts[0];
 			int32 B = ColVerts.Last();
 			FVector3d PosA = Mesh.GetVertex(A);
 			FVector3d PosB = Mesh.GetVertex(B);
-			FVector3d NormalA = Edge.NormalsA[Col];
-			FVector3d NormalB = Edge.NormalsB[Col];
 
-			bool bProjectToTangentPlane = true;
-			bool bIsEndpoint = (Col == 0) || (Col == NumCols - 1);
-			int32 BevelEdgeIndex = (Col == 0) ? Edge.BevelVertices.A : Edge.BevelVertices.B;
-			check(BevelEdgeIndex >= 0);
-			const FBevelVertex& BevelVtx = Vertices[BevelEdgeIndex];
-			if (BevelVtx.VertexType == EBevelVertexType::JunctionVertex && BevelVtx.IncomingBevelEdgeIndices.Num() == 2)
-			{
-				bProjectToTangentPlane = false;
+			int UseEdgeVertIndex = (bPatchIsFlippedX) ? (NumCols-Col-1) : Col;
+			FVector3d NormalA = Edge.NormalsA[UseEdgeVertIndex];
+			FVector3d NormalB = Edge.NormalsB[UseEdgeVertIndex];
+
+			// Ok for each column along a bevel edge-strip, we want to bend the column into a curve.
+			// We have the endpoints we want to interpolate, and we have normals at those points that we
+			// computed elsewhere and are going to assume are good, ie those are the normals we want the
+			// perfect rounded bevel to have, if it had infinite sections.
+			//
+			// The issue is that we have 2 points and 2 normals and they may not all lie in the same plane.
+			// For most cases, the original vertex position has been inset in two directions (to PosA and PosB), and 
+			// so that gives us the plane they all should lie in. However at valence3+ junction vertices this is not true,
+			// as the one corner vertex was inset along 3+ edges and so we don't have the 'InitialPosition'
+			// value below for each of those edges, that we would need to define the simple arc-sections.
+			// 
+			// (todo: maybe keeping track of those positions would be better than what we do now
+
+			bool bInferTangentPlaneFromInitialPosition = true;
+			bool bIsEndpoint = (UseEdgeVertIndex == 0) || (UseEdgeVertIndex == NumCols - 1);
+			if ( bIsEndpoint )
+			{
+				int32 BevelVertexIndex = (UseEdgeVertIndex == 0) ? Edge.BevelVertices.A : Edge.BevelVertices.B;
+				const FBevelVertex& BevelVtx = Vertices[BevelVertexIndex];
+				if (BevelVtx.VertexType == EBevelVertexType::JunctionVertex && BevelVtx.IncomingBevelEdgeIndices.Num() > 2)
+				{
+					bInferTangentPlaneFromInitialPosition = false;
+				}
 			}
 
-			FVector3d Tangent; bool bTangentIsValid = false;
-			if (bProjectToTangentPlane)
+			FVector3d InitialPosition = Edge.InitialPositions[UseEdgeVertIndex];
+			FVector3d Direction1 = Normalized(PosA - InitialPosition);
+			FVector3d Direction2 = Normalized(PosB - InitialPosition);
+			FVector3d SectionPlaneNormal = Normalized(Cross(Direction1, Direction2));
+
+			if (!bInferTangentPlaneFromInitialPosition)
 			{
-				PlanarizeArcNormals(Edge.InitialPositions, false, Col, NormalA, NormalB, Tangent, bTangentIsValid);
+				FVector3d InitialEdgeDirection = (UseEdgeVertIndex == 0) ?
+					(Edge.InitialPositions[1] - InitialPosition) : (InitialPosition - Edge.InitialPositions[NumCols-2]);
+				if (InitialEdgeDirection.Normalize())
+				{
+					FFrame3d TempFrame(PosA, Normalized(PosB - PosA));
+					TempFrame.ConstrainedAlignAxis(1, InitialEdgeDirection, TempFrame.Z());
+					SectionPlaneNormal = TempFrame.Y();
+				}
 			}
 
+			// project normals onto section plane
+			NormalA = Normalized(NormalA - NormalA.Dot(SectionPlaneNormal) * SectionPlaneNormal);
+			NormalB = Normalized(NormalB - NormalB.Dot(SectionPlaneNormal) * SectionPlaneNormal);
+
 			// TODO: should have special case here for flat patches? they get squished out a bit...
 			// This may update NormalA/NormalB for 'inverted' bevels (ie w/ negative RoundWeight)
 			FInterpCurveVector Curve = MakeArcSplineCurve(PosA, NormalA, PosB, NormalB);
@@ -2887,13 +2915,13 @@ void FMeshBevel::ApplyProfileShape_Round(FDynamicMesh3& Mesh)
 
 			// map linear-interpolation to curve parameter and evaluate curve
 			TArray<FVector3d> CurveVerts;
-			CurveVerts.SetNum(NV);
-			CurveVerts[0] = PosA; CurveVerts[NV-1] = PosB;
-			for (int32 k = 1; k < (NV - 1); ++k)
+			CurveVerts.SetNum(NumColVerts);
+			CurveVerts[0] = PosA; CurveVerts[NumColVerts-1] = PosB;
+			for (int32 k = 1; k < (NumColVerts - 1); ++k)
 			{
 				int32 VID = ColVerts[k];
 				FVector3d Pos = Mesh.GetVertex(VID);
-				double T = (double)k / (double)(NV-1);
+				double T = (double)k / (double)(NumColVerts-1);
 				FVector3d CurvePos = Curve.Eval((float)T, Lerp(PosA, PosB, T));
 				Mesh.SetVertex(VID, CurvePos);
 				CurveVerts[k] = CurvePos;
@@ -2901,7 +2929,7 @@ void FMeshBevel::ApplyProfileShape_Round(FDynamicMesh3& Mesh)
 
 			// accumulate new interior normals (only triangles inside the edge are considered)
 			// (do we need these anywhere??)
-			for (int32 k = 1; k < (NV - 1); ++k)
+			for (int32 k = 1; k < (NumColVerts - 1); ++k)
 			{
 				int32 VID = ColVerts[k];
 				FVector3d CurveEdgeNormal = FMeshNormals::ComputeVertexNormal(Mesh, VID,
diff --git a/Engine/Plugins/Runtime/GeometryProcessing/Source/DynamicMesh/Private/Selections/QuadGridPatch.cpp b/Engine/Plugins/Runtime/GeometryProcessing/Source/DynamicMesh/Private/Selections/QuadGridPatch.cpp
--- a/Engine/Plugins/Runtime/GeometryProcessing/Source/DynamicMesh/Private/Selections/QuadGridPatch.cpp
+++ b/Engine/Plugins/Runtime/GeometryProcessing/Source/DynamicMesh/Private/Selections/QuadGridPatch.cpp
@@ -129,6 +129,11 @@ bool FQuadGridPatch::InitializeFromQuadPatch(const FDynamicMesh3& Mesh, const TA
 
 			// figure out unique vertices from the quad info, and verify that there are 4 and they are A/B/C/D
 			FIndex2i& QuadTris = QuadTriangles[j][k];
+			if (Mesh.IsTriangle(QuadTris.A) == false || Mesh.IsTriangle(QuadTris.B) == false)
+			{
+				bAllOK = false;
+				break;
+			}
 			FIndex3i TriA = Mesh.GetTriangle(QuadTris.A);
 			FIndex3i TriB = Mesh.GetTriangle(QuadTris.B);
 			TArray<int32, TInlineAllocator<6>> TriVerts({ TriA.A, TriA.B, TriA.C });
diff --git a/Engine/Plugins/Runtime/GeometryProcessing/Source/DynamicMesh/Public/Operations/MeshBevel.h b/Engine/Plugins/Runtime/GeometryProcessing/Source/DynamicMesh/Public/Operations/MeshBevel.h
--- a/Engine/Plugins/Runtime/GeometryProcessing/Source/DynamicMesh/Public/Operations/MeshBevel.h
+++ b/Engine/Plugins/Runtime/GeometryProcessing/Source/DynamicMesh/Public/Operations/MeshBevel.h
@@ -204,6 +204,7 @@ public:
 		int32 WedgeVertex;						// central vertex of this wedge (updated by unlink functions)
 
 		FVector3d NewPosition;					// new calculated position for vertex of this wedge
+		bool bHaveNewPosition = false;			// flag indicating if NewPosition is valid
 	};
 
 	// a FBevelVertex can have various types, depending on the topology of the bevel edge graph and input mesh
@@ -320,7 +321,6 @@ protected:
 	void AppendTerminatorVertexPairQuad_Multi(FDynamicMesh3& Mesh, FBevelVertex& Vertex0, FBevelVertex& Vertex1);
 	void ApplyProfileShape_Round(FDynamicMesh3& Mesh);
 	FInterpCurveVector MakeArcSplineCurve(const FVector3d& PosA, FVector3d& NormalA, const FVector3d& PosB, FVector3d& NormalB) const;
-	void PlanarizeArcNormals(const TArray<FVector3d>& InitialEdgeCurve, bool bIsLoop, int32 Index, FVector3d& NormalA, FVector3d& NormalB, FVector3d& TangentOut, bool& bTangentIsValid) const;
 
 
 	// Normals phase - calculate normals for new geometry
diff --git a/Engine/Plugins/Runtime/MeshModelingToolset/Source/MeshModelingTools/Private/ToolActivities/PolyEditBevelEdgeActivity.cpp b/Engine/Plugins/Runtime/MeshModelingToolset/Source/MeshModelingTools/Private/ToolActivities/PolyEditBevelEdgeActivity.cpp
--- a/Engine/Plugins/Runtime/MeshModelingToolset/Source/MeshModelingTools/Private/ToolActivities/PolyEditBevelEdgeActivity.cpp
+++ b/Engine/Plugins/Runtime/MeshModelingToolset/Source/MeshModelingTools/Private/ToolActivities/PolyEditBevelEdgeActivity.cpp
@@ -147,6 +147,16 @@ void UPolyEditBevelEdgeActivity::Setup(UInteractiveTool* ParentToolIn)
 			ActivityContext->Preview->InvalidateResult();
 		}
 	});
+	BevelProperties->WatchProperty(BevelProperties->bInferMaterialID, [this](bool) {
+		if (bIsRunning) {
+			ActivityContext->Preview->InvalidateResult();
+		}
+	});
+	BevelProperties->WatchProperty(BevelProperties->SetMaterialID, [this](int) {
+		if (bIsRunning) {
+			ActivityContext->Preview->InvalidateResult();
+		}
+	});
 }
 
 void UPolyEditBevelEdgeActivity::Shutdown(EToolShutdownType ShutdownType)
