Make SimpleShapeSet3's sphere transform method match the editor's approach for transforming simple collision spheres -- scaling the radius by the minimum abs-value axis scale -- to get more consistent results when baking component transforms.

diff --git a/Engine/Plugins/Runtime/GeometryProcessing/Source/DynamicMesh/Private/ShapeApproximation/SimpleShapeSet3.cpp b/Engine/Plugins/Runtime/GeometryProcessing/Source/DynamicMesh/Private/ShapeApproximation/SimpleShapeSet3.cpp
--- a/Engine/Plugins/Runtime/GeometryProcessing/Source/DynamicMesh/Private/ShapeApproximation/SimpleShapeSet3.cpp
+++ b/Engine/Plugins/Runtime/GeometryProcessing/Source/DynamicMesh/Private/ShapeApproximation/SimpleShapeSet3.cpp
@@ -391,7 +391,7 @@ void FSimpleShapeSet3d::RemoveContainedGeometry()
 template<typename TransformType>
 static void TransformSphereShape(FSphereShape3d& SphereShape, const TransformType& Transform)
 {
-	double RadiusScale = Transform.GetScale3D().Length() / FVector3d::One().Length();
+	double RadiusScale = Transform.GetScale3D().GetAbsMin(); // Scale radius by the minimum axis scale, to match how the editor applies non-uniform scales to spheres
 	SphereShape.Sphere.Center = Transform.TransformPosition(SphereShape.Sphere.Center);
 	SphereShape.Sphere.Radius *= RadiusScale;
 }
@@ -443,7 +443,7 @@ static void TransformSphereShapeByArray(FSphereShape3d& SphereShape, const Trans
 	for (const auto& XForm : TransformSequence)
 	{
 		SphereShape.Sphere.Center = XForm.TransformPosition(SphereShape.Sphere.Center);
-		double RadiusScale = XForm.GetScale3D().Length() / FVector3d::One().Length();
+		double RadiusScale = XForm.GetScale3D().GetAbsMin(); // Scale radius by the minimum axis scale, to match how the editor applies non-uniform scales to spheres
 		SphereShape.Sphere.Radius *= RadiusScale;
 	}
 }
