Mesh Bevel: remove ensures that are getting hit in some cases in Lyra. Failure is not catastrophic, so we do not really need to ensure here.

diff --git a/Engine/Plugins/Runtime/GeometryProcessing/Source/DynamicMesh/Private/Operations/MeshBevel.cpp b/Engine/Plugins/Runtime/GeometryProcessing/Source/DynamicMesh/Private/Operations/MeshBevel.cpp
--- a/Engine/Plugins/Runtime/GeometryProcessing/Source/DynamicMesh/Private/Operations/MeshBevel.cpp
+++ b/Engine/Plugins/Runtime/GeometryProcessing/Source/DynamicMesh/Private/Operations/MeshBevel.cpp
@@ -1219,7 +1219,9 @@ void FMeshBevel::DisplaceVertices(FDynamicMesh3& Mesh, double Distance)
 
 			if (Vertex.VertexType == EBevelVertexType::TerminatorVertex || bIsSimpleBoundary)
 			{
-				if (ensure(SolveLines.Num() == 1))
+				// TODO: this ensure has been sporadically hit in Lyra. Needs to be investigated.
+				//ensure(SolveLines.Num() == 1);
+				if (SolveLines.Num() == 1)
 				{
 					// This will be on the inset edge-line but possibly pulled away from the face the incoming terminating edge is 'hitting'
 					// It's fine on right-angles but you can see the problem on the front edge of a cube shaped like:
@@ -1260,7 +1262,12 @@ void FMeshBevel::DisplaceVertices(FDynamicMesh3& Mesh, double Distance)
 			}
 			else 
 			{
-				if (ensure(SolveLines.Num() >= 2))
+				// TODO: this ensure has been sporadically hit in Lyra, the case appears to be that a wedge ends up with
+				// only a single SolveLine. Likely this happens because some bevel-edge was not processed as we expected,
+				// ie the wedge-building went wrong due to a topological case we have not properly handled. Needs to be
+				// investigated more deeply, but in the meantime, skipping here just means the vertex will end up mis-positioned
+				//ensure(SolveLines.Num() >= 2);
+				if (SolveLines.Num() >= 2)
 				{
 					Wedge.NewPosition = UE::Geometry::SolveInsetVertexPositionFromLinePair(CurPos, SolveLines[0], SolveLines[1]);
 					Wedge.bHaveNewPosition = true;
