Fix crash on smooth bevel of some edges

diff --git a/Engine/Plugins/Runtime/GeometryProcessing/Source/DynamicMesh/Private/Operations/MeshBevel.cpp b/Engine/Plugins/Runtime/GeometryProcessing/Source/DynamicMesh/Private/Operations/MeshBevel.cpp
--- a/Engine/Plugins/Runtime/GeometryProcessing/Source/DynamicMesh/Private/Operations/MeshBevel.cpp
+++ b/Engine/Plugins/Runtime/GeometryProcessing/Source/DynamicMesh/Private/Operations/MeshBevel.cpp
@@ -1752,9 +1752,16 @@ void FMeshBevel::AppendEdgeQuads_Multi(FDynamicMesh3& Mesh, FBevelEdge& Edge)
 		return (const FQuadGridPatch*)nullptr;
 	};
 
+	// Below code expects SequentialQuadEdges to be ordered s.t. for edges spanning 
+	// vertices 0,1,2 the edges should be (1,0),(2,1) not (0,1),(1,2) ...
+	// We enforce this by reversing the array if needed
+	if (SequentialQuadEdges.Num() > 1 && SequentialQuadEdges[0].EdgeV0.A == Edge.MeshVertices[0])
+	{
+		Algo::Reverse(SequentialQuadEdges);
+	}
 
-	// somehow SequentialQuadEdges ordering may not be in agreement with the [BevelVertices.A, BevelVertices.B] 
-	// ordering of the BevelEdge....unclear how or why this would be the case. But if so, then the Prev/Next QuadGridPatch
+	// SequentialQuadEdges ordering may not be in agreement with the [BevelVertices.A, BevelVertices.B] 
+	// ordering of the BevelEdge. If so, then the Prev/Next QuadGridPatch
 	// search below would return the Prev & Next flipped from where we need them, causing the later AppendExistingVertColumn
 	// to fail. The simplest fix for this here is to swap PrevQuadPatch/NextQuadPatch...
 	int OriginalPrevVtxID = (Edge.BevelVertices.A >= 0) ? Vertices[Edge.BevelVertices.A].VertexID : -1;
