fix smooth bevel of loops for cases where the loop's edges are reversed from the expected order

diff --git a/Engine/Plugins/Runtime/GeometryProcessing/Source/DynamicMesh/Private/Operations/MeshBevel.cpp b/Engine/Plugins/Runtime/GeometryProcessing/Source/DynamicMesh/Private/Operations/MeshBevel.cpp
--- a/Engine/Plugins/Runtime/GeometryProcessing/Source/DynamicMesh/Private/Operations/MeshBevel.cpp
+++ b/Engine/Plugins/Runtime/GeometryProcessing/Source/DynamicMesh/Private/Operations/MeshBevel.cpp
@@ -1755,7 +1755,7 @@ void FMeshBevel::AppendEdgeQuads_Multi(FDynamicMesh3& Mesh, FBevelEdge& Edge)
 	// Below code expects SequentialQuadEdges to be ordered s.t. for edges spanning 
 	// vertices 0,1,2 the edges should be (1,0),(2,1) not (0,1),(1,2) ...
 	// We enforce this by reversing the array if needed
-	if (SequentialQuadEdges.Num() > 1 && SequentialQuadEdges[0].EdgeV0.A == Edge.MeshVertices[0])
+	if (SequentialQuadEdges.Num() > 1 && SequentialQuadEdges[0].EdgeV0.B == SequentialQuadEdges[1].EdgeV0.A)
 	{
 		Algo::Reverse(SequentialQuadEdges);
 	}
@@ -1900,6 +1900,10 @@ void FMeshBevel::AppendLoopQuads_Multi(FDynamicMesh3& Mesh, FBevelLoop& Loop)
 		return;
 	}
 
+	// Test the edge ordering -- i.e. for vertices 0,1,2, if bEdgeHasReverseOrder == true, 
+	// edge vertices will be (0,1),(1,2); otherwise it will be (1,0),(2,1)
+	bool bEdgeHasReverseOrder = (SequentialQuadEdges.Num() > 1 && SequentialQuadEdges[0].EdgeV0.B == SequentialQuadEdges[1].EdgeV0.A);
+
 	int32 N = NumSubdivisions;
 	TArray<TArray<int32>> VertexSpans;
 	VertexSpans.SetNum(N + 2);
@@ -1923,6 +1927,12 @@ void FMeshBevel::AppendLoopQuads_Multi(FDynamicMesh3& Mesh, FBevelLoop& Loop)
 		FIndex2i EdgeV0 = SequentialQuadEdges[k].EdgeV0;
 		FIndex2i EdgeV1 = SequentialQuadEdges[k].EdgeV1;
 
+		if (bEdgeHasReverseOrder)
+		{
+			EdgeV0.Swap();
+			EdgeV1.Swap();
+		}
+
 		int32 QuadA = EdgeV0.B;
 		int32 QuadB = EdgeV0.A;
 		int32 QuadC = EdgeV1.B;
@@ -1951,6 +1961,10 @@ void FMeshBevel::AppendLoopQuads_Multi(FDynamicMesh3& Mesh, FBevelLoop& Loop)
 	TArray<TArray<FIndex2i>> QuadSpans;
 	QuadSpans.SetNum(N + 1);
 
+	// Offsets for flipping triangles in cases where the edges had the opposite ordering from our expectation
+	int32 SwapOffset0 = (int32)(bEdgeHasReverseOrder);
+	int32 SwapOffset1 = (int32)(!bEdgeHasReverseOrder);
+
 	int32 NumStrips = VertexSpans.Num() - 1;
 	for (int32 k = 0; k < NumEdges; ++k)
 	{
@@ -1960,8 +1974,8 @@ void FMeshBevel::AppendLoopQuads_Multi(FDynamicMesh3& Mesh, FBevelLoop& Loop)
 		for (int32 j = 0; j < NumStrips; ++j)
 		{
 			FIndex2i QuadTris(IndexConstants::InvalidID, IndexConstants::InvalidID);
-			QuadTris.A = Mesh.AppendTriangle(VertexSpans[j][k], VertexSpans[j][k+1], VertexSpans[j+1][k], NewGroupID);
-			QuadTris.B = Mesh.AppendTriangle(VertexSpans[j+1][k+1], VertexSpans[j+1][k], VertexSpans[j][k+1], NewGroupID);
+			QuadTris.A = Mesh.AppendTriangle(VertexSpans[j][k + SwapOffset0], VertexSpans[j][k + SwapOffset1], VertexSpans[j+1][k], NewGroupID);
+			QuadTris.B = Mesh.AppendTriangle(VertexSpans[j+1][k + SwapOffset1], VertexSpans[j+1][k + SwapOffset0], VertexSpans[j][k+1], NewGroupID);
 			QuadSpans[j].Add(QuadTris);
 
 			Loop.StripQuads.Add(QuadTris);
