fix unreachable code warning in meshlinearization.h (change from an iteration+break to directly testing whether vertex has edges)

#changelist validated
#virtualized

diff --git a/Engine/Plugins/Runtime/GeometryProcessing/Source/DynamicMesh/Public/Solvers/MeshLinearization.h b/Engine/Plugins/Runtime/GeometryProcessing/Source/DynamicMesh/Public/Solvers/MeshLinearization.h
--- a/Engine/Plugins/Runtime/GeometryProcessing/Source/DynamicMesh/Public/Solvers/MeshLinearization.h
+++ b/Engine/Plugins/Runtime/GeometryProcessing/Source/DynamicMesh/Public/Solvers/MeshLinearization.h
@@ -66,15 +66,9 @@ private:
 		for (int32 i = 0, I = ToIdMap.Num(); i < I; ++i)
 		{
 			int32 VtxId = ToIdMap[i];
+
 			// Test if the vertex has a one ring
-			
-			bool bEmptyOneRing = true;
-			for (int NeighborVertId : DynamicMesh.VtxVerticesItr(VtxId))
-			{
-				bEmptyOneRing = false;
-				break;
-			};
-				
+			bool bEmptyOneRing = DynamicMesh.GetVtxEdgeCount(VtxId) < 1;
 
 			if (bEmptyOneRing || DynamicMesh.IsBoundaryVertex(VtxId))
 			{
