use a prime value greater than max unsigned int for the modulo iteration in remesher, mesh simplification, and FModuloIteration to ensure that it is always coprime with the number of elements being iterated over

#changelist validated
#virtualized

diff --git a/Engine/Plugins/Runtime/GeometryProcessing/Source/DynamicMesh/Public/MeshSimplification.h b/Engine/Plugins/Runtime/GeometryProcessing/Source/DynamicMesh/Public/MeshSimplification.h
--- a/Engine/Plugins/Runtime/GeometryProcessing/Source/DynamicMesh/Public/MeshSimplification.h
+++ b/Engine/Plugins/Runtime/GeometryProcessing/Source/DynamicMesh/Public/MeshSimplification.h
@@ -194,17 +194,17 @@ protected:
 	// Subclasses can override these two functions to restrict the affected edges (eg EdgeLoopRemesher)
 
 	// We are using a modulo-index loop to break symmetry/pathological conditions. 
-	const int ModuloPrime = 31337;     // any prime will do...
-	int MaxEdgeID = 0;
+	uint64 MaxEdgeID = 0;
 	virtual int StartEdges() 
 	{
-		MaxEdgeID = Mesh->MaxEdgeID();
+		MaxEdgeID = static_cast<uint64>(Mesh->MaxEdgeID());
 		return 0;
 	}
 
 	virtual int GetNextEdge(int CurEdgeID, bool& bDoneOut) 
 	{
-		int new_eid = (CurEdgeID + ModuloPrime) % MaxEdgeID;
+		constexpr uint64 ModuloPrime = 4294967311ull;     // choose prime > max uint32, to always be co-prime with MaxEdgeID
+		int new_eid = static_cast<int>((static_cast<uint64>(CurEdgeID) + ModuloPrime) % MaxEdgeID);
 		bDoneOut = (new_eid == 0);
 		return new_eid;
 	}
diff --git a/Engine/Plugins/Runtime/GeometryProcessing/Source/DynamicMesh/Public/Remesher.h b/Engine/Plugins/Runtime/GeometryProcessing/Source/DynamicMesh/Public/Remesher.h
--- a/Engine/Plugins/Runtime/GeometryProcessing/Source/DynamicMesh/Public/Remesher.h
+++ b/Engine/Plugins/Runtime/GeometryProcessing/Source/DynamicMesh/Public/Remesher.h
@@ -210,17 +210,17 @@ protected:
 	// sequential edge IDs, and all edges are < min edge length, then we can easily end
 	// up successively collapsing each tiny edge, and eroding away the entire mesh!
 	// By using modulo-index loop we jump around and hence this is unlikely to happen.
-	const int ModuloPrime = 31337;     // any prime will do...
-	int MaxEdgeID = 0;
+	uint64 MaxEdgeID = 0;
 	virtual int StartEdges()
 	{
-		MaxEdgeID = Mesh->MaxEdgeID();
+		MaxEdgeID = static_cast<uint64>(Mesh->MaxEdgeID());
 		return 0;
 	}
 
 	virtual int GetNextEdge(int CurEdgeID, bool& bDone)
 	{
-		int new_eid = (CurEdgeID + ModuloPrime) % MaxEdgeID;
+		constexpr uint64 ModuloPrime = 4294967311ull;     // choose prime > max uint32, to always be co-prime with MaxEdgeID
+		int new_eid = static_cast<int>( (static_cast<uint64>(CurEdgeID) + ModuloPrime) % MaxEdgeID );
 		bDone = (new_eid == 0);
 		return new_eid;
 	}
diff --git a/Engine/Source/Runtime/GeometryCore/Public/Util/IteratorUtil.h b/Engine/Source/Runtime/GeometryCore/Public/Util/IteratorUtil.h
--- a/Engine/Source/Runtime/GeometryCore/Public/Util/IteratorUtil.h
+++ b/Engine/Source/Runtime/GeometryCore/Public/Util/IteratorUtil.h
@@ -427,7 +427,7 @@ public:
 struct FModuloIteration
 {
 	uint64 MaxIndex = 0;
-	uint64 ModuloPrime = 3208642561;		// first elite prime > max_unsigned_int/2
+	uint64 ModuloPrime = 4294967311ull;		// prime > max_unsigned_int
 	uint64 CurIndex = 0;
 	uint64 StartIndex = 0;
 	uint64 Count = 0;
