Geometry Processing: Make simple intrinsic triangulation mesh robust to misoriented triangles.

diff --git a/Engine/Plugins/Runtime/GeometryProcessing/Source/DynamicMesh/Private/Operations/IntrinsicTriangulationMesh.cpp b/Engine/Plugins/Runtime/GeometryProcessing/Source/DynamicMesh/Private/Operations/IntrinsicTriangulationMesh.cpp
--- a/Engine/Plugins/Runtime/GeometryProcessing/Source/DynamicMesh/Private/Operations/IntrinsicTriangulationMesh.cpp
+++ b/Engine/Plugins/Runtime/GeometryProcessing/Source/DynamicMesh/Private/Operations/IntrinsicTriangulationMesh.cpp
@@ -120,7 +120,7 @@ namespace
 
 	/**
 	* Return the location of third vertex in a triangle in R2 with prescribed lengths
-	* such that the resulting triangle could be constructed as { (0,0),  (L1, 0), Result}.
+	* such that the resulting triangle could be constructed as { (0,0),  (L0, 0), Result}.
 	* 
 	* NB: this assumes, but does not check, that the lengths satisfy the triangle inequality
 	*/
@@ -142,12 +142,12 @@ namespace
 	* Given three side lengths that satisfy the triangle inequality,
 	* this updates the 2d positions to be the vertices of a triangle, 
 	* the edge lengths of which ( when in CCW order)  match the prescribed lengths.  
-	* with vertices { (0,0),  (L1, 0), (X, Y)} where Y > 0
+	* with vertices { (0,0),  (L0, 0), (X, Y)} where Y > 0
 	*/
 	void TriangleFromLengths(const double L0, const double L1, const double L2, FVector2d& p0, FVector2d& p1, FVector2d& p3)
 	{
 		p0 = FVector2d(0., 0.);
-		p1 = FVector2d(L1, 0.);
+		p1 = FVector2d(L0, 0.);
 		p3 = ComputeOpposingVert2d(L0, L1, L2);
 	}
 
@@ -1919,16 +1919,21 @@ EMeshResult FSimpleIntrinsicEdgeFlipMesh::FlipEdgeTopology(int32 eab, FEdgeFlipI
 	const int32 ebc = T0te[AddOneModThree[T0IndexOf]];
 	const int32 eca = T0te[AddTwoModThree[T0IndexOf]];
 
-	const int32 ead = T1te[AddOneModThree[T1IndexOf]];
-	const int32 edb = T1te[AddTwoModThree[T1IndexOf]];
+	int32 ead = T1te[AddOneModThree[T1IndexOf]];
+	int32 edb = T1te[AddTwoModThree[T1IndexOf]];
+	if (!FlipInfo.bHadSameOrientations)
+	{
+		Swap(ead,edb);
+	}
+	
 
 	// update triangles
 	Triangles[t0] = FIndex3i(c, d, b);
 	Triangles[t1] = FIndex3i(d, c, a);
 
 	// update edge AB, which becomes flipped edge CD
-	SetEdgeVerticesInternal(eab, c, d);
-	SetEdgeTrianglesInternal(eab, t0, t1);
+	SetEdgeVerticesInternal(eab, c, d);     // Edge.Vert = (c, d) or (d, c)   
+	SetEdgeTrianglesInternal(eab, t0, t1);  // Edge.Tri = (t0, t1)
 	const int32 ecd = eab;
 
 	// update the two other edges whose triangle nbrs have changed
@@ -2002,6 +2007,8 @@ FVector3d FSimpleIntrinsicEdgeFlipMesh::ComputeTriInternalAnglesR(const int32 TI
 
 double FSimpleIntrinsicEdgeFlipMesh::GetOpposingVerticesDistance(int32 EID) const
 {
+	const bool bHasSameOrientation = AreSameOrientation(EID);
+
 	const double OrgLength  = GetEdgeLength(EID);
 	const FIndex2i EdgeTris = GetEdgeT(EID);
 	checkSlow(EdgeTris[1] != FDynamicMesh3::InvalidID);
@@ -2018,8 +2025,11 @@ double FSimpleIntrinsicEdgeFlipMesh::GetOpposingVerticesDistance(int32 EID) cons
 	}
 	// rotate second tri so the shared edge aligns
 	Opp2dVerts[1].Y = -Opp2dVerts[1].Y;
-	Opp2dVerts[1].X = OrgLength - Opp2dVerts[1].X;
-
+	if (bHasSameOrientation)
+	{ 
+		Opp2dVerts[1].X = OrgLength - Opp2dVerts[1].X;
+	}
+	
 	return (Opp2dVerts[0] - Opp2dVerts[1]).Length();
 }
 
@@ -2036,19 +2046,41 @@ EMeshResult FSimpleIntrinsicEdgeFlipMesh::FlipEdge(int32 EID, FEdgeFlipInfo& Edg
 		return EMeshResult::Failed_IsBoundaryEdge;
 	}
 
+	const bool bHasSameOrientation = AreSameOrientation(EID);
+	EdgeFlipInfo.bHadSameOrientations = bHasSameOrientation;
+
 	// prohibit case where the edge is shared by a non-convex pair of triangles.
 	{
 		// original triangles
 		FDynamicMesh3::FEdge OrgEdge = GetEdge(EID);
 
-		// Assumes both triangles have same orientation
+
+
+		// total interior angles at the edge vertices as ordered by the first triangle.
 		double TotalAngleAtOrgVert[2] = { 0., 0. };
-		for (int32 i = 0; i < 2; ++i)
-		{	
-			int32 TriID = OrgEdge.Tri[i];
-			int32 IndexOf = GetTriEdges(TriID).IndexOf(EID);
-			TotalAngleAtOrgVert[i] += InternalAngles[TriID][IndexOf];
-			TotalAngleAtOrgVert[(i+1) % 2] += InternalAngles[TriID][AddOneModThree[IndexOf]];
+		{
+			// accumulate interior angles at edge vertices in first triangle. 
+			{
+				const int32 TriID = OrgEdge.Tri[0];
+				const int32 IndexOf = GetTriEdges(TriID).IndexOf(EID);
+				TotalAngleAtOrgVert[0] += InternalAngles[TriID][IndexOf];
+				TotalAngleAtOrgVert[1] += InternalAngles[TriID][AddOneModThree[IndexOf]];
+			}
+			// accumulate interior angles at the edge vertices in the second triangle
+			{
+				const int32 TriID = OrgEdge.Tri[1];
+				const int32 IndexOf = GetTriEdges(TriID).IndexOf(EID);
+				if (bHasSameOrientation)
+				{ 
+					TotalAngleAtOrgVert[1] += InternalAngles[TriID][IndexOf];
+					TotalAngleAtOrgVert[0] += InternalAngles[TriID][AddOneModThree[IndexOf]];
+				}
+				else
+				{
+					TotalAngleAtOrgVert[0] += InternalAngles[TriID][IndexOf];
+					TotalAngleAtOrgVert[1] += InternalAngles[TriID][AddOneModThree[IndexOf]];
+				}
+			}
 		}
 
 		if (TotalAngleAtOrgVert[0] > TMathUtilConstants<double>::Pi - TMathUtilConstants<double>::ZeroTolerance || TotalAngleAtOrgVert[1] > TMathUtilConstants<double>::Pi - TMathUtilConstants<double>::ZeroTolerance)
diff --git a/Engine/Plugins/Runtime/GeometryProcessing/Source/DynamicMesh/Public/Operations/IntrinsicTriangulationMesh.h b/Engine/Plugins/Runtime/GeometryProcessing/Source/DynamicMesh/Public/Operations/IntrinsicTriangulationMesh.h
--- a/Engine/Plugins/Runtime/GeometryProcessing/Source/DynamicMesh/Public/Operations/IntrinsicTriangulationMesh.h
+++ b/Engine/Plugins/Runtime/GeometryProcessing/Source/DynamicMesh/Public/Operations/IntrinsicTriangulationMesh.h
@@ -65,7 +65,13 @@ class DYNAMICMESH_API FSimpleIntrinsicEdgeFlipMesh
 public:
 	
 	using FEdge = FDynamicMesh3::FEdge;
-	using FEdgeFlipInfo = DynamicMeshInfo::FEdgeFlipInfo;
+	
+	// information collected when flipping an edge.
+	struct FEdgeFlipInfo : DynamicMeshInfo::FEdgeFlipInfo
+	{
+		bool bHadSameOrientations = true; // will be false if the original triangles had different orientations.
+	};
+	
 	/** InvalidID indicates that a vertex/edge/triangle ID is invalid */
 	constexpr static int InvalidID = IndexConstants::InvalidID;
 
@@ -83,6 +89,10 @@ public:
 	void Reset(const FDynamicMesh3& SrcMesh);
 	/**
 	* Flip a single edge in the Intrinsic Mesh.
+	* Note: After a successful flip, both triangles will share the same orientation (matching that of the pre-flip Edge.Tri.A)
+	* even if they had different orientations (e.g. cw & ccw) prior to flipping. 
+	*       
+	* @param EID indicates the edge to flip.
 	* @param EdgeFlipInfo [out] populated on return.
 	*
 	* @return a success or failure info.
@@ -303,6 +313,32 @@ public:
 		return Vector3Type(Src[Index], Src[AddOneModThree[Index]], Src[AddTwoModThree[Index]]);
 	}
 
+
+	/**
+	* @return false if the two triangles adjacent to the indicated edge have the different orientation (eg counter-clockwise and clockwise),
+	* otherwise return true.
+	* Note: by convention this will return true for boundary edges ( or invalid edges ).
+	*/
+	bool AreSameOrientation(int32 EID) const
+	{
+		if (!IsEdge(EID)) return true;
+
+		const FSimpleIntrinsicEdgeFlipMesh::FEdge Edge = GetEdge(EID);
+
+		// boundary edge is true by convention
+		if (Edge.Tri.B == FSimpleIntrinsicEdgeFlipMesh::InvalidID) return true;
+
+		const FIndex3i TriA = GetTriangle(Edge.Tri.A);
+		const FIndex3i TriB = GetTriangle(Edge.Tri.B);
+		const int V0 = Edge.Vert.A;
+		const int V1 = Edge.Vert.B;
+		const bool bOrderedAB = (IndexUtil::FindTriOrderedEdge(V0, V1, TriA) != FSimpleIntrinsicEdgeFlipMesh::InvalidID && IndexUtil::FindTriOrderedEdge(V1, V0, TriB) != FSimpleIntrinsicEdgeFlipMesh::InvalidID);
+		const bool bOrderedBA = (IndexUtil::FindTriOrderedEdge(V1, V0, TriA) != FSimpleIntrinsicEdgeFlipMesh::InvalidID && IndexUtil::FindTriOrderedEdge(V0, V1, TriB) != FSimpleIntrinsicEdgeFlipMesh::InvalidID);
+
+		// true if consistent orientation
+		return bOrderedAB || bOrderedBA;
+	}
+
 protected:
 	
 	/** @return three requested edge lengths */
@@ -328,7 +364,10 @@ protected:
 protected:
 	// basic mesh operations
 
-	/** update the mesh topology (e.g. triangles and edges) by doing an edge flip, but does not update the intrinsic edge lengths or angles.*/
+	/**
+	* update the mesh topology (e.g. triangles and edges) by doing an edge flip, but does not update the intrinsic edge lengths or angles.
+	* Note: this requires precomputed FlipInfo.bHadSameOrientation. 
+	*/
 	EMeshResult FlipEdgeTopology(int32 eab, FEdgeFlipInfo& FlipInfo);
 
 	/** update the existing edge to reference vertices 'a' and 'b'.*/
@@ -383,7 +422,7 @@ protected:
 
 	TDynamicVector<FIndex3i> Triangles;     // List of triangle vertex - index triplets[Vert0 Vert1 Vert2]
 	FRefCountVector TriangleRefCounts{};    // Reference counts of triangle indices. Ref count is always 1 if the triangle exists. Iterate over this to find out which triangle indices are valid.
-	TDynamicVector<FIndex3i> TriangleEdges; // List of triangle edge triplets [Edge0 Edge1 Edge2]
+	TDynamicVector<FIndex3i> TriangleEdges; // List of triangle edge triplets [Edge0 Edge1 Edge2], note these correspond to [Vert0-Vert1, Vert1-Vert2, Vert2-Vert0] for a given triangle.
 
 	TDynamicVector<FEdge> Edges;            // List of edge elements. An edge is four elements [VertA, VertB, Tri0, Tri1], where VertA < VertB, and Tri1 may be InvalidID (if the edge is a boundary edge)
 	FRefCountVector EdgeRefCounts{};        // Reference counts of edge indices. Ref count is always 1 if the edge exists. Iterate over this to find out which edge indices are valid.
@@ -414,8 +453,8 @@ class DYNAMICMESH_API FSimpleIntrinsicMesh : public FSimpleIntrinsicEdgeFlipMesh
 {
 public:
 	typedef FSimpleIntrinsicEdgeFlipMesh  MyBase;
-	using FEdge = FDynamicMesh3::FEdge;
-	using FEdgeFlipInfo     = DynamicMeshInfo::FEdgeFlipInfo;
+	using FEdge             = FSimpleIntrinsicEdgeFlipMesh::FEdge;
+	using FEdgeFlipInfo     = FSimpleIntrinsicEdgeFlipMesh::FEdgeFlipInfo;
 	using FPokeTriangleInfo = DynamicMeshInfo::FPokeTriangleInfo;
 	using FEdgeSplitInfo    = DynamicMeshInfo::FEdgeSplitInfo;
 
@@ -596,6 +635,8 @@ struct DYNAMICMESH_API TEdgeCorrespondence
 * NB: The lifetime of this structure should not exceed that of the original
 * FDynamicMesh as this class holds a pointer to that mesh to reference locations on its surface.
 * Also, this class isn't currently expected to work with surface meshes that have bow-ties.
+* 
+* Note:  This assumes but does not check that the FDynamicMesh3 "Surface Mesh" has consistent orientation of all triangles.
 */
 class DYNAMICMESH_API FIntrinsicEdgeFlipMesh : public  FSimpleIntrinsicEdgeFlipMesh
 {
@@ -716,6 +757,8 @@ protected:
 * NB: The lifetime of this structure should not exceed that of the original
 * FDynamicMesh as this class holds a pointer to that mesh to reference locations on its surface.
 * Also, this class isn't currently expected to work with surface meshes that have bow-ties.
+* 
+* Note:  This assumes but does not check that the FDynamicMesh3 "Surface Mesh" has consistent orientation of all triangles.
 */
 
 class DYNAMICMESH_API FIntrinsicMesh  : public  FSimpleIntrinsicMesh
@@ -723,8 +766,8 @@ class DYNAMICMESH_API FIntrinsicMesh  : public  FSimpleIntrinsicMesh
 public:
 	typedef FSimpleIntrinsicMesh     MyBase;
 
-	using FEdge = FIntrinsicEdgeFlipMesh::FEdge;
-	using FEdgeFlipInfo = FIntrinsicEdgeFlipMesh::FEdgeFlipInfo;
+	using FEdge = FSimpleIntrinsicMesh::FEdge;
+	using FEdgeFlipInfo = FSimpleIntrinsicMesh::FEdgeFlipInfo;
 	using FSurfacePoint = IntrinsicCorrespondenceUtils::FSurfacePoint;
 	using FEdgeAndCrossingIdx = IntrinsicCorrespondenceUtils::FNormalCoordinates::FEdgeAndCrossingIdx;
 	using FPokeTriangleInfo   = DynamicMeshInfo::FPokeTriangleInfo;
@@ -865,6 +908,8 @@ protected:
 * NB: The lifetime of this structure should not exceed that of the original
 * FDynamicMesh as this class holds a pointer to that mesh to reference locations on its surface.
 * Also, this class isn't currently expected to work with surface meshes that have bow-ties.
+* 
+* Note:  This assumes but does not check that the FDynamicMesh3 "Surface Mesh" has consistent orientation of all triangles.
 */ 
 class DYNAMICMESH_API FIntrinsicTriangulation : public  FSimpleIntrinsicMesh
 {
@@ -873,9 +918,9 @@ public:
 
 	typedef FSimpleIntrinsicMesh MyBase;
 
-	using FEdgeFlipInfo = DynamicMeshInfo::FEdgeFlipInfo;
-	using FEdgeSplitInfo = DynamicMeshInfo::FEdgeSplitInfo;
-	using FPokeTriangleInfo = DynamicMeshInfo::FPokeTriangleInfo;
+	using FEdgeFlipInfo = FSimpleIntrinsicMesh::FEdgeFlipInfo;
+	using FEdgeSplitInfo = FSimpleIntrinsicMesh::FEdgeSplitInfo;
+	using FPokeTriangleInfo = FSimpleIntrinsicMesh::FPokeTriangleInfo;
 	using FSurfacePoint  = IntrinsicCorrespondenceUtils::FSurfacePoint;
 
 	FIntrinsicTriangulation(const FDynamicMesh3& SrcMesh);
@@ -993,5 +1038,6 @@ int32 DYNAMICMESH_API FlipToDelaunay(FSimpleIntrinsicMesh& IntrinsicMesh, TSet<i
 int32 DYNAMICMESH_API FlipToDelaunay(FIntrinsicMesh& IntrinsicMesh, TSet<int>& Uncorrected, const int32 MaxFlipCount = TMathUtilConstants<int>::MaxReal);
 int32 DYNAMICMESH_API FlipToDelaunay(FIntrinsicTriangulation& IntrinsicMesh, TSet<int>& Uncorrected, const int32 MaxFlipCount = TMathUtilConstants<int>::MaxReal);
 
-}; // end namespace Geometry
-}; // end namespace UE
+
+} // end namespace Geometry
+} // end namespace UE
