fix issue where in rare cases bevel can fail to set correct matching new edges/vertices

diff --git a/Engine/Plugins/Runtime/GeometryProcessing/Source/DynamicMesh/Private/Operations/MeshBevel.cpp b/Engine/Plugins/Runtime/GeometryProcessing/Source/DynamicMesh/Private/Operations/MeshBevel.cpp
--- a/Engine/Plugins/Runtime/GeometryProcessing/Source/DynamicMesh/Private/Operations/MeshBevel.cpp
+++ b/Engine/Plugins/Runtime/GeometryProcessing/Source/DynamicMesh/Private/Operations/MeshBevel.cpp
@@ -1074,22 +1074,35 @@ void FMeshBevel::FixUpUnlinkedBevelEdges(FDynamicMesh3& Mesh)
 	// into wedges until afterwards. 
 	for (FBevelEdge& Edge : Edges)
 	{
-		// If this is a one-mesh-edge edge, then Edge.NewMeshEdges is incorrect because we could
-		// not actually unlink the edge in UnlinkBevelEdgeInterior (since there were no interior vertices).
-		// But now that we have unlinked the vertices, the other edge now exists, and we can find it here.
-		if (Edge.MeshEdges.Num() == 1)
+		// In some cases, like single edges, Edge.NewMeshEdges is incorrect (mapped to itself); e.g., because we
+		// could not actually unlink the edge in UnlinkBevelEdgeInterior (if there were no interior vertices).
+		// Now that all vertices are unlinked, take one more pass through and fix the single-edge or self-mapped edges.
+		bool bSingleEdge = Edge.MeshEdges.Num() == 1;
+		bool bFailedEdgePairing = false;
+		for (int32 Idx = 0; Idx < Edge.MeshEdges.Num(); ++Idx)
 		{
-			int32* FoundOtherEdge = MeshEdgePairs.Find(Edge.MeshEdges[0]);
+			bool bNewEdgeIsOld = Edge.MeshEdges[Idx] == Edge.NewMeshEdges[Idx];
+			if (!bSingleEdge && !bNewEdgeIsOld)
+			{
+				continue;
+			}
+
+			int32* FoundOtherEdge = MeshEdgePairs.Find(Edge.MeshEdges[Idx]);
 			MESH_BEVEL_DEBUG_CHECK(FoundOtherEdge != nullptr);
 			if (FoundOtherEdge != nullptr)
 			{
-				Edge.NewMeshEdges[0] = *FoundOtherEdge;
+				Edge.NewMeshEdges[Idx] = *FoundOtherEdge;
 			}
 			else
 			{
-				continue;		// something went wrong, loop below will break things
+				bFailedEdgePairing = true;		// something went wrong, loop below will break things
+				break;
 			}
 		}
+		if (bFailedEdgePairing)
+		{
+			continue;
+		}
 
 		// process start and end vertices of the path
 		for (int32 j = 0; j < 2; ++j)
@@ -1262,11 +1275,7 @@ void FMeshBevel::DisplaceVertices(FDynamicMesh3& Mesh, double Distance)
 			}
 			else 
 			{
-				// TODO: this ensure has been sporadically hit in Lyra, the case appears to be that a wedge ends up with
-				// only a single SolveLine. Likely this happens because some bevel-edge was not processed as we expected,
-				// ie the wedge-building went wrong due to a topological case we have not properly handled. Needs to be
-				// investigated more deeply, but in the meantime, skipping here just means the vertex will end up mis-positioned
-				//ensure(SolveLines.Num() >= 2);
+				MESH_BEVEL_DEBUG_CHECK(SolveLines.Num() >= 2);
 				if (SolveLines.Num() >= 2)
 				{
 					Wedge.NewPosition = UE::Geometry::SolveInsetVertexPositionFromLinePair(CurPos, SolveLines[0], SolveLines[1]);
