Remesher: add an optional edge length scale function that can be used for adaptive-density remeshing

diff --git a/Engine/Plugins/Runtime/GeometryProcessing/Source/DynamicMesh/Private/Remesher.cpp b/Engine/Plugins/Runtime/GeometryProcessing/Source/DynamicMesh/Private/Remesher.cpp
--- a/Engine/Plugins/Runtime/GeometryProcessing/Source/DynamicMesh/Private/Remesher.cpp
+++ b/Engine/Plugins/Runtime/GeometryProcessing/Source/DynamicMesh/Private/Remesher.cpp
@@ -154,6 +154,12 @@ FRemesher::EProcessResult FRemesher::ProcessEdge(int edgeID)
 	FVector3d vB(Mesh->GetVertex(b));
 	double edge_len_sqr = DistanceSquared(vA, vB);
 
+	if (CustomEdgeLengthScaleF)
+	{
+		const double EdgeLengthScale = CustomEdgeLengthScaleF(*Mesh, a, b);
+		edge_len_sqr = (EdgeLengthScale * EdgeLengthScale) * edge_len_sqr;
+	}
+
 	// check if we should collapse, and also find which vertex we should collapse to,
 	// in cases where we have constraints/etc
 	int collapse_to = -1;
diff --git a/Engine/Plugins/Runtime/GeometryProcessing/Source/DynamicMesh/Public/Remesher.h b/Engine/Plugins/Runtime/GeometryProcessing/Source/DynamicMesh/Public/Remesher.h
--- a/Engine/Plugins/Runtime/GeometryProcessing/Source/DynamicMesh/Public/Remesher.h
+++ b/Engine/Plugins/Runtime/GeometryProcessing/Source/DynamicMesh/Public/Remesher.h
@@ -102,6 +102,8 @@ public:
 	/** Override constant SmoothSpeedT with function */
 	TFunction<double(const FDynamicMesh3&, int)> CustomSmoothSpeedF;
 
+	/** Scale the computed edge length to allow for adaptive remeshing */
+	TFunction<double(const FDynamicMesh3& Mesh, int VertexA, int VertexB)> CustomEdgeLengthScaleF;
 
 
 	/** enable parallel projection. Only applied in AfterRefinement mode */
diff --git a/Engine/Plugins/Runtime/MeshModelingToolset/Source/ModelingOperators/Private/CleaningOps/RemeshMeshOp.cpp b/Engine/Plugins/Runtime/MeshModelingToolset/Source/ModelingOperators/Private/CleaningOps/RemeshMeshOp.cpp
--- a/Engine/Plugins/Runtime/MeshModelingToolset/Source/ModelingOperators/Private/CleaningOps/RemeshMeshOp.cpp
+++ b/Engine/Plugins/Runtime/MeshModelingToolset/Source/ModelingOperators/Private/CleaningOps/RemeshMeshOp.cpp
@@ -182,6 +182,7 @@ void FRemeshMeshOp::CalculateResult(FProgressCancel* Progress)
 	}
 	Remesher->SetProjectionTarget(ProjTarget.Get());
 
+	Remesher->CustomEdgeLengthScaleF = CustomEdgeLengthScaleF;
 
 	Remesher->Progress = Progress;
 
diff --git a/Engine/Plugins/Runtime/MeshModelingToolset/Source/ModelingOperators/Public/CleaningOps/RemeshMeshOp.h b/Engine/Plugins/Runtime/MeshModelingToolset/Source/ModelingOperators/Public/CleaningOps/RemeshMeshOp.h
--- a/Engine/Plugins/Runtime/MeshModelingToolset/Source/ModelingOperators/Public/CleaningOps/RemeshMeshOp.h
+++ b/Engine/Plugins/Runtime/MeshModelingToolset/Source/ModelingOperators/Public/CleaningOps/RemeshMeshOp.h
@@ -90,6 +90,8 @@ public:
 	bool bUseWorldSpace = false;
 	bool bParallel = true;
 
+	TFunction<double(const FDynamicMesh3&, int VertexA, int VertexB)> CustomEdgeLengthScaleF;
+
 	// Normal flow only:
 
 	/// During each call to RemeshIteration, do this many passes of face-aligned projection
