GeometryProcessing: Improve the UV generation scheme PatchBuilder via a change to the polygroup generation stage, providing area aware sampling of the mesh surface and reducing the overcompensation effects that small, disconnected regions in the mesh were having on the overall segmentation.

Also provides a new CVar, modeling.DisableAutoUVAreaDensitySampling, which when set, returns the behavior to the original algorithm, in case there are legacy reasons to maintain that behavior.

diff --git a/Engine/Plugins/Runtime/GeometryProcessing/Source/DynamicMesh/Private/Parameterization/PatchBasedMeshUVGenerator.cpp b/Engine/Plugins/Runtime/GeometryProcessing/Source/DynamicMesh/Private/Parameterization/PatchBasedMeshUVGenerator.cpp
--- a/Engine/Plugins/Runtime/GeometryProcessing/Source/DynamicMesh/Private/Parameterization/PatchBasedMeshUVGenerator.cpp
+++ b/Engine/Plugins/Runtime/GeometryProcessing/Source/DynamicMesh/Private/Parameterization/PatchBasedMeshUVGenerator.cpp
@@ -120,9 +120,10 @@ FGeometryResult FPatchBasedMeshUVGenerator::AutoComputeUVs(
 	return ResultInfo;
 }
 
-
-
-
+static TAutoConsoleVariable<int32> CVarEnableAutoUVAreaDensitySampling(
+	TEXT("modeling.DisableAutoUVAreaDensitySampling"),
+	0,
+	TEXT("If set, return the behavior of AutoUV's PatchBuilder algorithm to the legacy behavior."));
 
 bool FPatchBasedMeshUVGenerator::ComputeInitialMeshPatches(
 	FDynamicMesh3& TargetMesh,
@@ -137,7 +138,16 @@ bool FPatchBasedMeshUVGenerator::ComputeInitialMeshPatches(
 		FPolygroupsGenerator::EWeightingType::NormalDeviation : FPolygroupsGenerator::EWeightingType::None;
 	FVector3d GeneratorParams(this->PatchNormalWeight, 1.0, 1.0);
 
-	bool bOK = Generator.FindPolygroupsFromFurthestPointSampling(this->TargetPatchCount, WeightType, GeneratorParams, this->GroupConstraint);
+	bool bUseAreaDensitySampling = (CVarEnableAutoUVAreaDensitySampling.GetValueOnAnyThread() == 0);
+	bool bOK = false;
+	if (bUseAreaDensitySampling)
+	{
+		bOK = Generator.FindPolygroupsFromAreaDensityPointSampling(this->TargetPatchCount, WeightType, GeneratorParams, this->GroupConstraint);
+	}
+	else
+	{
+		bOK = Generator.FindPolygroupsFromFurthestPointSampling(this->TargetPatchCount, WeightType, GeneratorParams, this->GroupConstraint);
+	}
 
 	if (!bOK)
 	{
diff --git a/Engine/Plugins/Runtime/GeometryProcessing/Source/DynamicMesh/Private/Polygroups/PolygroupsGenerator.cpp b/Engine/Plugins/Runtime/GeometryProcessing/Source/DynamicMesh/Private/Polygroups/PolygroupsGenerator.cpp
--- a/Engine/Plugins/Runtime/GeometryProcessing/Source/DynamicMesh/Private/Polygroups/PolygroupsGenerator.cpp
+++ b/Engine/Plugins/Runtime/GeometryProcessing/Source/DynamicMesh/Private/Polygroups/PolygroupsGenerator.cpp
@@ -894,7 +894,177 @@ bool FPolygroupsGenerator::FindPolygroupsFromFurthestPointSampling(
 }
 
 
+bool FPolygroupsGenerator::FindPolygroupsFromAreaDensityPointSampling(
+	int32 NumPoints,
+	EWeightingType WeightingType,
+	FVector3d WeightingCoeffs,
+	FPolygroupSet* StartingGroups)
+{
+	NumPoints = FMath::Min(NumPoints, Mesh->VertexCount());
+
+	// cannot seem to use auto or TUniqueFunction here...
+	TFunction<bool(int32, int32)> TrisConnectedPredicate;
+	TrisConnectedPredicate = [](int, int) -> bool { return true; };
+	if (StartingGroups != nullptr)
+	{
+		TrisConnectedPredicate = [StartingGroups](int a, int b) -> bool { return StartingGroups->GetGroup(a) == StartingGroups->GetGroup(b) ? true : false; };
+	}
+
+	double MeshTotalArea = 0.0;
+	for (int32 tid : Mesh->TriangleIndicesItr())
+	{
+		MeshTotalArea += Mesh->GetTriArea(tid);
+	}
+	double ExpectedAreaPerSeed = MeshTotalArea / NumPoints;
+
+	FMeshFaceDualGraph FaceGraph;
+	FMeshFaceDualGraph::MakeFaceDualGraphForMesh(Mesh, FaceGraph, TrisConnectedPredicate);
+
+	TArray<int32> SeedIndices;
+
+	// need to add at least one seed point for each mesh connected component, so that all triangles are assigned a group
+	// TODO: two seed points for components that have no boundary?
+	FMeshConnectedComponents Components(Mesh);
+	Components.FindConnectedTriangles(TrisConnectedPredicate);
+	int32 NumConnected = Components.Num();
+
+	struct FComponentContext
+	{
+		TArray< int32 > Seeds;
+	};
+
+	TArray<FComponentContext> Contexts;
+	Contexts.SetNum(NumConnected);
+	TArrayView<FComponentContext> ContextView = TArrayView<FComponentContext>(Contexts);
+
+	ParallelForWithExistingTaskContext(ContextView, NumConnected, 1,
+		[this, ExpectedAreaPerSeed, &FaceGraph, &Components](FComponentContext &Context, int32 k) {
+
+			TArray<int32> ComponentSeeds;
 
+			double ComponentTotalArea = 0.0;
+			for (int32 tid : Components[k].Indices)
+			{
+				ComponentTotalArea += Mesh->GetTriArea(tid);
+			}
+			int32 DesiredSeeds = FMath::Max(1, FMath::CeilToInt32(ComponentTotalArea / ExpectedAreaPerSeed));
+			
+			ComponentSeeds.Add(Components[k].Indices[0]);
+
+			// We're going to bail early if we only need one seed for this component.
+			// No need to build the expensive Dijkstra weights if we're already done.
+			if (DesiredSeeds == 1)
+			{
+				Context.Seeds.Append(ComponentSeeds);
+				return;
+			}
+
+			TIncrementalMeshDijkstra<FMeshFaceDualGraph>::FSeedPoint InitialSeedPoint{ ComponentSeeds[0], ComponentSeeds[0], 0.0 };
+			TIncrementalMeshDijkstra<FMeshFaceDualGraph> FurthestPoints(&FaceGraph);
+			FurthestPoints.AddSeedPoints({ InitialSeedPoint });
+
+			while (ComponentSeeds.Num() < DesiredSeeds)
+			{
+				int32 NextPointID = FurthestPoints.FindMaxGraphDistancePointID();
+				if (NextPointID >= 0)
+				{
+					TIncrementalMeshDijkstra<FMeshFaceDualGraph>::FSeedPoint NextSeedPoint{ NextPointID, NextPointID, 0.0 };
+					FurthestPoints.AddSeedPoints({ NextSeedPoint });
+					ComponentSeeds.Add(NextPointID);
+				}
+				else
+				{
+					break;
+				}
+			}
+
+			Context.Seeds.Append(ComponentSeeds);
+		}
+	);
+
+	for (int32 k = 0; k < NumConnected; ++k)
+	{
+		SeedIndices.Append(Contexts[k].Seeds);
+	}
+
+	// Now that we have furthest point set, recompute a Dijkstra propagation with optional weighting
+	// (unweighted version should be the same as the FurthestPoints dijkstra though...could re-use?)
+
+	TMeshDijkstra<FMeshFaceDualGraph> SuperPixels(&FaceGraph);
+
+	if (WeightingType == EWeightingType::NormalDeviation)
+	{
+		SuperPixels.bEnableDistanceWeighting = true;
+		SuperPixels.GetWeightedDistanceFunc = [this, WeightingCoeffs, &FaceGraph](int32 FromVID, int32 ToVID, int32 SeedVID, double Distance)
+		{
+			FVector3d NA = FaceGraph.GetNormal(ToVID);
+			FVector3d NB = FaceGraph.GetNormal(SeedVID);
+			double Dot = NA.Dot(NB);
+			if (WeightingCoeffs.X > 0.001)
+			{
+				Dot = FMathd::Pow(Dot, WeightingCoeffs.X);
+			}
+			double W = FMathd::Clamp(1.0 - Dot * Dot, 0.0, 1.0);
+			W = W * W * W;
+			double Weight = FMathd::Clamp(W, 0.001, 1.0);
+			return Weight * Distance;
+		};
+	}
+
+	TArray<TMeshDijkstra<FMeshFaceDualGraph>::FSeedPoint> SuperPixelSeeds;
+	for (int32 k = 0; k < SeedIndices.Num(); ++k)
+	{
+		SuperPixelSeeds.Add(TMeshDijkstra<FMeshFaceDualGraph>::FSeedPoint{ k, SeedIndices[k], 0.0 });
+	}
+	SuperPixels.ComputeToMaxDistance(SuperPixelSeeds, TNumericLimits<double>::Max());
+
+	TArray<TArray<int32>> TriSets;
+	TriSets.SetNum(SeedIndices.Num());
+	TArray<int32> FailedSet;
+
+	for (int32 tid : Mesh->TriangleIndicesItr())
+	{
+		int32 SeedID = SuperPixels.GetSeedExternalIDForPointSetID(tid);
+		if (SeedID >= 0)
+		{
+			TriSets[SeedID].Add(tid);
+		}
+		else
+		{
+			FailedSet.Add(tid);
+		}
+	}
+
+
+	for (int32 k = 0; k < TriSets.Num(); ++k)
+	{
+		if (TriSets[k].Num() > 0)
+		{
+			FoundPolygroups.Add(MoveTemp(TriSets[k]));
+		}
+	}
+
+	if (FailedSet.Num() > 0)
+	{
+		FMeshConnectedComponents FailedComponents(Mesh);
+		FailedComponents.FindConnectedTriangles(FailedSet);
+		for (FMeshConnectedComponents::FComponent& Component : FailedComponents)
+		{
+			FoundPolygroups.Add(Component.Indices);
+		}
+	}
+
+	if (bApplyPostProcessing)
+	{
+		PostProcessPolygroups(true, TrisConnectedPredicate);
+	}
+	if (bCopyToMesh)
+	{
+		CopyPolygroupsToMesh();
+	}
+
+	return (FoundPolygroups.Num() > 0);
+}
 
 void FPolygroupsGenerator::PostProcessPolygroups(bool bApplyMerging, TFunctionRef<bool(int32, int32)> TrisConnectedPredicate)
 {
diff --git a/Engine/Plugins/Runtime/GeometryProcessing/Source/DynamicMesh/Public/Polygroups/PolygroupsGenerator.h b/Engine/Plugins/Runtime/GeometryProcessing/Source/DynamicMesh/Public/Polygroups/PolygroupsGenerator.h
--- a/Engine/Plugins/Runtime/GeometryProcessing/Source/DynamicMesh/Public/Polygroups/PolygroupsGenerator.h
+++ b/Engine/Plugins/Runtime/GeometryProcessing/Source/DynamicMesh/Public/Polygroups/PolygroupsGenerator.h
@@ -113,6 +113,32 @@ public:
 		FVector3d WeightingCoeffs = FVector3d::One(),
 		FPolygroupSet* StartingGroups = nullptr);
 
+	/**
+	*
+	* Similar to the technique used in FindPolygroupsFromFurthestPointSampling, except this algorithm
+	* attempts to select sample points in connected mesh regions based on the rough area of each region,
+	* granting more initial points to larger regions and fewer points to smaller regions, minium of one.
+	* This approach handles cases better when there is a large varience in region sizes and a large number
+	* of small regions compared to large regions. In this case, the alternative routine can over focus on
+	* small regions, granting the larger regions too few points. This leads to a relative over sampling of
+	* small regions and creates large polygroups out of the larger regions instead of more equitably breaking
+	* up larger regions into multiple groups. 
+	* 
+	* In this algorithm, NumPoints is less a strict target of final patches, but a rough guideline to direct
+	* the area sampler on how much total area each sample should occupy. E.g. a value of 100 would indicate
+	* we ideally want 100 equal sized regions of the mesh. However, since each connected region must have at
+	* least one sample, the final count may not be exactly NumPoints. Instead, each region, and it's area, is
+	* considered independently according to this desired "sample density". A "perfect" mesh (with one connected
+	* region and a high tesselation of equal sized triangles) should generate equal sized patches of NumPoints
+	* in count. A poorer quality mesh will result in a best effort take at this, while ensuring all triangles
+	* are part of a group.
+	* 
+	*/
+	bool FindPolygroupsFromAreaDensityPointSampling(
+		int32 NumPoints,
+		EWeightingType WeightingType,
+		FVector3d WeightingCoeffs = FVector3d::One(),
+		FPolygroupSet* StartingGroups = nullptr);
 
 	/**
 	* Copy the computed Polygroups to the input Mesh. Will be automatically called by the 
