fix some edges/vertices not moving when via mesh elements selection handles due to incorrect bounds check in deformer

diff --git a/Engine/Plugins/Runtime/GeometryProcessing/Source/DynamicMesh/Private/Operations/GroupTopologyDeformer.cpp b/Engine/Plugins/Runtime/GeometryProcessing/Source/DynamicMesh/Private/Operations/GroupTopologyDeformer.cpp
--- a/Engine/Plugins/Runtime/GeometryProcessing/Source/DynamicMesh/Private/Operations/GroupTopologyDeformer.cpp
+++ b/Engine/Plugins/Runtime/GeometryProcessing/Source/DynamicMesh/Private/Operations/GroupTopologyDeformer.cpp
@@ -147,7 +147,7 @@ void FGroupTopologyDeformer::SetActiveHandleFromSelection(const FGeometrySelecti
 		for (uint64 Element : Selection.Selection)
 		{
 			int32 GroupEdgeID = FGeoSelectionID(Element).TopologyID;
-			if (GroupEdgeID > 0 && GroupEdgeID < Topology->Edges.Num())
+			if (GroupEdgeID >= 0 && GroupEdgeID < Topology->Edges.Num())
 			{
 				Elements.Add(GroupEdgeID);
 			}
@@ -159,7 +159,7 @@ void FGroupTopologyDeformer::SetActiveHandleFromSelection(const FGeometrySelecti
 		for (uint64 Element : Selection.Selection)
 		{
 			int32 CornerID = FGeoSelectionID(Element).TopologyID;
-			if (CornerID > 0 && CornerID < Topology->Corners.Num())
+			if (CornerID >= 0 && CornerID < Topology->Corners.Num())
 			{
 				Elements.Add(CornerID);
 			}
