Modeling Tools: Adds Geometry Selection Support for the Weld Tool

diff --git a/Engine/Plugins/Experimental/MeshModelingToolsetExp/Source/MeshModelingToolsExp/Private/WeldMeshEdgesTool.cpp b/Engine/Plugins/Experimental/MeshModelingToolsetExp/Source/MeshModelingToolsExp/Private/WeldMeshEdgesTool.cpp
--- a/Engine/Plugins/Experimental/MeshModelingToolsetExp/Source/MeshModelingToolsExp/Private/WeldMeshEdgesTool.cpp
+++ b/Engine/Plugins/Experimental/MeshModelingToolsetExp/Source/MeshModelingToolsExp/Private/WeldMeshEdgesTool.cpp
@@ -13,6 +13,7 @@
 #include "Operations/MeshResolveTJunctions.h"
 
 #include "Math/UnrealMathUtility.h"
+#include "Selections/GeometrySelectionUtil.h"
 
 #include UE_INLINE_GENERATED_CPP_BY_NAME(WeldMeshEdgesTool)
 
@@ -24,7 +25,7 @@ using namespace UE::Geometry;
  * ToolBuilder
  */
 
-USingleSelectionMeshEditingTool* UWeldMeshEdgesToolBuilder::CreateNewTool(const FToolBuilderState& SceneState) const
+USingleTargetWithSelectionTool* UWeldMeshEdgesToolBuilder::CreateNewTool(const FToolBuilderState& SceneState) const
 {
 	return NewObject<UWeldMeshEdgesTool>(SceneState.ToolManager);
 }
@@ -48,6 +49,7 @@ public:
 
 	// parameters set by the tool
 	TSharedPtr<FDynamicMesh3, ESPMode::ThreadSafe> SourceMesh;
+	TSet<int32> SelectedEdges;
 	double Tolerance;
 	bool bOnlyUnique;
 	bool bResolveTJunctions;
@@ -99,6 +101,10 @@ public:
 		FinalNumBoundaryEdges = InitialNumBoundaryEdges;
 
 		FMergeCoincidentMeshEdges Merger(ResultMesh.Get());
+		if (!SelectedEdges.IsEmpty())
+		{
+			Merger.EdgesToMerge = &SelectedEdges;	
+		}
 		Merger.MergeVertexTolerance = Tolerance;
 		Merger.MergeSearchTolerance = 2 * Merger.MergeVertexTolerance;
 		Merger.OnlyUniquePairs = bOnlyUnique;
@@ -187,6 +193,15 @@ void UWeldMeshEdgesTool::Setup()
 
 	SourceMesh = MakeShared<FDynamicMesh3, ESPMode::ThreadSafe>(UE::ToolTarget::GetDynamicMeshCopy(Target, true));
 
+	// initialize selection if exists
+	if (HasGeometrySelection())
+	{
+		const FGeometrySelection& InputSelection = GetGeometrySelection();
+
+		UE::Geometry::EnumerateSelectionEdges(InputSelection, *SourceMesh,
+			[&](const int32 EdgeID) { SelectedEdges.Add(EdgeID); });
+	}
+
 	FTransform MeshTransform = (FTransform)UE::ToolTarget::GetLocalToWorldTransform(Target);
 
 	OperatorFactory = NewObject<UWeldMeshEdgesOperatorFactory>(this);
@@ -331,6 +346,7 @@ void UWeldMeshEdgesTool::UpdateOpParameters(FWeldMeshEdgesOp& Op) const
 	Op.SplitTangentsThreshold = Settings->SplitTangentsThreshold;
 	
 	Op.SourceMesh = SourceMesh;
+	Op.SelectedEdges = SelectedEdges;
 
 	FTransform LocalToWorld = (FTransform)UE::ToolTarget::GetLocalToWorldTransform(Target);
 	Op.SetTransform(LocalToWorld);
diff --git a/Engine/Plugins/Experimental/MeshModelingToolsetExp/Source/MeshModelingToolsExp/Public/WeldMeshEdgesTool.h b/Engine/Plugins/Experimental/MeshModelingToolsetExp/Source/MeshModelingToolsExp/Public/WeldMeshEdgesTool.h
--- a/Engine/Plugins/Experimental/MeshModelingToolsetExp/Source/MeshModelingToolsExp/Public/WeldMeshEdgesTool.h
+++ b/Engine/Plugins/Experimental/MeshModelingToolsetExp/Source/MeshModelingToolsExp/Public/WeldMeshEdgesTool.h
@@ -7,8 +7,10 @@
 #include "PreviewMesh.h"
 #include "ModelingOperators.h"
 #include "MeshOpPreviewHelpers.h"
+#include "BaseTools/SingleTargetWithSelectionTool.h"
 #include "DynamicMesh/DynamicMesh3.h"
 #include "DynamicMesh/DynamicMeshAABBTree3.h"
+#include "DynamicMesh/MeshSharingUtil.h"
 #include "WeldMeshEdgesTool.generated.h"
 
 // predeclarations
@@ -21,11 +23,13 @@ class FWeldMeshEdgesOp;
  *
  */
 UCLASS()
-class MESHMODELINGTOOLSEXP_API UWeldMeshEdgesToolBuilder : public USingleSelectionMeshEditingToolBuilder
+class MESHMODELINGTOOLSEXP_API UWeldMeshEdgesToolBuilder : public USingleTargetWithSelectionToolBuilder
 {
 	GENERATED_BODY()
 public:
-	virtual USingleSelectionMeshEditingTool* CreateNewTool(const FToolBuilderState& SceneState) const override;
+	virtual USingleTargetWithSelectionTool* CreateNewTool(const FToolBuilderState& SceneState) const override;
+
+	virtual bool RequiresInputSelection() const override { return false; };
 };
 
 UENUM()
@@ -110,7 +114,7 @@ public:
  * Mesh Weld Edges Tool
  */
 UCLASS()
-class MESHMODELINGTOOLSEXP_API UWeldMeshEdgesTool : public USingleSelectionMeshEditingTool
+class MESHMODELINGTOOLSEXP_API UWeldMeshEdgesTool : public USingleTargetWithSelectionTool
 {
 	GENERATED_BODY()
 public:
@@ -144,7 +148,9 @@ protected:
 protected:
 
 	TSharedPtr<UE::Geometry::FDynamicMesh3, ESPMode::ThreadSafe> SourceMesh;
-
+	
+	// If there is an active selection, SelectedEdges will be initialized
+	TSet<int32> SelectedEdges;
 
 
 
diff --git a/Engine/Plugins/Runtime/GeometryProcessing/Source/DynamicMesh/Private/Selections/GeometrySelectionUtil.cpp b/Engine/Plugins/Runtime/GeometryProcessing/Source/DynamicMesh/Private/Selections/GeometrySelectionUtil.cpp
--- a/Engine/Plugins/Runtime/GeometryProcessing/Source/DynamicMesh/Private/Selections/GeometrySelectionUtil.cpp
+++ b/Engine/Plugins/Runtime/GeometryProcessing/Source/DynamicMesh/Private/Selections/GeometrySelectionUtil.cpp
@@ -775,6 +775,192 @@ bool UE::Geometry::EnumeratePolygroupSelectionTriangles(
 	return true;
 }
 
+bool UE::Geometry::EnumerateSelectionEdges(
+	const FGeometrySelection& MeshSelection,
+	const UE::Geometry::FDynamicMesh3& Mesh,
+	TFunctionRef<void(int32)> EdgeFunc,
+	const UE::Geometry::FPolygroupSet* UseGroupSet
+)
+{
+	if (MeshSelection.TopologyType == EGeometryTopologyType::Triangle)
+	{
+		return EnumerateTriangleSelectionEdges(MeshSelection, Mesh, EdgeFunc);
+	}
+	else if (MeshSelection.TopologyType == EGeometryTopologyType::Polygroup)
+	{
+		return EnumeratePolygroupSelectionEdges(MeshSelection, Mesh, 
+			(UseGroupSet != nullptr) ? *UseGroupSet : FPolygroupSet(&Mesh), 
+			EdgeFunc);
+	}
+	return false;
+}
+
+
+bool UE::Geometry::EnumerateTriangleSelectionEdges(
+	const FGeometrySelection& MeshSelection,
+	const UE::Geometry::FDynamicMesh3& Mesh,
+	TFunctionRef<void(int32)> EdgeFunc)
+{
+	if ( ensure( MeshSelection.TopologyType == EGeometryTopologyType::Triangle ) == false )
+	{
+		return false;
+	}
+
+	if (MeshSelection.ElementType == EGeometryElementType::Face)
+	{
+		for (const uint64 TriangleID : MeshSelection.Selection)
+		{
+			if (Mesh.IsTriangle((int32)TriangleID))
+			{
+				FIndex3i TriEdges = Mesh.GetTriEdges((int32)TriangleID);
+				EdgeFunc(TriEdges[0]);
+				EdgeFunc(TriEdges[1]);
+				EdgeFunc(TriEdges[2]);
+			}
+		}
+	}
+	else if (MeshSelection.ElementType == EGeometryElementType::Edge)
+	{
+		for (const uint64 EncodedID : MeshSelection.Selection)
+		{
+			const FMeshTriEdgeID TriEdgeID( FGeoSelectionID(EncodedID).GeometryID );
+			const int32 EdgeID = Mesh.IsTriangle(TriEdgeID.TriangleID) ? Mesh.GetTriEdge(TriEdgeID.TriangleID, TriEdgeID.TriEdgeIndex) : IndexConstants::InvalidID;
+			if (Mesh.IsEdge((int32)EdgeID))
+			{
+				EdgeFunc((int32)EdgeID);
+			}
+		}
+	}
+	else if (MeshSelection.ElementType == EGeometryElementType::Vertex)
+	{
+		for (const uint64 VertexID : MeshSelection.Selection)
+		{
+			Mesh.EnumerateVertexEdges((int32)VertexID, [&EdgeFunc](int32 EdgeID)
+			{
+				EdgeFunc(EdgeID);
+			});
+		}
+	}
+	else
+	{
+		return false;
+	}
+
+	return true;
+}
+
+
+bool UE::Geometry::EnumeratePolygroupSelectionEdges(
+	const FGeometrySelection& MeshSelection,
+	const UE::Geometry::FDynamicMesh3& Mesh,
+	const UE::Geometry::FPolygroupSet& GroupSet,
+	TFunctionRef<void(int32)> EdgeFunc
+)
+{
+	if ( ensure( MeshSelection.TopologyType == EGeometryTopologyType::Polygroup ) == false )
+	{
+		return false;
+	}
+
+	TArray<int32> SeedGroups;
+	TArray<int32> SeedTriangles;
+	TSet<int32> UniqueSeedGroups;
+
+	// TODO: the code below will not work correctly if the selection contains
+	// multiple disconnected-components with the same GroupID. They will be
+	// filtered out by the UniqueSeedGroups test. Seems like it will be necessary
+	// to detect this case up-front and do something more expensive, like filtering
+	// out duplicates inside the connected-components loop instead of up-front
+
+	if (MeshSelection.ElementType == EGeometryElementType::Face)
+	{
+		for (const uint64 EncodedID : MeshSelection.Selection)
+		{
+			const FGeoSelectionID SelectionID(EncodedID);
+			int32 SeedTriangleID = (int32)SelectionID.GeometryID;
+			if (Mesh.IsTriangle(SeedTriangleID))
+			{
+				int32 GroupID = GroupSet.GetGroup(SeedTriangleID);
+				check(GroupID == (int32)SelectionID.TopologyID);		// sanity-check that we are using the right group
+				if ( GroupID >= 0 && UniqueSeedGroups.Contains(GroupID) == false)
+				{
+					UniqueSeedGroups.Add(GroupID);
+					SeedGroups.Add(GroupID);
+					SeedTriangles.Add(SeedTriangleID);
+				}
+			}
+		}
+	}
+	else if (MeshSelection.ElementType == EGeometryElementType::Edge)
+	{
+		for (const uint64 EncodedID : MeshSelection.Selection)
+		{
+			const FMeshTriEdgeID TriEdgeID( FGeoSelectionID(EncodedID).GeometryID );
+			const int32 SeedEdgeID = Mesh.IsTriangle(TriEdgeID.TriangleID) ? Mesh.GetTriEdge(TriEdgeID.TriangleID, TriEdgeID.TriEdgeIndex) : IndexConstants::InvalidID;
+			if (Mesh.IsEdge(SeedEdgeID))
+			{
+				Mesh.EnumerateEdgeTriangles(SeedEdgeID, [&](int32 TriangleID)
+				{
+					const int32 GroupID = GroupSet.GetGroup(TriangleID);
+					if (GroupID >= 0 && UniqueSeedGroups.Contains(GroupID) == false)
+					{
+						UniqueSeedGroups.Add(GroupID);
+						SeedGroups.Add(GroupID);
+						SeedTriangles.Add(TriangleID);
+					}
+				});
+			}
+		}
+	}
+	else if (MeshSelection.ElementType == EGeometryElementType::Vertex)
+	{
+		for (const uint64 EncodedID : MeshSelection.Selection)
+		{
+			const int32 VertexID = (int32)FGeoSelectionID(EncodedID).GeometryID;
+			if (Mesh.IsVertex(VertexID))
+			{
+				Mesh.EnumerateVertexTriangles(VertexID, [&](int32 TriangleID)
+				{
+					const int32 GroupID = GroupSet.GetGroup(TriangleID);
+					if (GroupID >= 0 && UniqueSeedGroups.Contains(GroupID) == false)
+					{
+						UniqueSeedGroups.Add(GroupID);
+						SeedGroups.Add(GroupID);
+						SeedTriangles.Add(TriangleID);
+					}
+				});
+			}
+		}
+	}
+	else
+	{
+		return false;
+	}
+
+
+	TSet<int> TempROI;		// if we could provide this as input we would not need a temporary roi...
+	TArray<int32> QueueBuffer;
+	const int32 NumGroups = SeedGroups.Num();
+	for (int32 k = 0; k < NumGroups; ++k)
+	{
+		check(GroupSet.GetGroup(SeedTriangles[k]) == SeedGroups[k]);
+		int32 GroupID = SeedGroups[k];
+		FMeshConnectedComponents::GrowToConnectedTriangles(&Mesh, 
+			TArray<int>{SeedTriangles[k]}, TempROI, &QueueBuffer, 
+			[&](const int32 T1, const int32 T2) { return GroupSet.GetGroup(T2) == GroupID; });
+
+		// EnumeratePolygroupSelectionEdges is the same as EnumeratePolygroupSelectionTriangles up until this point
+		for (const int32 TID : TempROI)
+		{
+			FIndex3i TriEdges = Mesh.GetTriEdges((int32)TID);
+			EdgeFunc(TriEdges[0]);
+			EdgeFunc(TriEdges[1]);
+			EdgeFunc(TriEdges[2]);
+		}
+	}
+
+	return true;
+}
 
 
 
diff --git a/Engine/Plugins/Runtime/GeometryProcessing/Source/DynamicMesh/Public/Selections/GeometrySelectionUtil.h b/Engine/Plugins/Runtime/GeometryProcessing/Source/DynamicMesh/Public/Selections/GeometrySelectionUtil.h
--- a/Engine/Plugins/Runtime/GeometryProcessing/Source/DynamicMesh/Public/Selections/GeometrySelectionUtil.h
+++ b/Engine/Plugins/Runtime/GeometryProcessing/Source/DynamicMesh/Public/Selections/GeometrySelectionUtil.h
@@ -143,6 +143,40 @@ DYNAMICMESH_API bool EnumeratePolygroupSelectionTriangles(
 	const UE::Geometry::FPolygroupSet& GroupSet,
 	TFunctionRef<void(int32)> TriangleFunc
 );
+/**
+ * Call EdgeFunc for each mesh EdgeID included in MeshSelection.
+ * EdgeFunc may be called multiple times for the same EdgeID.
+ * This will forward to EnumerateTriangleSelectionEdges() or 
+ * EnumeratePolygroupSelectionEdges() depending on the selection topology type.
+ * If UseGroupSet and MeshSelection is for polygroups, the default Mesh group layer will be used.
+ */
+DYNAMICMESH_API bool EnumerateSelectionEdges(
+	const FGeometrySelection& MeshSelection,
+	const UE::Geometry::FDynamicMesh3& Mesh,
+	TFunctionRef<void(int32)> EdgeFunc,
+	const UE::Geometry::FPolygroupSet* UseGroupSet = nullptr
+);
+/**
+ * Call EdgeFunc for each mesh EdgeID included in MeshSelection.
+ * For Triangles, all 3 edges are included.
+ * For Vertices, all edges touching each selected vertex are included.
+ */
+DYNAMICMESH_API bool EnumerateTriangleSelectionEdges(
+	const FGeometrySelection& MeshSelection,
+	const UE::Geometry::FDynamicMesh3& Mesh,
+	TFunctionRef<void(int32)> EdgeFunc
+);
+/**
+ * Call EdgeFunc for each mesh EdgeID included in MeshSelection, where MeshSelection has polygroup topology.
+ * For Polygroup Faces, all edges in the face are included.
+ * For Polygroup Corners, currently includes all edges in any group touching the corner
+ */
+DYNAMICMESH_API bool EnumeratePolygroupSelectionEdges(
+	const FGeometrySelection& MeshSelection,
+	const UE::Geometry::FDynamicMesh3& Mesh,
+	const UE::Geometry::FPolygroupSet& GroupSet,
+	TFunctionRef<void(int32)> EdgeFunc
+);
 
 
 /**
diff --git a/Engine/Source/Runtime/GeometryCore/Private/DynamicMesh/Operations/MergeCoincidentMeshEdges.cpp b/Engine/Source/Runtime/GeometryCore/Private/DynamicMesh/Operations/MergeCoincidentMeshEdges.cpp
--- a/Engine/Source/Runtime/GeometryCore/Private/DynamicMesh/Operations/MergeCoincidentMeshEdges.cpp
+++ b/Engine/Source/Runtime/GeometryCore/Private/DynamicMesh/Operations/MergeCoincidentMeshEdges.cpp
@@ -166,32 +166,40 @@ bool FMergeCoincidentMeshEdges::Apply()
 				continue;
 			}
 
-			FDynamicMesh3::FMergeEdgesInfo mergeInfo;
-			EMeshResult result = Mesh->MergeEdges(eid, other_eid, mergeInfo);
-			if (result != EMeshResult::Ok) 
+			// When there is no geometry selection, EdgesToMerge is never initialized
+			bool bWeldingAcrossEntireMesh = (EdgesToMerge == nullptr);
+
+			// Edges only considered for merging if we're welding the entire mesh or, when using a geometry selection,
+			// if EITHER edge in the Match are a part of the selection
+			if (bWeldingAcrossEntireMesh || EdgesToMerge->Contains(eid) || EdgesToMerge->Contains(other_eid))
 			{
-				// if the operation failed we remove this edge from the equivalence set
-				Matches.RemoveAt(i);
-				i--;
+				FDynamicMesh3::FMergeEdgesInfo MergeInfo;
+				EMeshResult Result = Mesh->MergeEdges(eid, other_eid, MergeInfo);
+				if (Result != EMeshResult::Ok) 
+				{
+					// if the operation failed we remove this edge from the equivalence set
+					Matches.RemoveAt(i);
+					i--;
 
-				EquivalenceSets[other_eid]->Remove(eid);
-				//DuplicatesQueue.UpdatePriority(...);  // should we do this?
+					EquivalenceSets[other_eid]->Remove(eid);
+					//DuplicatesQueue.UpdatePriority(...);  // should we do this?
 
-				FailedCount++;
-			}
-			else 
-			{
-				// ok we merged, other edge is no longer free
-				bMerged = true;
-				delete EquivalenceSets[other_eid];
-				EquivalenceSets[other_eid] = nullptr;
-				RemainingEdges.Remove(other_eid);
-
-				// weld attributes 
-				if (bWeldAttrsOnMergedEdges)
-				{ 
-					SplitAttributeWelder.WeldSplitElements(*Mesh, mergeInfo.KeptVerts[0]);
-					SplitAttributeWelder.WeldSplitElements(*Mesh, mergeInfo.KeptVerts[1]);
+					FailedCount++;
+				}
+				else 
+				{
+					// ok we merged, other edge is no longer free
+					bMerged = true;
+					delete EquivalenceSets[other_eid];
+					EquivalenceSets[other_eid] = nullptr;
+					RemainingEdges.Remove(other_eid);
+
+					// weld attributes 
+					if (bWeldAttrsOnMergedEdges)
+					{ 
+						SplitAttributeWelder.WeldSplitElements(*Mesh, MergeInfo.KeptVerts[0]);
+						SplitAttributeWelder.WeldSplitElements(*Mesh, MergeInfo.KeptVerts[1]);
+					}
 				}
 			}
 		}
diff --git a/Engine/Source/Runtime/GeometryCore/Public/DynamicMesh/Operations/MergeCoincidentMeshEdges.h b/Engine/Source/Runtime/GeometryCore/Public/DynamicMesh/Operations/MergeCoincidentMeshEdges.h
--- a/Engine/Source/Runtime/GeometryCore/Public/DynamicMesh/Operations/MergeCoincidentMeshEdges.h
+++ b/Engine/Source/Runtime/GeometryCore/Public/DynamicMesh/Operations/MergeCoincidentMeshEdges.h
@@ -58,7 +58,9 @@ public:
 
 	/** Used to weld attributes at the merged edges */
 	FSplitAttributeWelder  SplitAttributeWelder;
-
+	
+	/** Set of EdgesIds to merge. If null, perform merge operation across entire mesh */
+	TSet<int32>* EdgesToMerge = nullptr;
 
 public:
 	FMergeCoincidentMeshEdges(FDynamicMesh3* mesh) : Mesh(mesh)
