Modeling Tools: Add support for customizing colors in Mesh Element Selection, add support for different between hovering over a selected vs non-selected elements, fix issues with EnumerateSelectionEdges

diff --git a/Engine/Plugins/Editor/ModelingToolsEditorMode/Source/ModelingToolsEditorMode/Private/ModelingToolsEditorMode.cpp b/Engine/Plugins/Editor/ModelingToolsEditorMode/Source/ModelingToolsEditorMode/Private/ModelingToolsEditorMode.cpp
--- a/Engine/Plugins/Editor/ModelingToolsEditorMode/Source/ModelingToolsEditorMode/Private/ModelingToolsEditorMode.cpp
+++ b/Engine/Plugins/Editor/ModelingToolsEditorMode/Source/ModelingToolsEditorMode/Private/ModelingToolsEditorMode.cpp
@@ -883,6 +883,7 @@ void UModelingToolsEditorMode::Enter()
 						ModelingEditorSettings->LastMeshSelectionTopologyMode = static_cast<int>(TopoMode);
 						ModelingEditorSettings->LastMeshSelectionElementType = static_cast<int>(ElementMode);
 						ModelingEditorSettings->SaveConfig();
+						SelectionManager->SetSelectionColors(ModelingEditorSettings->UnselectedColor, ModelingEditorSettings->HoverOverSelectedColor, ModelingEditorSettings->HoverOverUnselectedColor, ModelingEditorSettings->GeometrySelectedColor);
 					}
 				}
 			}),
diff --git a/Engine/Plugins/Editor/ModelingToolsEditorMode/Source/ModelingToolsEditorMode/Public/ModelingToolsEditorModeSettings.h b/Engine/Plugins/Editor/ModelingToolsEditorMode/Source/ModelingToolsEditorMode/Public/ModelingToolsEditorModeSettings.h
--- a/Engine/Plugins/Editor/ModelingToolsEditorMode/Source/ModelingToolsEditorMode/Public/ModelingToolsEditorModeSettings.h
+++ b/Engine/Plugins/Editor/ModelingToolsEditorMode/Source/ModelingToolsEditorMode/Public/ModelingToolsEditorModeSettings.h
@@ -351,6 +351,30 @@ public:
 	UPROPERTY(config, EditAnywhere, Category = "Modeling Mode|UI Customization")
 	bool bAlwaysShowToolButtons = false;
 
+	/**
+	 * Custom color for when geometry is not a part of current selection when using Mesh Element Selection
+	 */
+	UPROPERTY(Config, EditAnywhere, Category = "Modeling Mode|Mesh Element Selection Mode")
+	FLinearColor UnselectedColor = FLinearColor(.906f, 0.038f, 0.038f); // default red
+
+	/**
+	 * Custom color hovering over selected geometry in Mesh Element Selection
+	 */
+	UPROPERTY(Config, EditAnywhere, Category = "Modeling Mode|Mesh Element Selection Mode")
+	FLinearColor HoverOverSelectedColor = FLinearColor(0.791f, 0.262f, 1.0f); // default pink
+
+	/**
+	* Custom color hovering over unselected geometry in Mesh Element Selection
+	*/
+	UPROPERTY(Config, EditAnywhere, Category = "Modeling Mode|Mesh Element Selection Mode")
+	FLinearColor HoverOverUnselectedColor = FLinearColor(0.027f, 0.656f, 0.027f); // default green
+
+	/**
+	* Custom color for when geometry is a part of the current selection when using Mesh Element Selection
+	*/
+	UPROPERTY(Config, EditAnywhere, Category = "Modeling Mode|Mesh Element Selection Mode")
+	FLinearColor GeometrySelectedColor = FLinearColor(0.859f, 0.799f, 0.076f); // default yellow
+
 public:
 
 	// saved-state for various mode settings that are configured via UI toggles/etc, and not exposed in settings dialog
diff --git a/Engine/Plugins/Runtime/GeometryProcessing/Source/DynamicMesh/Private/Selections/GeometrySelectionUtil.cpp b/Engine/Plugins/Runtime/GeometryProcessing/Source/DynamicMesh/Private/Selections/GeometrySelectionUtil.cpp
--- a/Engine/Plugins/Runtime/GeometryProcessing/Source/DynamicMesh/Private/Selections/GeometrySelectionUtil.cpp
+++ b/Engine/Plugins/Runtime/GeometryProcessing/Source/DynamicMesh/Private/Selections/GeometrySelectionUtil.cpp
@@ -88,8 +88,6 @@ bool EnumerateFaceElementSelectionTriangles(
 	return true;
 }
 
-// We don't currently have an EnumerateSelectionEdges. If we build out the below to handle
-// other selection types (vertex, face), we might want to expose it.
 /**
  * Given a selection with elements of type EGeometryElementType::Edge, call EdgeFunc on
  * each mesh edge (with the Eid passed in) of the selection.
@@ -862,11 +860,13 @@ bool UE::Geometry::EnumeratePolygroupSelectionEdges(
 		return false;
 	}
 
-	TArray<int32> SeedGroups;
+	TArray<int32> SeedTriGroups;
 	TArray<int32> SeedTriangles;
-	TSet<int32> UniqueSeedGroups;
+	TSet<int32> UniqueSeedTriGroups;
+	
+	TArray<int32> SeedEdges;
 
-	// TODO: the code below will not work correctly if the selection contains
+	// TODO: the face code below will not work correctly if the selection contains
 	// multiple disconnected-components with the same GroupID. They will be
 	// filtered out by the UniqueSeedGroups test. Seems like it will be necessary
 	// to detect this case up-front and do something more expensive, like filtering
@@ -877,15 +877,15 @@ bool UE::Geometry::EnumeratePolygroupSelectionEdges(
 		for (const uint64 EncodedID : MeshSelection.Selection)
 		{
 			const FGeoSelectionID SelectionID(EncodedID);
-			int32 SeedTriangleID = (int32)SelectionID.GeometryID;
+			const int32 SeedTriangleID = (int32)SelectionID.GeometryID;
 			if (Mesh.IsTriangle(SeedTriangleID))
 			{
 				int32 GroupID = GroupSet.GetGroup(SeedTriangleID);
 				check(GroupID == (int32)SelectionID.TopologyID);		// sanity-check that we are using the right group
-				if ( GroupID >= 0 && UniqueSeedGroups.Contains(GroupID) == false)
+				if ( GroupID >= 0 && UniqueSeedTriGroups.Contains(GroupID) == false)
 				{
-					UniqueSeedGroups.Add(GroupID);
-					SeedGroups.Add(GroupID);
+					UniqueSeedTriGroups.Add(GroupID);
+					SeedTriGroups.Add(GroupID);
 					SeedTriangles.Add(SeedTriangleID);
 				}
 			}
@@ -895,20 +895,12 @@ bool UE::Geometry::EnumeratePolygroupSelectionEdges(
 	{
 		for (const uint64 EncodedID : MeshSelection.Selection)
 		{
+			// record selected edges, need to find other edges that are part of the same polygroup edge to include in selection
 			const FMeshTriEdgeID TriEdgeID( FGeoSelectionID(EncodedID).GeometryID );
 			const int32 SeedEdgeID = Mesh.IsTriangle(TriEdgeID.TriangleID) ? Mesh.GetTriEdge(TriEdgeID.TriangleID, TriEdgeID.TriEdgeIndex) : IndexConstants::InvalidID;
 			if (Mesh.IsEdge(SeedEdgeID))
 			{
-				Mesh.EnumerateEdgeTriangles(SeedEdgeID, [&](int32 TriangleID)
-				{
-					const int32 GroupID = GroupSet.GetGroup(TriangleID);
-					if (GroupID >= 0 && UniqueSeedGroups.Contains(GroupID) == false)
-					{
-						UniqueSeedGroups.Add(GroupID);
-						SeedGroups.Add(GroupID);
-						SeedTriangles.Add(TriangleID);
-					}
-				});
+				SeedEdges.Add(SeedEdgeID);
 			}
 		}
 	}
@@ -919,14 +911,11 @@ bool UE::Geometry::EnumeratePolygroupSelectionEdges(
 			const int32 VertexID = (int32)FGeoSelectionID(EncodedID).GeometryID;
 			if (Mesh.IsVertex(VertexID))
 			{
-				Mesh.EnumerateVertexTriangles(VertexID, [&](int32 TriangleID)
+				Mesh.EnumerateVertexEdges(VertexID, [&Mesh, &SeedEdges](const int32 EdgeID)
 				{
-					const int32 GroupID = GroupSet.GetGroup(TriangleID);
-					if (GroupID >= 0 && UniqueSeedGroups.Contains(GroupID) == false)
+					if (Mesh.IsEdge(EdgeID))
 					{
-						UniqueSeedGroups.Add(GroupID);
-						SeedGroups.Add(GroupID);
-						SeedTriangles.Add(TriangleID);
+						SeedEdges.Add(EdgeID);
 					}
 				});
 			}
@@ -940,24 +929,91 @@ bool UE::Geometry::EnumeratePolygroupSelectionEdges(
 
 	TSet<int> TempROI;		// if we could provide this as input we would not need a temporary roi...
 	TArray<int32> QueueBuffer;
-	const int32 NumGroups = SeedGroups.Num();
-	for (int32 k = 0; k < NumGroups; ++k)
+
+	// Edge Type: currently enumerates all edges that are part of selected poly edge(s).
+	// Vertex Type: currently enumerates all edges that are a part of a polyedge which contains the selected vertex/vertices
+	// neither includes non poly-group edges
+	if (MeshSelection.ElementType == EGeometryElementType::Vertex || MeshSelection.ElementType == EGeometryElementType::Edge)
 	{
-		check(GroupSet.GetGroup(SeedTriangles[k]) == SeedGroups[k]);
-		int32 GroupID = SeedGroups[k];
-		FMeshConnectedComponents::GrowToConnectedTriangles(&Mesh, 
-			TArray<int>{SeedTriangles[k]}, TempROI, &QueueBuffer, 
-			[&](const int32 T1, const int32 T2) { return GroupSet.GetGroup(T2) == GroupID; });
+		const int32 NumEdges = SeedEdges.Num();
+		for (int32 j = 0; j < NumEdges; j++)
+		{
+			TArray<int32> EdgeGroups; // the 1 or 2 groups which an edge belongs to
+			// retrieves the 1 or 2 triangles and groups to which a selected edge belongs, and finds which group(s) the triangles belong to
+			Mesh.EnumerateEdgeTriangles(SeedEdges[j], [&GroupSet, &EdgeGroups](int32 TriangleID)
+			{
+				const int32 GroupID = GroupSet.GetGroup(TriangleID);
+				if (GroupID >= 0)
+				{
+					EdgeGroups.Add(GroupID);
+				}
+			});
+
+			EdgeGroups.Sort();
+
+			// if an edge's 2 triangles are from the same group, it is not a Polygroup boundary edge and can be disregarded
+			const bool bIsInnerEdge = EdgeGroups.Num() == 2 && EdgeGroups[0] == EdgeGroups[1];
 
-		// EnumeratePolygroupSelectionEdges is the same as EnumeratePolygroupSelectionTriangles up until this point
-		for (const int32 TID : TempROI)
+			if (!bIsInnerEdge)
+			{
+				// finds connected edges by looking at all edges in a mesh and retrieving which 2 (or 1) groups they belong to
+				// if an edge's 2 groups are the same as the selected edges' 2 groups, we know they are part of the same poly edge
+				FMeshConnectedComponents::GrowToConnectedEdges(Mesh, TArray<int>{SeedEdges[j]}, TempROI, &QueueBuffer, 
+						[&Mesh, &GroupSet, &EdgeGroups](const int32 E1, const int32 E2)
+						{
+							TArray<int32> OtherEdgeGroups;
+							Mesh.EnumerateEdgeTriangles(E2, [&GroupSet, &OtherEdgeGroups](const int32 TriangleID)
+							{
+								const int GrpID = GroupSet.GetGroup(TriangleID);
+								if (GrpID >= 0)
+								{
+									OtherEdgeGroups.Add(GrpID);
+								}
+							});
+							OtherEdgeGroups.Sort();
+
+							// if selected edge groups and current edge groups are the same, they belong to the same poly edge
+							// note: currently if a border edge is selected (i.e an edge has only 1 triangle and 1 group), function will include
+							// all other border edges of the polygroup
+							return  OtherEdgeGroups.Num() == EdgeGroups.Num() && OtherEdgeGroups == EdgeGroups;
+						});
+				// apply edge function to all edges in poly edge
+				for (const int32 EdgeID : TempROI)
+				{
+					EdgeFunc(EdgeID);
+				}
+			}
+			
+		}
+	}
+	else if (MeshSelection.ElementType == EGeometryElementType::Face)
+	{
+		const int32 NumGroups = SeedTriGroups.Num();
+		for (int32 k = 0; k < NumGroups; ++k)
 		{
-			FIndex3i TriEdges = Mesh.GetTriEdges((int32)TID);
-			EdgeFunc(TriEdges[0]);
-			EdgeFunc(TriEdges[1]);
-			EdgeFunc(TriEdges[2]);
+			check(GroupSet.GetGroup(SeedTriangles[k]) == SeedTriGroups[k]);
+			const int32 GroupID = SeedTriGroups[k];
+			FMeshConnectedComponents::GrowToConnectedTriangles(&Mesh, 
+				TArray<int>{SeedTriangles[k]}, TempROI, &QueueBuffer, 
+				[&GroupSet, &GroupID](const int32 T1, const int32 T2)
+				{
+					return GroupSet.GetGroup(T2) == GroupID;
+				});
+			
+			for (const int32 TID : TempROI)
+			{
+				FIndex3i TriEdges = Mesh.GetTriEdges(TID);
+				EdgeFunc(TriEdges[0]);
+				EdgeFunc(TriEdges[1]);
+				EdgeFunc(TriEdges[2]);
+			}
+		
 		}
 	}
+	else
+	{
+		return false;
+	}
 
 	return true;
 }
@@ -971,7 +1027,28 @@ bool UE::Geometry::EnumerateTriangleSelectionElements(
 	TFunctionRef<void(int32, const FSegment3d&)> EdgeFunc,
 	TFunctionRef<void(int32, const FTriangle3d&)> TriangleFunc,
 	const FTransform* ApplyTransform,
-	bool bMapFacesToEdgeLoops
+	const bool bMapFacesToEdgeLoops
+)
+{
+	return EnumerateTriangleSelectionElements(
+		MeshSelection,
+		Mesh,
+		VertexFunc,
+		EdgeFunc,
+		TriangleFunc,
+		ApplyTransform,
+		EEnumerateMappingFlags::Default | (bMapFacesToEdgeLoops ? EEnumerateMappingFlags::FacesToEdges : EEnumerateMappingFlags::None)
+		);
+}
+
+bool UE::Geometry::EnumerateTriangleSelectionElements(
+	const FGeometrySelection& MeshSelection,
+	const UE::Geometry::FDynamicMesh3& Mesh,
+	TFunctionRef<void(int32, const FVector3d&)> VertexFunc,
+	TFunctionRef<void(int32, const FSegment3d&)> EdgeFunc,
+	TFunctionRef<void(int32, const FTriangle3d&)> TriangleFunc,
+	const FTransform* ApplyTransform,
+	const EEnumerateMappingFlags Flags
 )
 {
 	if ( ensure( MeshSelection.TopologyType == EGeometryTopologyType::Triangle ) == false )
@@ -981,9 +1058,9 @@ bool UE::Geometry::EnumerateTriangleSelectionElements(
 
 	if (MeshSelection.ElementType == EGeometryElementType::Face)
 	{
-		for (uint64 EncodedID : MeshSelection.Selection)
+		for (const uint64 EncodedID : MeshSelection.Selection)
 		{
-			int32 TriangleID = (int32)FGeoSelectionID(EncodedID).GeometryID;
+			const int32 TriangleID = (int32)FGeoSelectionID(EncodedID).GeometryID;
 			if (Mesh.IsTriangle(TriangleID))
 			{
 				FVector3d A, B, C;
@@ -994,26 +1071,32 @@ bool UE::Geometry::EnumerateTriangleSelectionElements(
 					B = ApplyTransform->TransformPosition(B);
 					C = ApplyTransform->TransformPosition(C);
 				}
-				if (bMapFacesToEdgeLoops)
+				
+				// Call TriangleFunc if we are mapping Faces to Faces
+				if ((Flags & EEnumerateMappingFlags::FacesToFaces) != EEnumerateMappingFlags::None)
 				{
-					FIndex3i Edges = Mesh.GetTriEdges((int32)TriangleID);
+					TriangleFunc(TriangleID, FTriangle3d(A, B, C));
+				}
+				
+				// Call EdgeFunc for each of the triangle's edges if we are mapping Faces to Edges
+				if ((Flags & EEnumerateMappingFlags::FacesToEdges) != EEnumerateMappingFlags::None)
+				{
+					const FIndex3i Edges = Mesh.GetTriEdges((int32)TriangleID);
 					EdgeFunc(Edges.A, FSegment3d(A, B));
 					EdgeFunc(Edges.B, FSegment3d(B, C));
 					EdgeFunc(Edges.C, FSegment3d(C, A));
 				}
-				else
-				{
-					TriangleFunc(TriangleID, FTriangle3d(A, B, C));
-				}
+
+				// TODO [katie.caputo] add FacesToVerts path
 			}
 		}
 	}
 	else if (MeshSelection.ElementType == EGeometryElementType::Edge)
 	{
-		for (uint64 EncodedID : MeshSelection.Selection)
+		for (const uint64 EncodedID : MeshSelection.Selection)
 		{
-			FMeshTriEdgeID TriEdgeID( FGeoSelectionID(EncodedID).GeometryID );
-			int32 EdgeID = Mesh.IsTriangle(TriEdgeID.TriangleID) ? Mesh.GetTriEdge(TriEdgeID.TriangleID, TriEdgeID.TriEdgeIndex) : IndexConstants::InvalidID;
+			const FMeshTriEdgeID TriEdgeID( FGeoSelectionID(EncodedID).GeometryID );
+			const int32 EdgeID = Mesh.IsTriangle(TriEdgeID.TriangleID) ? Mesh.GetTriEdge(TriEdgeID.TriangleID, TriEdgeID.TriEdgeIndex) : IndexConstants::InvalidID;
 			if (Mesh.IsEdge(EdgeID))
 			{
 				FVector3d A, B;
@@ -1023,19 +1106,33 @@ bool UE::Geometry::EnumerateTriangleSelectionElements(
 					A = ApplyTransform->TransformPosition(A);
 					B = ApplyTransform->TransformPosition(B);
 				}
-				EdgeFunc(EdgeID, FSegment3d(A, B));
+
+				// Call EdgeFunc if we are mapping Edges to Edges
+				if ((Flags & EEnumerateMappingFlags::EdgesToEdges) != EEnumerateMappingFlags::None)
+				{
+					EdgeFunc(EdgeID, FSegment3d(A, B));
+				}
+				// TODO [katie.caputo] add EdgesToFaces and EdgesToVerts path
 			}
 		}
+		
 	}
 	else if (MeshSelection.ElementType == EGeometryElementType::Vertex)
 	{
-		for (uint64 EncodedID : MeshSelection.Selection)
+		for (const uint64 EncodedID : MeshSelection.Selection)
 		{
-			int32 VertexID = (int32)FGeoSelectionID(EncodedID).GeometryID;
+			const int32 VertexID = (int32)FGeoSelectionID(EncodedID).GeometryID;
 			if (Mesh.IsVertex(VertexID))
 			{
-				FVector3d A = Mesh.GetVertex(VertexID);
-				VertexFunc(VertexID, (ApplyTransform != nullptr) ? ApplyTransform->TransformPosition(A) : A);
+				const FVector3d A = Mesh.GetVertex(VertexID);
+
+				// Call VertexFunc if we are mapping Vertices to Vertices
+				if ((Flags & EEnumerateMappingFlags::VertsToVerts) != EEnumerateMappingFlags::None)
+				{
+					VertexFunc(VertexID, (ApplyTransform != nullptr) ? ApplyTransform->TransformPosition(A) : A);
+				}
+
+				// TODO [katie.caputo] add VertsToFaces and VertsToEdges path
 			}
 		}
 	}
@@ -1047,9 +1144,28 @@ bool UE::Geometry::EnumerateTriangleSelectionElements(
 	return true;
 }
 
-
-
-
+bool UE::Geometry::EnumeratePolygroupSelectionElements(
+	const FGeometrySelection& MeshSelection,
+	const UE::Geometry::FDynamicMesh3& Mesh,
+	const FGroupTopology* GroupTopology,
+	TFunctionRef<void(int32, const FVector3d&)> VertexFunc,
+	TFunctionRef<void(int32, const FSegment3d&)> EdgeFunc,
+	TFunctionRef<void(int32, const FTriangle3d&)> TriangleFunc,
+	const FTransform* ApplyTransform,
+	const bool bMapFacesToEdgeLoops
+)
+{
+	return EnumeratePolygroupSelectionElements(
+		MeshSelection,
+		Mesh,
+		GroupTopology,
+		VertexFunc,
+		EdgeFunc,
+		TriangleFunc,
+		ApplyTransform,
+		EEnumerateMappingFlags::Default | (bMapFacesToEdgeLoops ? EEnumerateMappingFlags::FacesToEdges : EEnumerateMappingFlags::None)
+		);
+}
 
 bool UE::Geometry::EnumeratePolygroupSelectionElements(
 	const FGeometrySelection& MeshSelection,
@@ -1059,17 +1175,17 @@ bool UE::Geometry::EnumeratePolygroupSelectionElements(
 	TFunctionRef<void(int32, const FSegment3d&)> EdgeFunc,
 	TFunctionRef<void(int32, const FTriangle3d&)> TriangleFunc,
 	const FTransform* ApplyTransform,
-	bool bMapFacesToEdgeLoops
+	const EEnumerateMappingFlags Flags
 )
 {
-	if ( ensure( MeshSelection.TopologyType == EGeometryTopologyType::Polygroup ) == false )
+if ( ensure( MeshSelection.TopologyType == EGeometryTopologyType::Polygroup ) == false )
 	{
 		return false;
 	}
 
-	auto ProcessGroupEdgeID = [GroupTopology, &Mesh, ApplyTransform, &EdgeFunc](int32 GroupEdgeID)
+	auto ProcessGroupEdgeID = [GroupTopology, &Mesh, ApplyTransform, &EdgeFunc, Flags](const int32 GroupEdgeID)
 	{
-		for (int32 EdgeID : GroupTopology->GetGroupEdgeEdges(GroupEdgeID))
+		for (const int32 EdgeID : GroupTopology->GetGroupEdgeEdges(GroupEdgeID))
 		{
 			FVector3d A, B;
 			Mesh.GetEdgeV(EdgeID, A, B);
@@ -1084,40 +1200,19 @@ bool UE::Geometry::EnumeratePolygroupSelectionElements(
 
 	if (MeshSelection.ElementType == EGeometryElementType::Face)
 	{
-		if (bMapFacesToEdgeLoops)
-		{
-			TArray<int32> GroupEdgeIDs;
-			for (uint64 EncodedID : MeshSelection.Selection)
-			{
-				int32 GroupID = FGeoSelectionID(EncodedID).TopologyID;
-				if ( const FGroupTopology::FGroup* Group = GroupTopology->FindGroupByID(GroupID) )
-				{
-					for (const FGroupTopology::FGroupBoundary& Boundary : Group->Boundaries)
-					{
-						for (int32 GroupEdgeID : Boundary.GroupEdges)
-						{
-							GroupEdgeIDs.AddUnique(GroupEdgeID);
-						}
-					}
-				}
-			}
-			for (int32 GroupEdgeID : GroupEdgeIDs)
-			{
-				ProcessGroupEdgeID(GroupEdgeID);
-			}
-		}
-		else
+		for (const uint64 EncodedID : MeshSelection.Selection)
 		{
-			for (uint64 EncodedID : MeshSelection.Selection)
+			FGeoSelectionID SelectionID(EncodedID);
+			const int32 SeedTriangleID = (int32)SelectionID.GeometryID, GroupID = (int32)SelectionID.TopologyID;
+			if (Mesh.IsTriangle(SeedTriangleID))
 			{
-				FGeoSelectionID SelectionID(EncodedID);
-				int32 SeedTriangleID = (int32)SelectionID.GeometryID, GroupID = (int32)SelectionID.TopologyID;
-				if (Mesh.IsTriangle(SeedTriangleID))
+				// Call TriangleFunc if we are mapping Faces to Faces
+				if ((Flags & EEnumerateMappingFlags::FacesToFaces) != EEnumerateMappingFlags::None)
 				{
-					for (int32 TriangleID : GroupTopology->GetGroupFaces(GroupID))
+					for (const int32 TriangleID : GroupTopology->GetGroupFaces(GroupID))
 					{
 						FVector3d A, B, C;
-						Mesh.GetTriVertices((int32)TriangleID, A, B, C);
+						Mesh.GetTriVertices(TriangleID, A, B, C);
 						if (ApplyTransform != nullptr)
 						{
 							A = ApplyTransform->TransformPosition(A);
@@ -1127,31 +1222,67 @@ bool UE::Geometry::EnumeratePolygroupSelectionElements(
 						TriangleFunc(TriangleID, FTriangle3d(A, B, C));
 					}
 				}
+
+				// Process the polygroup edges if we are mapping Faces to Edges
+				if ((Flags & EEnumerateMappingFlags::FacesToEdges) != EEnumerateMappingFlags::None)
+				{
+					TArray<int32> GroupEdgeIDs;
+					if ( const FGroupTopology::FGroup* Group = GroupTopology->FindGroupByID(GroupID) )
+					{
+						for (const FGroupTopology::FGroupBoundary& Boundary : Group->Boundaries)
+						{
+							for (const int32 GroupEdgeID : Boundary.GroupEdges)
+							{
+								GroupEdgeIDs.AddUnique(GroupEdgeID);
+							}
+						}
+					}
+					for (const int32 GroupEdgeID : GroupEdgeIDs)
+					{
+						ProcessGroupEdgeID(GroupEdgeID);
+					}
+				}
+
+				// TODO [katie.caputo] add FaceToVerts path
 			}
 		}
 	}
 	else if (MeshSelection.ElementType == EGeometryElementType::Edge)
 	{
-		for (uint64 EncodedID : MeshSelection.Selection)
+		for (const uint64 EncodedID : MeshSelection.Selection)
 		{
-			FMeshTriEdgeID TriEdgeID( FGeoSelectionID(EncodedID).GeometryID );
-			int32 SeedEdgeID = Mesh.IsTriangle(TriEdgeID.TriangleID) ? Mesh.GetTriEdge(TriEdgeID.TriangleID, TriEdgeID.TriEdgeIndex) : IndexConstants::InvalidID;
+			const FMeshTriEdgeID TriEdgeID( FGeoSelectionID(EncodedID).GeometryID );
+			const int32 SeedEdgeID = Mesh.IsTriangle(TriEdgeID.TriangleID) ? Mesh.GetTriEdge(TriEdgeID.TriangleID, TriEdgeID.TriEdgeIndex) : IndexConstants::InvalidID;
 			if (Mesh.IsEdge(SeedEdgeID))
 			{
-				int32 GroupEdgeID = GroupTopology->FindGroupEdgeID(SeedEdgeID);
-				ProcessGroupEdgeID(GroupEdgeID);
+				const int32 GroupEdgeID = GroupTopology->FindGroupEdgeID(SeedEdgeID);
+
+				// Process the polygroup edges if we are mapping Edges to Edges
+				if ((Flags & EEnumerateMappingFlags::EdgesToEdges) != EEnumerateMappingFlags::None)
+				{
+					ProcessGroupEdgeID(GroupEdgeID);
+				}
+
+				// TODO [katie.caputo]: Add EdgesToVerts and EdgesToFaces path
 			}
 		}
 	}
 	else if (MeshSelection.ElementType == EGeometryElementType::Vertex)
 	{
-		for (uint64 EncodedID : MeshSelection.Selection)
+		for (const uint64 EncodedID : MeshSelection.Selection)
 		{
-			int32 VertexID = (int32)FGeoSelectionID(EncodedID).GeometryID;
+			const int32 VertexID = (int32)FGeoSelectionID(EncodedID).GeometryID;
 			if (Mesh.IsVertex(VertexID))
 			{
-				FVector3d A = Mesh.GetVertex(VertexID);
-				VertexFunc(VertexID, (ApplyTransform != nullptr) ? ApplyTransform->TransformPosition(A) : A);
+				const FVector3d A = Mesh.GetVertex(VertexID);
+
+				// Call VertexFunc if we are mapping Vertices to Vertices
+				if ((Flags & EEnumerateMappingFlags::VertsToVerts) != EEnumerateMappingFlags::None)
+				{
+					VertexFunc(VertexID, (ApplyTransform != nullptr) ? ApplyTransform->TransformPosition(A) : A);
+				}
+
+				// todo [katie.caputo]: Add VertsToEdges and VertsToFaces path
 			}
 		}
 	}
diff --git a/Engine/Plugins/Runtime/GeometryProcessing/Source/DynamicMesh/Public/Selections/GeometrySelectionUtil.h b/Engine/Plugins/Runtime/GeometryProcessing/Source/DynamicMesh/Public/Selections/GeometrySelectionUtil.h
--- a/Engine/Plugins/Runtime/GeometryProcessing/Source/DynamicMesh/Public/Selections/GeometrySelectionUtil.h
+++ b/Engine/Plugins/Runtime/GeometryProcessing/Source/DynamicMesh/Public/Selections/GeometrySelectionUtil.h
@@ -18,8 +18,27 @@ namespace Geometry
 class FColliderMesh;
 class FGroupTopology;
 struct FGroupTopologySelection;
+	
+enum class EEnumerateMappingFlags : uint16
+{
+	None         = 0,
+	
+	FacesToFaces = 1 << 0,
+	FacesToEdges = 1 << 1,
+	FacesToVerts = 1 << 2, // not used yet
+
+	EdgesToFaces = 1 << 3, // not used yet
+	EdgesToEdges = 1 << 4,
+	EdgesToVerts = 1 << 5, // not used yet
 
+	VertsToFaces = 1 << 6, // not used yet
+	VertsToEdges = 1 << 7, // not used yet
+	VertsToVerts = 1 << 8,
 
+	Default = FacesToFaces | EdgesToEdges | VertsToVerts,
+};
+ENUM_CLASS_FLAGS(EEnumerateMappingFlags);
+	
 /**
  * Test if SelectionA and SelectionB are the same selection.
  * This is currently relatively expensive on Polygroup selections due to how they are encoded
@@ -177,6 +196,17 @@ DYNAMICMESH_API bool EnumeratePolygroupSelectionEdges(
 	const UE::Geometry::FPolygroupSet& GroupSet,
 	TFunctionRef<void(int32)> EdgeFunc
 );
+	
+/** Prefer EnumerateTriangleSelectionElements with Flags parameter. */
+DYNAMICMESH_API bool EnumerateTriangleSelectionElements(
+	const FGeometrySelection& MeshSelection,
+	const UE::Geometry::FDynamicMesh3& Mesh,
+	TFunctionRef<void(int32, const FVector3d&)> VertexFunc,
+	TFunctionRef<void(int32, const FSegment3d&)> EdgeFunc,
+	TFunctionRef<void(int32, const FTriangle3d&)> TriangleFunc,
+	const FTransform* ApplyTransform = nullptr,
+	const bool bMapFacesToEdgeLoops = false
+);
 
 
 /**
@@ -193,7 +223,19 @@ DYNAMICMESH_API bool EnumerateTriangleSelectionElements(
 	TFunctionRef<void(int32, const FSegment3d&)> EdgeFunc,
 	TFunctionRef<void(int32, const FTriangle3d&)> TriangleFunc,
 	const FTransform* ApplyTransform = nullptr,
-	bool bMapFacesToEdgeLoops = false
+	const EEnumerateMappingFlags Flags = EEnumerateMappingFlags::Default
+);
+
+/** Prefer EnumeratePolygroupSelectionElements with Flags parameter. */
+DYNAMICMESH_API bool EnumeratePolygroupSelectionElements(
+	const FGeometrySelection& MeshSelection,
+	const UE::Geometry::FDynamicMesh3& Mesh,
+	const FGroupTopology* GroupTopology,
+	TFunctionRef<void(int32, const FVector3d&)> VertexFunc,
+	TFunctionRef<void(int32, const FSegment3d&)> EdgeFunc,
+	TFunctionRef<void(int32, const FTriangle3d&)> TriangleFunc,
+	const FTransform* ApplyTransform = nullptr,
+	const bool bMapFacesToEdgeLoops = false
 );
 /**
  * Call VertexFunc/EdgeFunc/TriangleFunc for the mesh vertices/edges/triangles identified by MeshSelection,
@@ -211,7 +253,7 @@ DYNAMICMESH_API bool EnumeratePolygroupSelectionElements(
 	TFunctionRef<void(int32, const FSegment3d&)> EdgeFunc,
 	TFunctionRef<void(int32, const FTriangle3d&)> TriangleFunc,
 	const FTransform* ApplyTransform = nullptr,
-	bool bMapFacesToEdgeLoops = false
+	const EEnumerateMappingFlags Flags = EEnumerateMappingFlags::Default
 );
 
 
diff --git a/Engine/Plugins/Runtime/MeshModelingToolset/Source/ModelingComponents/Private/Selection/DynamicMeshSelector.cpp b/Engine/Plugins/Runtime/MeshModelingToolset/Source/ModelingComponents/Private/Selection/DynamicMeshSelector.cpp
--- a/Engine/Plugins/Runtime/MeshModelingToolset/Source/ModelingComponents/Private/Selection/DynamicMeshSelector.cpp
+++ b/Engine/Plugins/Runtime/MeshModelingToolset/Source/ModelingComponents/Private/Selection/DynamicMeshSelector.cpp
@@ -291,6 +291,44 @@ void FBaseDynamicMeshSelector::UpdateSelectionViaRaycast_MeshTopology(
 	UE::Geometry::UpdateTriangleSelectionViaRaycast(
 		GetColliderMesh(), &SelectionEditor,
 		LocalRay, UpdateConfig, ResultOut);
+	
+	if (SelectionEditor.GetElementType() == EGeometryElementType::Edge)
+	{
+		for (uint64 AddedElement : ResultOut.SelectionDelta.Added)
+		{
+			FMeshTriEdgeID TriEdgeID(FGeoSelectionID(AddedElement).GeometryID);
+			TargetMesh->ProcessMesh([TriEdgeID, &SelectionEditor](const UE::Geometry::FDynamicMesh3& SourceMesh)
+			{
+				// the added EdgeID
+				const int32 EdgeID = SourceMesh.IsTriangle(TriEdgeID.TriangleID) ? SourceMesh.GetTriEdge(TriEdgeID.TriangleID, TriEdgeID.TriEdgeIndex) : IndexConstants::InvalidID;
+				if (SourceMesh.IsEdge(EdgeID))
+				{
+					const FIndex2i EdgeV = SourceMesh.GetEdgeV(EdgeID);
+					// necessary so that edge is recognized from both 'sides' since edge is stored as MeshTriEdgeID.Encoded in SelectionEditor
+					// therefore need to include both of the Tris which the edge belongs to to the SelectionEditor
+					SourceMesh.EnumerateEdgeTriangles(EdgeID, [&SourceMesh, &EdgeV, &SelectionEditor, &TriEdgeID](const int32 TriangleID)
+					{
+						// avoid adding to selection edge already exists in selection
+						if (TriangleID == TriEdgeID.TriangleID)
+						{
+							return;
+						}
+
+						// iterate through all edges of the triangle
+						for (int i = 0; i <= 2; i++)
+						{
+							const int32 OtherEdgeID = SourceMesh.GetTriEdge(TriangleID, i);
+							if (EdgeV == SourceMesh.GetEdgeV(OtherEdgeID))
+							{
+								FMeshTriEdgeID MeshTriEdgeToSelect = FMeshTriEdgeID(TriangleID, i);
+								SelectionEditor.Select((uint64)MeshTriEdgeToSelect.Encoded());
+							}
+						}
+					});
+				}
+			});
+		}
+	}
 }
 
 
@@ -640,6 +678,20 @@ void FBaseDynamicMeshSelector::AccumulateSelectionElements(const FGeometrySelect
 		bMapFacesToEdges);
 }
 
+void FBaseDynamicMeshSelector::AccumulateSelectionElements(const FGeometrySelection& Selection, FGeometrySelectionElements& Elements, bool bTransformToWorld, const EEnumerateMappingFlags Flags)
+{
+	const FTransform UseWorldTransform = GetLocalToWorldTransform();
+	const FTransform* ApplyTransform = (bTransformToWorld) ? &UseWorldTransform : nullptr;
+
+	ToolSelectionUtil::AccumulateSelectionElements(
+		Elements,
+		Selection,
+		TargetMesh->GetMeshRef(),
+		Selection.TopologyType == EGeometryTopologyType::Polygroup ? GetGroupTopology() : nullptr,
+		ApplyTransform,
+		Flags);
+}
+
 void FBaseDynamicMeshSelector::AccumulateElementsFromPredicate(FGeometrySelectionElements& Elements, bool bTransformToWorld, bool bIsForPreview, bool bUseGroupTopology, TFunctionRef<bool(EGeometryElementType, FGeoSelectionID)> Predicate)
 {
 	auto AccumulateElementsOfTypeFromPredicate = [this, &Elements, &Predicate, bTransformToWorld, bIsForPreview, bUseGroupTopology](EGeometryElementType ElementType)
diff --git a/Engine/Plugins/Runtime/MeshModelingToolset/Source/ModelingComponents/Private/Selection/GeometrySelectionManager.cpp b/Engine/Plugins/Runtime/MeshModelingToolset/Source/ModelingComponents/Private/Selection/GeometrySelectionManager.cpp
--- a/Engine/Plugins/Runtime/MeshModelingToolset/Source/ModelingComponents/Private/Selection/GeometrySelectionManager.cpp
+++ b/Engine/Plugins/Runtime/MeshModelingToolset/Source/ModelingComponents/Private/Selection/GeometrySelectionManager.cpp
@@ -17,7 +17,7 @@
 
 static TAutoConsoleVariable<int32> CVarGeometrySelectionManager_FullSelectionHoverHighlights(
 	TEXT("modeling.Selection.FullHoverHighlights"),
-	0,
+	1,
 	TEXT("Use full selection hover highlights instead of simplified highlights")
 );
 
@@ -982,7 +982,55 @@ bool UGeometrySelectionManager::UpdateSelectionPreviewViaRaycast(
 	if ( ! UE::Geometry::AreSelectionsIdentical(NewPreview.PreviewSelection, ActivePreviewSelection) )
 	{
 		ActivePreviewSelection = MoveTemp(NewPreview.PreviewSelection);
-		CachedPreviewRenderElements.Reset();
+		
+		// Initialize [Un]SelectedActivePreviewSelection(s) so that they are of the correct Topology and Geometry type, then clear them
+		SelectedActivePreviewSelection = MoveTemp(NewPreview.PreviewSelection);
+		UnselectedActivePreviewSelection = MoveTemp(NewPreview.PreviewSelection);
+		SelectedActivePreviewSelection.Reset();
+		UnselectedActivePreviewSelection.Reset();
+
+		if (MeshTopologyMode == EMeshTopologyMode::Polygroup)
+		{
+			// Get all polygroup IDs in current preview selection
+			TSet<uint32> SelectedGroupIDs;
+			for (const uint64 ID : Target->Selection.Selection)
+			{
+				SelectedGroupIDs.Add(FGeoSelectionID(ID).TopologyID);
+			}
+
+			// Get GroupID of active preview selection (hovered items)
+			for (const uint64 ID : ActivePreviewSelection.Selection)
+			{
+				const uint32 TopoID = FGeoSelectionID(ID).TopologyID;
+
+				// add to selection according to if an element with the GroupID is already selected
+				if (SelectedGroupIDs.Contains(TopoID))
+				{
+					SelectedActivePreviewSelection.Selection.Add(ID);
+				}
+				else
+				{
+					UnselectedActivePreviewSelection.Selection.Add(ID);
+				}
+			}
+		}
+		// Triangle Topology mode is more straightforward
+		else if (MeshTopologyMode == EMeshTopologyMode::Triangle)
+		{
+			for (const uint64 ID : ActivePreviewSelection.Selection)
+			{
+				if (Target->Selection.Selection.Contains(ID))
+				{
+					SelectedActivePreviewSelection.Selection.Add(ID);
+				}
+				else
+				{
+					UnselectedActivePreviewSelection.Selection.Add(ID);
+				}
+			}
+		}
+		CachedSelectedPreviewRenderElements.Reset();
+		CachedUnselectedPreviewRenderElements.Reset();
 		MarkRenderCachesDirty();
 	}
 
@@ -1353,7 +1401,10 @@ void UGeometrySelectionManager::RebuildSelectionRenderCaches()
 void UGeometrySelectionManager::ClearActivePreview()
 {
 	ActivePreviewSelection.Reset();
-	CachedPreviewRenderElements.Reset();
+	SelectedActivePreviewSelection.Reset();
+	UnselectedActivePreviewSelection.Reset();
+	CachedSelectedPreviewRenderElements.Reset();
+	CachedUnselectedPreviewRenderElements.Reset();
 	bSelectableRenderCachesDirty = true;
 	bPreviewRenderCachesDirty = true;
 }
@@ -1395,26 +1446,12 @@ void UGeometrySelectionManager::RebuildSelectableRenderCache()
 		
 		Target->Selector->AccumulateElementsFromPredicate(AllElements, true, false, MeshTopologyMode == EMeshTopologyMode::Polygroup, [Target, this](EGeometryElementType Type, FGeoSelectionID ID)
 		{
-			uint64 EncodedID = ID.Encoded();
-
 			// Selectable faces are not displayed directly, just implicitly via displayed edges.
 			if (Type == EGeometryElementType::Face)
 			{
 				return false;
 			}
 			
-			// Exclude selected elements from selectable elements
-			if (Target->Selection.ElementType == Type && Target->Selection.Selection.Contains(EncodedID))
-			{
-				return false;
-			}
-
-			// Exclude hovered elements from selectable elements
-			if (ActivePreviewSelection.ElementType == Type && ActivePreviewSelection.Selection.Contains(EncodedID))
-			{
-				return false;
-			}
-			
 			return true;
 		});
 	}
@@ -1429,12 +1466,20 @@ void UGeometrySelectionManager::RebuildPreviewRenderCache()
 		return;
 	}
 
-	const bool bUseSimplifiedPreviewHighlight = (CVarGeometrySelectionManager_FullSelectionHoverHighlights.GetValueOnGameThread() == 0);
+	// defaults to off/false; when off, will show outlines and fill color when hovering. When on/true, will only show outlines
 	const TSharedPtr<FGeometrySelectionTarget> Target = ActiveTargetReferences[0];
 	
 	if (ActivePreviewSelection.IsEmpty() == false)
 	{
-		Target->Selector->AccumulateSelectionElements(ActivePreviewSelection, CachedPreviewRenderElements, true, bUseSimplifiedPreviewHighlight);
+		EEnumerateMappingFlags MappingFlags = EEnumerateMappingFlags::Default | EEnumerateMappingFlags::FacesToEdges;
+		if (CVarGeometrySelectionManager_FullSelectionHoverHighlights.GetValueOnGameThread() == 0)
+		{
+			// Unset FacesToFaces flag if full hover highlights are disabled
+			MappingFlags &= ~EEnumerateMappingFlags::FacesToFaces;
+		}
+		
+		Target->Selector->AccumulateSelectionElements(SelectedActivePreviewSelection, CachedSelectedPreviewRenderElements, true, MappingFlags);
+		Target->Selector->AccumulateSelectionElements(UnselectedActivePreviewSelection, CachedUnselectedPreviewRenderElements, true, MappingFlags);
 	}
 	
 	bPreviewRenderCachesDirty = false;
@@ -1482,21 +1527,31 @@ void UGeometrySelectionManager::DebugRender(IToolsContextRenderAPI* RenderAPI)
 		
 		for ( const FGeometrySelectionElements& Elements : CachedSelectionRenderElements )
 		{
-			ToolSelectionUtil::DebugRender(RenderAPI, Elements, 4.f, LinearColors::Gold3f(), 10.f, LinearColors::Gold3f(), 6.f);
+			// render geometry selection
+			ToolSelectionUtil::DebugRender(RenderAPI, Elements, 4.f, GeometrySelectedColor, 10.f, GeometrySelectedColor, 6.f, GeometrySelectedColor);
 		}
 
 		if (MeshTopologyMode != EMeshTopologyMode::None)
 		{
 			for ( const FGeometrySelectionElements& Elements : CachedSelectableRenderElements )
 			{
-				ToolSelectionUtil::DebugRender(RenderAPI, Elements, 2.f, LinearColors::Red3f(), 8.f, LinearColors::Red3f(), 5.f);
+				// render unselected geometry (wireframe overlay color in MeshElementSelection mode)
+				ToolSelectionUtil::DebugRender(RenderAPI, Elements, 2.f, UnselectedColor, 8.f, UnselectedColor, 5.f, UnselectedColor);
 			}
 		}
-		
-		ToolSelectionUtil::DebugRender(RenderAPI, CachedPreviewRenderElements, 4.f, LinearColors::Green3f(), 10.f, LinearColors::Green3f(), 7.f);
+		// render both hover cases (if element is part of the selection, or if not)
+		ToolSelectionUtil::DebugRender(RenderAPI, CachedSelectedPreviewRenderElements, 4.f, HoverOverSelectedColor, 10.f, HoverOverSelectedColor, 7.f, HoverOverSelectedColor);
+		ToolSelectionUtil::DebugRender(RenderAPI, CachedUnselectedPreviewRenderElements, 4.f, HoverOverUnselectedColor, 10.f, HoverOverUnselectedColor, 7.f, HoverOverUnselectedColor);
 	}
 }
 
+void UGeometrySelectionManager::SetSelectionColors(const FLinearColor UnselectedCol, const FLinearColor HoverOverSelectedCol, const FLinearColor HoverOverUnselectedCol, const FLinearColor GeometrySelectedCol)
+{
+	UnselectedColor = UnselectedCol;
+	HoverOverSelectedColor = HoverOverSelectedCol;
+	HoverOverUnselectedColor = HoverOverUnselectedCol;
+	GeometrySelectedColor = GeometrySelectedCol;
+}
 
 
 
diff --git a/Engine/Plugins/Runtime/MeshModelingToolset/Source/ModelingComponents/Private/Selection/ToolSelectionUtil.cpp b/Engine/Plugins/Runtime/MeshModelingToolset/Source/ModelingComponents/Private/Selection/ToolSelectionUtil.cpp
--- a/Engine/Plugins/Runtime/MeshModelingToolset/Source/ModelingComponents/Private/Selection/ToolSelectionUtil.cpp
+++ b/Engine/Plugins/Runtime/MeshModelingToolset/Source/ModelingComponents/Private/Selection/ToolSelectionUtil.cpp
@@ -49,6 +49,23 @@ bool ToolSelectionUtil::AccumulateSelectionElements(
 	const FGroupTopology* GroupTopology,
 	const FTransform* ApplyTransform,
 	bool bMapFacesToEdges)
+{
+	return AccumulateSelectionElements(
+		Elements,
+		Selection,
+		SourceMesh,
+		GroupTopology,
+		ApplyTransform,
+		EEnumerateMappingFlags::Default | (bMapFacesToEdges ? EEnumerateMappingFlags::FacesToEdges : EEnumerateMappingFlags::None));
+}
+
+bool ToolSelectionUtil::AccumulateSelectionElements(
+	FGeometrySelectionElements& Elements,
+	const FGeometrySelection& Selection,
+	const FDynamicMesh3& SourceMesh,
+	const FGroupTopology* GroupTopology,
+	const FTransform* ApplyTransform,
+	const EEnumerateMappingFlags Flags)
 {
 	auto AddPoint = [&Elements](uint32 Vid, const FVector3d& Point) { Elements.Points.Add(Point); };
 	auto AddSegment = [&Elements](uint32 Eid, const FSegment3d& Segment) { Elements.Segments.Add(Segment); };
@@ -59,16 +76,16 @@ bool ToolSelectionUtil::AccumulateSelectionElements(
 		if (GroupTopology)
 		{
 			return EnumeratePolygroupSelectionElements(Selection, SourceMesh, GroupTopology,
-				AddPoint, AddSegment, AddTriangle, ApplyTransform, bMapFacesToEdges);
+				AddPoint, AddSegment, AddTriangle, ApplyTransform, Flags);
 		}
 
 		const FGroupTopology ComputedGroupTopology(&SourceMesh, true);
 		return EnumeratePolygroupSelectionElements(Selection, SourceMesh, &ComputedGroupTopology,
-			AddPoint, AddSegment, AddTriangle, ApplyTransform, bMapFacesToEdges);
+			AddPoint, AddSegment, AddTriangle, ApplyTransform, Flags);
 	}
 
 	return EnumerateTriangleSelectionElements(Selection, SourceMesh, 
-			AddPoint, AddSegment, AddTriangle, ApplyTransform, bMapFacesToEdges);
+			AddPoint, AddSegment, AddTriangle, ApplyTransform, Flags);
 }
 
 void ToolSelectionUtil::DebugRenderGeometrySelectionElements(
@@ -92,7 +109,8 @@ void ToolSelectionUtil::DebugRender(
 	FLinearColor LineColor,
 	float PointSize,
 	FLinearColor PointColor,
-	float DepthBias)
+	float DepthBias,
+	FLinearColor FillColor)
 {
 	if (RenderAPI == nullptr)
 	{
@@ -115,7 +133,16 @@ void ToolSelectionUtil::DebugRender(
 		MeshBuilder.AddTriangle(V0, V1, V2);
 	}
 	//FMaterialRenderProxy* MaterialRenderProxy = TriangleMaterial->GetRenderProxy();		// currently does not work, material does not render
-	FMaterialRenderProxy* MaterialRenderProxy = GEngine->ConstraintLimitMaterialX->GetRenderProxy();
+
+	// desaturated color to use when triangle is selected in Mesh Element Selection
+	// uses ConstraintLimitMaterial bc ConstraintLimitMaterialX was used previously to get desaturated red color
+	TObjectPtr<class UMaterialInstanceDynamic> MeshElementSelectionFillColor;
+	MeshElementSelectionFillColor = UMaterialInstanceDynamic::Create((UMaterialInterface*)GEngine->ConstraintLimitMaterial, NULL);
+	MeshElementSelectionFillColor->SetVectorParameterValue(FName("Color"), FillColor);
+	MeshElementSelectionFillColor->SetScalarParameterValue(FName("Desaturation"), 0.6f);
+	FMaterialRenderProxy* MaterialRenderProxy = MeshElementSelectionFillColor->GetRenderProxy();
+	//FMaterialRenderProxy* MaterialRenderProxy = GEngine->ConstraintLimitMaterialX->GetRenderProxy();
+
 	MeshBuilder.Draw(CurrentPDI, FMatrix::Identity, MaterialRenderProxy, DepthPriority, false, false);
 
 	FToolDataVisualizer Visualizer;
diff --git a/Engine/Plugins/Runtime/MeshModelingToolset/Source/ModelingComponents/Public/Selection/DynamicMeshSelector.h b/Engine/Plugins/Runtime/MeshModelingToolset/Source/ModelingComponents/Public/Selection/DynamicMeshSelector.h
--- a/Engine/Plugins/Runtime/MeshModelingToolset/Source/ModelingComponents/Public/Selection/DynamicMeshSelector.h
+++ b/Engine/Plugins/Runtime/MeshModelingToolset/Source/ModelingComponents/Public/Selection/DynamicMeshSelector.h
@@ -119,6 +119,7 @@ public:
 	virtual void GetTargetFrame(const FGeometrySelection& Selection, UE::Geometry::FFrame3d& SelectionFrame) override;
 	virtual void AccumulateSelectionBounds(const FGeometrySelection& Selection, FGeometrySelectionBounds& BoundsInOut, bool bTransformToWorld) override;
 	virtual void AccumulateSelectionElements(const FGeometrySelection& Selection, FGeometrySelectionElements& Elements, bool bTransformToWorld, bool bIsForPreview) override;
+	virtual void AccumulateSelectionElements(const FGeometrySelection& Selection, FGeometrySelectionElements& Elements, bool bTransformToWorld, UE::Geometry::EEnumerateMappingFlags Flags = UE::Geometry::EEnumerateMappingFlags::Default) override;
 	
 	virtual void AccumulateElementsFromPredicate(
 		FGeometrySelectionElements& Elements,
diff --git a/Engine/Plugins/Runtime/MeshModelingToolset/Source/ModelingComponents/Public/Selection/GeometrySelectionManager.h b/Engine/Plugins/Runtime/MeshModelingToolset/Source/ModelingComponents/Public/Selection/GeometrySelectionManager.h
--- a/Engine/Plugins/Runtime/MeshModelingToolset/Source/ModelingComponents/Public/Selection/GeometrySelectionManager.h
+++ b/Engine/Plugins/Runtime/MeshModelingToolset/Source/ModelingComponents/Public/Selection/GeometrySelectionManager.h
@@ -296,7 +296,7 @@ public:
 	virtual bool HasSelection() const;
 
 	/** 
-	 * Get avaialble information about the active selection/state
+	 * Get available information about the active selection/state
 	 */
 	virtual void GetActiveSelectionInfo(EGeometryTopologyType& TopologyTypeOut, EGeometryElementType& ElementTypeOut, int& NumTargetsOut, bool& bIsEmpty) const;
 
@@ -392,8 +392,12 @@ public:
 	virtual void DebugPrintSelection();
 	/** Visualize the active selection using PDI drawing */
 	virtual void DebugRender(IToolsContextRenderAPI* RenderAPI);
-
-
+	
+	/** Set the colors to be used during mesh element selection for:
+	 * Unselected elements, Hover over selection, Hover over non-selection, and Selected elements
+	 */
+	void SetSelectionColors(FLinearColor UnselectedCol, FLinearColor HoverOverSelectedCol, FLinearColor HoverOverUnselectedCol, FLinearColor GeometrySelectedCol);
+	
 protected:
 
 	// current selection mode settings
@@ -460,18 +464,29 @@ protected:
 	void OnTargetGeometryModified(IGeometrySelector* Selector);
 
 	
-	// todo [nickolas.drake]: cane we move CachedSelectionRenderElements, CachedPreviewRenderElements, and bSelectionRenderCachesDirty to private?
+	// todo [nickolas.drake]: cane we move CachedSelectionRenderElements, Cached[Un]selectedPreviewRenderElements, and bSelectionRenderCachesDirty to private?
 	
 	TArray<FGeometrySelectionElements> CachedSelectionRenderElements;				// Cached 3D geometry for current selection
 	bool bSelectionRenderCachesDirty = false;
 	void UpdateSelectionRenderCacheOnTargetChange();
 	void RebuildSelectionRenderCaches();
 
-    FGeometrySelection ActivePreviewSelection;										// Selection representing the active preview
-	FGeometrySelectionElements CachedPreviewRenderElements;							// Cached 3D geometry for active preview elements
+    FGeometrySelection ActivePreviewSelection;
+	FGeometrySelection SelectedActivePreviewSelection;
+	FGeometrySelection UnselectedActivePreviewSelection;
+	
+	FGeometrySelectionElements CachedSelectedPreviewRenderElements;		// Cached 3D geometry for active preview elements that are in the current selection
+	FGeometrySelectionElements CachedUnselectedPreviewRenderElements;  // Cached 3D geometry for active preview elements that are NOT in the current selection
     void ClearActivePreview();
 
 	
+	// colors used in mesh element selection
+	FLinearColor UnselectedColor;
+	FLinearColor HoverOverSelectedColor;
+	FLinearColor HoverOverUnselectedColor;
+	FLinearColor GeometrySelectedColor;
+
+	
 	// various change types need internal access
 
 	friend class FGeometrySelectionManager_SelectionTypeChange;
diff --git a/Engine/Plugins/Runtime/MeshModelingToolset/Source/ModelingComponents/Public/Selection/GeometrySelector.h b/Engine/Plugins/Runtime/MeshModelingToolset/Source/ModelingComponents/Public/Selection/GeometrySelector.h
--- a/Engine/Plugins/Runtime/MeshModelingToolset/Source/ModelingComponents/Public/Selection/GeometrySelector.h
+++ b/Engine/Plugins/Runtime/MeshModelingToolset/Source/ModelingComponents/Public/Selection/GeometrySelector.h
@@ -5,6 +5,7 @@
 #include "Components/PrimitiveComponent.h"
 #include "ConvexVolume.h"
 #include "Selections/GeometrySelection.h"
+#include "Selections/GeometrySelectionUtil.h"
 #include "FrameTypes.h"
 #include "InputState.h"
 #include "ToolContextInterfaces.h"
@@ -408,6 +409,9 @@ public:
 	 */
 	virtual void AccumulateSelectionBounds(const FGeometrySelection& Selection, FGeometrySelectionBounds& BoundsInOut, bool bTransformToWorld) = 0;
 
+	/** Prefer AccumulateSelectionElements with Flags parameter. */
+	virtual void AccumulateSelectionElements(const FGeometrySelection& Selection, FGeometrySelectionElements& ElementsInOut, bool bTransformToWorld, UE::Geometry::EEnumerateMappingFlags Flags = UE::Geometry::EEnumerateMappingFlags::Default) = 0;
+	
 	/**
 	 * Accumulate geometric elements (currently 3D triangles, line segments, and points) for the provided Selection in the provided ElementsInOut. 
 	 * ElementsInOut is not cleared.
diff --git a/Engine/Plugins/Runtime/MeshModelingToolset/Source/ModelingComponents/Public/Selection/ToolSelectionUtil.h b/Engine/Plugins/Runtime/MeshModelingToolset/Source/ModelingComponents/Public/Selection/ToolSelectionUtil.h
--- a/Engine/Plugins/Runtime/MeshModelingToolset/Source/ModelingComponents/Public/Selection/ToolSelectionUtil.h
+++ b/Engine/Plugins/Runtime/MeshModelingToolset/Source/ModelingComponents/Public/Selection/ToolSelectionUtil.h
@@ -3,6 +3,7 @@
 #pragma once
 
 #include "Selections/GeometrySelection.h"
+#include "Selections/GeometrySelectionUtil.h"
 #include "CoreMinimal.h"
 
 class UInteractiveToolManager;
@@ -34,6 +35,15 @@ namespace ToolSelectionUtil
 	 */
 	MODELINGCOMPONENTS_API void SetNewActorSelection(UInteractiveToolManager* ToolManager, const TArray<AActor*>& Actors);
 
+	/** Prefer AccumulateSelectionElements with Flags parameter. */
+	MODELINGCOMPONENTS_API bool AccumulateSelectionElements(
+		UE::Geometry::FGeometrySelectionElements& Elements,
+		const UE::Geometry::FGeometrySelection& Selection,
+		const UE::Geometry::FDynamicMesh3& SourceMesh,
+		const UE::Geometry::FGroupTopology* Topology = nullptr,
+		const FTransform* ApplyTransform = nullptr,
+		bool bMapFacesToEdges = false);
+
 	/**
 	 * Add the geometry selection elements corresponding to the given Selection to Elements. This function does not
 	 * reset Elements before adding elements. If the Selection has Polygroup topology then use the given Topology to
@@ -42,12 +52,12 @@ namespace ToolSelectionUtil
 	 * Return false if there was an error and true otherwise
 	 */
 	MODELINGCOMPONENTS_API bool AccumulateSelectionElements(
-			UE::Geometry::FGeometrySelectionElements& Elements,
-			const UE::Geometry::FGeometrySelection& Selection,
-			const UE::Geometry::FDynamicMesh3& SourceMesh,
-			const UE::Geometry::FGroupTopology* Topology = nullptr,
-			const FTransform* ApplyTransform = nullptr,
-			bool bMapFacesToEdges = false);
+		UE::Geometry::FGeometrySelectionElements& Elements,
+		const UE::Geometry::FGeometrySelection& Selection,
+		const UE::Geometry::FDynamicMesh3& SourceMesh,
+		const UE::Geometry::FGroupTopology* Topology = nullptr,
+		const FTransform* ApplyTransform = nullptr,
+		const UE::Geometry::EEnumerateMappingFlags Flags = UE::Geometry::EEnumerateMappingFlags::Default);
 	
 	/**
 	 * Render the given Elements using FPrimitiveDrawInterface
@@ -65,7 +75,8 @@ namespace ToolSelectionUtil
 		FLinearColor LineColor,
 		float PointSize,
 		FLinearColor PointColor,
-		float DepthBias = 0.f);
+		float DepthBias = 0.f,
+		FLinearColor FillColor = FLinearColor(1.0f, 0.0f, 0.0f));
 }
 
 namespace UE::Geometry
