PolyModelingFaceUtil  - fix error in detecting the case of a planar selection

diff --git a/Engine/Plugins/Runtime/GeometryProcessing/Source/DynamicMesh/Private/Operations/PolyModeling/PolyModelingFaceUtil.cpp b/Engine/Plugins/Runtime/GeometryProcessing/Source/DynamicMesh/Private/Operations/PolyModeling/PolyModelingFaceUtil.cpp
--- a/Engine/Plugins/Runtime/GeometryProcessing/Source/DynamicMesh/Private/Operations/PolyModeling/PolyModelingFaceUtil.cpp
+++ b/Engine/Plugins/Runtime/GeometryProcessing/Source/DynamicMesh/Private/Operations/PolyModeling/PolyModelingFaceUtil.cpp
@@ -71,6 +71,7 @@ UE::Geometry::FFrame3d UE::Geometry::ComputeFaceSelectionFrame(
 	for ( int32 tid : *UseTriangles )
 	{
 		FVector3d A,B,C;
+		Mesh.GetTriVertices(tid, A, B, C);
 		ProjectionRange.Contain( (A-Centroid).Dot(UseNormal) );
 		ProjectionRange.Contain( (B-Centroid).Dot(UseNormal) );
 		ProjectionRange.Contain( (C-Centroid).Dot(UseNormal) );
