Modeling Tools: Implement the selection Mapping Flags for MeshElementSelection and clarify its related code

diff --git a/Engine/Plugins/Runtime/GeometryProcessing/Source/DynamicMesh/Private/Selections/GeometrySelectionUtil.cpp b/Engine/Plugins/Runtime/GeometryProcessing/Source/DynamicMesh/Private/Selections/GeometrySelectionUtil.cpp
--- a/Engine/Plugins/Runtime/GeometryProcessing/Source/DynamicMesh/Private/Selections/GeometrySelectionUtil.cpp
+++ b/Engine/Plugins/Runtime/GeometryProcessing/Source/DynamicMesh/Private/Selections/GeometrySelectionUtil.cpp
@@ -466,13 +466,20 @@ bool UE::Geometry::UpdateSelectionWithNewElements(
 	
 }
 
-
-
 bool UE::Geometry::EnumerateTriangleSelectionVertices(
 	const FGeometrySelection& MeshSelection,
 	const UE::Geometry::FDynamicMesh3& Mesh,
 	const FTransform& ApplyTransform,
 	TFunctionRef<void(uint64, const FVector3d&)> VertexFunc)
+{
+	return EnumerateTriangleSelectionVertices(MeshSelection, Mesh, &ApplyTransform, VertexFunc);
+}
+
+bool UE::Geometry::EnumerateTriangleSelectionVertices(
+	const FGeometrySelection& MeshSelection,
+	const UE::Geometry::FDynamicMesh3& Mesh,
+	const FTransform* ApplyTransform,
+	TFunctionRef<void(uint64, const FVector3d&)> VertexFunc)
 {
 	if ( ensure( MeshSelection.TopologyType == EGeometryTopologyType::Triangle ) == false )
 	{
@@ -481,38 +488,63 @@ bool UE::Geometry::EnumerateTriangleSelectionVertices(
 
 	if (MeshSelection.ElementType == EGeometryElementType::Face)
 	{
-		for (uint64 TriangleID : MeshSelection.Selection)
+		for (const uint64 TriangleID : MeshSelection.Selection)
 		{
 			if (Mesh.IsTriangle((int32)TriangleID))
 			{
-				FIndex3i Triangle = Mesh.GetTriangle((int32)TriangleID);
-				VertexFunc((uint64)Triangle.A, ApplyTransform.TransformPosition(Mesh.GetVertex(Triangle.A)));
-				VertexFunc((uint64)Triangle.B, ApplyTransform.TransformPosition(Mesh.GetVertex(Triangle.B)));
-				VertexFunc((uint64)Triangle.C, ApplyTransform.TransformPosition(Mesh.GetVertex(Triangle.C)));
+				const FIndex3i Triangle = Mesh.GetTriangle((int32)TriangleID);
+				FVector3d TriVertexA, TriVertexB, TriVertexC;
+				TriVertexA = Mesh.GetVertex(Triangle.A);
+				TriVertexB = Mesh.GetVertex(Triangle.B);
+				TriVertexC = Mesh.GetVertex(Triangle.C);
+				
+				if (ApplyTransform)
+				{
+					TriVertexA = ApplyTransform->TransformPosition(TriVertexA);
+					TriVertexB = ApplyTransform->TransformPosition(TriVertexB);
+					TriVertexC = ApplyTransform->TransformPosition(TriVertexC);
+				}
+				VertexFunc((uint64)Triangle.A, TriVertexA);
+				VertexFunc((uint64)Triangle.B, TriVertexB);
+				VertexFunc((uint64)Triangle.C, TriVertexC);
 			}
 		}
 	}
 	else if (MeshSelection.ElementType == EGeometryElementType::Edge)
 	{
-		for (uint64 EncodedID : MeshSelection.Selection)
+		for (const uint64 EncodedID : MeshSelection.Selection)
 		{
-			FMeshTriEdgeID TriEdgeID(FGeoSelectionID(EncodedID).GeometryID);
-			int32 EdgeID = Mesh.IsTriangle(TriEdgeID.TriangleID) ? Mesh.GetTriEdge(TriEdgeID.TriangleID, TriEdgeID.TriEdgeIndex) : IndexConstants::InvalidID;
+			const FMeshTriEdgeID TriEdgeID(FGeoSelectionID(EncodedID).GeometryID);
+			const int32 EdgeID = Mesh.IsTriangle(TriEdgeID.TriangleID) ? Mesh.GetTriEdge(TriEdgeID.TriangleID, TriEdgeID.TriEdgeIndex) : IndexConstants::InvalidID;
 			if (Mesh.IsEdge(EdgeID))
 			{
-				FIndex2i EdgeV = Mesh.GetEdgeV(EdgeID);
-				VertexFunc((uint64)EdgeV.A, ApplyTransform.TransformPosition(Mesh.GetVertex(EdgeV.A)));
-				VertexFunc((uint64)EdgeV.B, ApplyTransform.TransformPosition(Mesh.GetVertex(EdgeV.B)));
+				const FIndex2i EdgeV = Mesh.GetEdgeV(EdgeID);
+				FVector3d EdgeVertexA, EdgeVertexB;
+				EdgeVertexA = Mesh.GetVertex(EdgeV.A);
+				EdgeVertexB = Mesh.GetVertex(EdgeV.B);
+
+				if (ApplyTransform)
+				{
+					EdgeVertexA = ApplyTransform->TransformPosition(EdgeVertexA);
+					EdgeVertexB = ApplyTransform->TransformPosition(EdgeVertexB);
+				}
+				VertexFunc((uint64)EdgeV.A, EdgeVertexA);
+				VertexFunc((uint64)EdgeV.B, EdgeVertexB);
 			}
 		}
 	}
 	else if (MeshSelection.ElementType == EGeometryElementType::Vertex)
 	{
-		for (uint64 VertexID : MeshSelection.Selection)
+		for (const uint64 VertexID : MeshSelection.Selection)
 		{
 			if (Mesh.IsVertex((int32)VertexID))
 			{
-				VertexFunc((uint64)VertexID, ApplyTransform.TransformPosition(Mesh.GetVertex((int32)VertexID)));
+				FVector3d Vertex = Mesh.GetVertex((int32)VertexID);
+				if (ApplyTransform)
+				{
+					Vertex = ApplyTransform->TransformPosition(Vertex);
+				}
+				VertexFunc(VertexID, Vertex);
 			}
 		}
 	}
@@ -1037,7 +1069,7 @@ bool UE::Geometry::EnumerateTriangleSelectionElements(
 		EdgeFunc,
 		TriangleFunc,
 		ApplyTransform,
-		EEnumerateMappingFlags::Default | (bMapFacesToEdgeLoops ? EEnumerateMappingFlags::FacesToEdges : EEnumerateMappingFlags::None)
+		EEnumerateSelectionMapping::Default | (bMapFacesToEdgeLoops ? EEnumerateSelectionMapping::FacesToEdges : EEnumerateSelectionMapping::None)
 		);
 }
 
@@ -1048,93 +1080,66 @@ bool UE::Geometry::EnumerateTriangleSelectionElements(
 	TFunctionRef<void(int32, const FSegment3d&)> EdgeFunc,
 	TFunctionRef<void(int32, const FTriangle3d&)> TriangleFunc,
 	const FTransform* ApplyTransform,
-	const EEnumerateMappingFlags Flags
+	const EEnumerateSelectionMapping Flags
 )
 {
 	if ( ensure( MeshSelection.TopologyType == EGeometryTopologyType::Triangle ) == false )
 	{
 		return false;
 	}
+	
+	auto ApplyTriangleFunc = [&ApplyTransform, &TriangleFunc, &Mesh](const int32 TriangleID)
+	{
+		FVector3d VertA, VertB, VertC;
+		Mesh.GetTriVertices((int32)TriangleID, VertA, VertB, VertC);
+		if (ApplyTransform)
+		{
+			VertA = ApplyTransform->TransformPosition(VertA);
+			VertB = ApplyTransform->TransformPosition(VertB);
+			VertC = ApplyTransform->TransformPosition(VertC);
+		}
+		TriangleFunc(TriangleID, FTriangle3d(VertA, VertB, VertC));	
+	};
+
+	auto ApplyEdgeFunc = [&ApplyTransform, &EdgeFunc, &Mesh](const int32 EdgeID)
+	{
+		FVector3d VertA, VertB;
+		Mesh.GetEdgeV(EdgeID, VertA, VertB);
+		if (ApplyTransform)
+		{
+			VertA = ApplyTransform->TransformPosition(VertA);
+			VertB = ApplyTransform->TransformPosition(VertB);
+		}
+		EdgeFunc(EdgeID, FSegment3d(VertA, VertB));
+	};
+
+	auto ApplyVertexFunc = [&VertexFunc](const uint64 VertexID, const FVector3d& VertA)
+	{
+		VertexFunc((int32)VertexID, VertA);
+	};
+	
 
 	if (MeshSelection.ElementType == EGeometryElementType::Face)
 	{
-		for (const uint64 EncodedID : MeshSelection.Selection)
+		// Call TriangleFunc if we are mapping Faces to Faces
+		if ((Flags & EEnumerateSelectionMapping::FacesToFaces) != EEnumerateSelectionMapping::None)
 		{
-			const int32 TriangleID = (int32)FGeoSelectionID(EncodedID).GeometryID;
-			if (Mesh.IsTriangle(TriangleID))
-			{
-				FVector3d A, B, C;
-				Mesh.GetTriVertices((int32)TriangleID, A, B, C);
-				if (ApplyTransform != nullptr)
-				{
-					A = ApplyTransform->TransformPosition(A);
-					B = ApplyTransform->TransformPosition(B);
-					C = ApplyTransform->TransformPosition(C);
-				}
-				
-				// Call TriangleFunc if we are mapping Faces to Faces
-				if ((Flags & EEnumerateMappingFlags::FacesToFaces) != EEnumerateMappingFlags::None)
-				{
-					TriangleFunc(TriangleID, FTriangle3d(A, B, C));
-				}
-				
-				// Call EdgeFunc for each of the triangle's edges if we are mapping Faces to Edges
-				if ((Flags & EEnumerateMappingFlags::FacesToEdges) != EEnumerateMappingFlags::None)
-				{
-					const FIndex3i Edges = Mesh.GetTriEdges((int32)TriangleID);
-					EdgeFunc(Edges.A, FSegment3d(A, B));
-					EdgeFunc(Edges.B, FSegment3d(B, C));
-					EdgeFunc(Edges.C, FSegment3d(C, A));
-				}
+			EnumerateTriangleSelectionTriangles(MeshSelection, Mesh, ApplyTriangleFunc);
+		}
 
-				// TODO [katie.caputo] add FacesToVerts path
-			}
+		// Call EdgeFunc if we are mapping Faces to Edges
+		if ((Flags & EEnumerateSelectionMapping::FacesToEdges) != EEnumerateSelectionMapping::None)
+		{
+			EnumerateTriangleSelectionEdges(MeshSelection, Mesh, ApplyEdgeFunc);
 		}
 	}
 	else if (MeshSelection.ElementType == EGeometryElementType::Edge)
 	{
-		for (const uint64 EncodedID : MeshSelection.Selection)
-		{
-			const FMeshTriEdgeID TriEdgeID( FGeoSelectionID(EncodedID).GeometryID );
-			const int32 EdgeID = Mesh.IsTriangle(TriEdgeID.TriangleID) ? Mesh.GetTriEdge(TriEdgeID.TriangleID, TriEdgeID.TriEdgeIndex) : IndexConstants::InvalidID;
-			if (Mesh.IsEdge(EdgeID))
-			{
-				FVector3d A, B;
-				Mesh.GetEdgeV(EdgeID, A, B);
-				if (ApplyTransform != nullptr)
-				{
-					A = ApplyTransform->TransformPosition(A);
-					B = ApplyTransform->TransformPosition(B);
-				}
-
-				// Call EdgeFunc if we are mapping Edges to Edges
-				if ((Flags & EEnumerateMappingFlags::EdgesToEdges) != EEnumerateMappingFlags::None)
-				{
-					EdgeFunc(EdgeID, FSegment3d(A, B));
-				}
-				// TODO [katie.caputo] add EdgesToFaces and EdgesToVerts path
-			}
-		}
-		
+		EnumerateTriangleSelectionEdges(MeshSelection, Mesh, ApplyEdgeFunc);
 	}
 	else if (MeshSelection.ElementType == EGeometryElementType::Vertex)
 	{
-		for (const uint64 EncodedID : MeshSelection.Selection)
-		{
-			const int32 VertexID = (int32)FGeoSelectionID(EncodedID).GeometryID;
-			if (Mesh.IsVertex(VertexID))
-			{
-				const FVector3d A = Mesh.GetVertex(VertexID);
-
-				// Call VertexFunc if we are mapping Vertices to Vertices
-				if ((Flags & EEnumerateMappingFlags::VertsToVerts) != EEnumerateMappingFlags::None)
-				{
-					VertexFunc(VertexID, (ApplyTransform != nullptr) ? ApplyTransform->TransformPosition(A) : A);
-				}
-
-				// TODO [katie.caputo] add VertsToFaces and VertsToEdges path
-			}
-		}
+		EnumerateTriangleSelectionVertices(MeshSelection, Mesh, ApplyTransform, ApplyVertexFunc);
 	}
 	else
 	{
@@ -1163,7 +1168,7 @@ bool UE::Geometry::EnumeratePolygroupSelectionElements(
 		EdgeFunc,
 		TriangleFunc,
 		ApplyTransform,
-		EEnumerateMappingFlags::Default | (bMapFacesToEdgeLoops ? EEnumerateMappingFlags::FacesToEdges : EEnumerateMappingFlags::None)
+		EEnumerateSelectionMapping::Default | (bMapFacesToEdgeLoops ? EEnumerateSelectionMapping::FacesToEdges : EEnumerateSelectionMapping::None)
 		);
 }
 
@@ -1175,7 +1180,7 @@ bool UE::Geometry::EnumeratePolygroupSelectionElements(
 	TFunctionRef<void(int32, const FSegment3d&)> EdgeFunc,
 	TFunctionRef<void(int32, const FTriangle3d&)> TriangleFunc,
 	const FTransform* ApplyTransform,
-	const EEnumerateMappingFlags Flags
+	const EEnumerateSelectionMapping Flags
 )
 {
 if ( ensure( MeshSelection.TopologyType == EGeometryTopologyType::Polygroup ) == false )
@@ -1183,13 +1188,13 @@ if ( ensure( MeshSelection.TopologyType == EGeometryTopologyType::Polygroup ) ==
 		return false;
 	}
 
-	auto ProcessGroupEdgeID = [GroupTopology, &Mesh, ApplyTransform, &EdgeFunc, Flags](const int32 GroupEdgeID)
+	auto ProcessGroupEdgeID = [GroupTopology, &Mesh, ApplyTransform, &EdgeFunc](const int32 GroupEdgeID)
 	{
 		for (const int32 EdgeID : GroupTopology->GetGroupEdgeEdges(GroupEdgeID))
 		{
 			FVector3d A, B;
 			Mesh.GetEdgeV(EdgeID, A, B);
-			if (ApplyTransform != nullptr)
+			if (ApplyTransform)
 			{
 				A = ApplyTransform->TransformPosition(A);
 				B = ApplyTransform->TransformPosition(B);
@@ -1207,13 +1212,15 @@ if ( ensure( MeshSelection.TopologyType == EGeometryTopologyType::Polygroup ) ==
 			if (Mesh.IsTriangle(SeedTriangleID))
 			{
 				// Call TriangleFunc if we are mapping Faces to Faces
-				if ((Flags & EEnumerateMappingFlags::FacesToFaces) != EEnumerateMappingFlags::None)
+				// Note: While EnumeratePolygroupSelectionTriangles would also return all faces, implementing it this way minimizes redundancy in the
+				// MeshSelection.Selection loop and more consistent with the rest of the Polygroup selection elements enumeration
+				if ((Flags & EEnumerateSelectionMapping::FacesToFaces) != EEnumerateSelectionMapping::None)
 				{
 					for (const int32 TriangleID : GroupTopology->GetGroupFaces(GroupID))
 					{
 						FVector3d A, B, C;
 						Mesh.GetTriVertices(TriangleID, A, B, C);
-						if (ApplyTransform != nullptr)
+						if (ApplyTransform)
 						{
 							A = ApplyTransform->TransformPosition(A);
 							B = ApplyTransform->TransformPosition(B);
@@ -1224,7 +1231,7 @@ if ( ensure( MeshSelection.TopologyType == EGeometryTopologyType::Polygroup ) ==
 				}
 
 				// Process the polygroup edges if we are mapping Faces to Edges
-				if ((Flags & EEnumerateMappingFlags::FacesToEdges) != EEnumerateMappingFlags::None)
+				if ((Flags & EEnumerateSelectionMapping::FacesToEdges) != EEnumerateSelectionMapping::None)
 				{
 					TArray<int32> GroupEdgeIDs;
 					if ( const FGroupTopology::FGroup* Group = GroupTopology->FindGroupByID(GroupID) )
@@ -1242,8 +1249,6 @@ if ( ensure( MeshSelection.TopologyType == EGeometryTopologyType::Polygroup ) ==
 						ProcessGroupEdgeID(GroupEdgeID);
 					}
 				}
-
-				// TODO [katie.caputo] add FaceToVerts path
 			}
 		}
 	}
@@ -1257,13 +1262,7 @@ if ( ensure( MeshSelection.TopologyType == EGeometryTopologyType::Polygroup ) ==
 			{
 				const int32 GroupEdgeID = GroupTopology->FindGroupEdgeID(SeedEdgeID);
 
-				// Process the polygroup edges if we are mapping Edges to Edges
-				if ((Flags & EEnumerateMappingFlags::EdgesToEdges) != EEnumerateMappingFlags::None)
-				{
-					ProcessGroupEdgeID(GroupEdgeID);
-				}
-
-				// TODO [katie.caputo]: Add EdgesToVerts and EdgesToFaces path
+				ProcessGroupEdgeID(GroupEdgeID);
 			}
 		}
 	}
@@ -1275,14 +1274,8 @@ if ( ensure( MeshSelection.TopologyType == EGeometryTopologyType::Polygroup ) ==
 			if (Mesh.IsVertex(VertexID))
 			{
 				const FVector3d A = Mesh.GetVertex(VertexID);
-
-				// Call VertexFunc if we are mapping Vertices to Vertices
-				if ((Flags & EEnumerateMappingFlags::VertsToVerts) != EEnumerateMappingFlags::None)
-				{
-					VertexFunc(VertexID, (ApplyTransform != nullptr) ? ApplyTransform->TransformPosition(A) : A);
-				}
-
-				// todo [katie.caputo]: Add VertsToEdges and VertsToFaces path
+				
+				VertexFunc(VertexID, (ApplyTransform) ? ApplyTransform->TransformPosition(A) : A);
 			}
 		}
 	}
@@ -1540,9 +1533,7 @@ bool UE::Geometry::ConvertSelection(
 		checkSlow(ToSelectionOut.TopologyType == EGeometryTopologyType::Triangle);
 		checkSlow(ToSelectionOut.ElementType == EGeometryElementType::Vertex);
 
-		// TODO Add a function which gets the Vids only, to remove the matrix-vector multiplication we just ignore
-		const FTransform Transform = FTransform::Identity;
-		return EnumerateTriangleSelectionVertices(FromSelectionIn, Mesh, Transform,
+		return EnumerateTriangleSelectionVertices(FromSelectionIn, Mesh, nullptr,
 			[&ToSelectionOut](uint64 Vid, const FVector3d& Unused)
 			{
 				ToSelectionOut.Selection.Add( FGeoSelectionID::MeshVertex((int32)Vid).Encoded() );
@@ -1661,9 +1652,7 @@ bool UE::Geometry::ConvertTriangleSelectionToOverlaySelection(
 			TrianglesOut.Add(ValidTid);
 		});
 
-		// TODO Add a function which gets the Vids only, to remove the matrix-vector multiplication we just ignore
-		const FTransform Transform = FTransform::Identity;
-		EnumerateTriangleSelectionVertices(MeshSelection, Mesh, Transform,
+		EnumerateTriangleSelectionVertices(MeshSelection, Mesh, nullptr,
 			[&VerticesOut, IncidentSelection](uint64 Vid, const FVector3d& Unused)
 		{
 			IncidentSelection->Selection.Add(FGeoSelectionID::MeshVertex((int)Vid).Encoded() );
@@ -1678,9 +1667,7 @@ bool UE::Geometry::ConvertTriangleSelectionToOverlaySelection(
 			TrianglesOut.Add(ValidTid);
 		});
 
-		// TODO Add a function which gets the Vids only, to remove the matrix-vector multiplication we just ignore
-		const FTransform Transform = FTransform::Identity;
-		EnumerateTriangleSelectionVertices(MeshSelection, Mesh, Transform,
+		EnumerateTriangleSelectionVertices(MeshSelection, Mesh, nullptr,
 			[&VerticesOut](uint64 Vid, const FVector3d& Unused)
 		{
 			VerticesOut.Add((int)Vid);
diff --git a/Engine/Plugins/Runtime/GeometryProcessing/Source/DynamicMesh/Public/Selections/GeometrySelectionUtil.h b/Engine/Plugins/Runtime/GeometryProcessing/Source/DynamicMesh/Public/Selections/GeometrySelectionUtil.h
--- a/Engine/Plugins/Runtime/GeometryProcessing/Source/DynamicMesh/Public/Selections/GeometrySelectionUtil.h
+++ b/Engine/Plugins/Runtime/GeometryProcessing/Source/DynamicMesh/Public/Selections/GeometrySelectionUtil.h
@@ -19,25 +19,16 @@ class FColliderMesh;
 class FGroupTopology;
 struct FGroupTopologySelection;
 	
-enum class EEnumerateMappingFlags : uint16
+enum class EEnumerateSelectionMapping : uint8
 {
 	None         = 0,
 	
 	FacesToFaces = 1 << 0,
 	FacesToEdges = 1 << 1,
-	FacesToVerts = 1 << 2, // not used yet
 
-	EdgesToFaces = 1 << 3, // not used yet
-	EdgesToEdges = 1 << 4,
-	EdgesToVerts = 1 << 5, // not used yet
-
-	VertsToFaces = 1 << 6, // not used yet
-	VertsToEdges = 1 << 7, // not used yet
-	VertsToVerts = 1 << 8,
-
-	Default = FacesToFaces | EdgesToEdges | VertsToVerts,
+	Default = FacesToFaces,
 };
-ENUM_CLASS_FLAGS(EEnumerateMappingFlags);
+ENUM_CLASS_FLAGS(EEnumerateSelectionMapping);
 	
 /**
  * Test if SelectionA and SelectionB are the same selection.
@@ -103,6 +94,18 @@ DYNAMICMESH_API bool UpdateSelectionWithNewElements(
 	FGeometrySelectionDelta* DeltaOut = nullptr);
 
 
+/**
+ * Prefer pointer version for ApplyTransform in EnumerateTriangleSelectionVertices
+ */
+UE_DEPRECATED(5.5, "EnumerateTriangleSelectionVertices with ApplyTransform passed a reference is deprecated. "
+				   "Please use the function of the same name which takes ApplyTransform as a pointer instead")
+DYNAMICMESH_API bool EnumerateTriangleSelectionVertices(
+	const FGeometrySelection& MeshSelection,
+	const UE::Geometry::FDynamicMesh3& Mesh,
+	const FTransform& ApplyTransform,
+	TFunctionRef<void(uint64, const FVector3d&)> VertexFunc
+);
+
 /**
  * Call VertexFunc for each selected Mesh element (vertex/edge/tri) in MeshSelection.
  * ApplyTransform will be applied to Vertex Positions before calling VertexFunc
@@ -110,13 +113,17 @@ DYNAMICMESH_API bool UpdateSelectionWithNewElements(
 DYNAMICMESH_API bool EnumerateTriangleSelectionVertices(
 	const FGeometrySelection& MeshSelection,
 	const UE::Geometry::FDynamicMesh3& Mesh,
-	const FTransform& ApplyTransform,
+	const FTransform* ApplyTransform,
 	TFunctionRef<void(uint64, const FVector3d&)> VertexFunc
 );
+
 /**
  * Call VertexFunc for each selected Mesh element (vertex/edge/tri) in the set
  * of polygroup faces/edges/corners specified by GroupSelection (relative to GroupTopology parameter)
  * ApplyTransform will be applied to Vertex Positions before calling VertexFunc
+ * 
+ * Enumerates all vertices in the polygroup, including internal, non-Corner vertices
+ * For retrieving Corners only, see EnumeratePolygroupSelectionElements
  */
 DYNAMICMESH_API bool EnumeratePolygroupSelectionVertices(
 	const FGeometrySelection& GroupSelection,
@@ -168,6 +175,7 @@ DYNAMICMESH_API bool EnumeratePolygroupSelectionTriangles(
  * This will forward to EnumerateTriangleSelectionEdges() or 
  * EnumeratePolygroupSelectionEdges() depending on the selection topology type.
  * If UseGroupSet and MeshSelection is for polygroups, the default Mesh group layer will be used.
+ * For polygroups - all edges and vertices within the polygroup are included, not exclusively Group Boundaries or Group Corners
  */
 DYNAMICMESH_API bool EnumerateSelectionEdges(
 	const FGeometrySelection& MeshSelection,
@@ -178,6 +186,7 @@ DYNAMICMESH_API bool EnumerateSelectionEdges(
 /**
  * Call EdgeFunc for each mesh EdgeID included in MeshSelection.
  * For Triangles, all 3 edges are included.
+ * For Edges, all selected edges are included
  * For Vertices, all edges touching each selected vertex are included.
  */
 DYNAMICMESH_API bool EnumerateTriangleSelectionEdges(
@@ -187,8 +196,11 @@ DYNAMICMESH_API bool EnumerateTriangleSelectionEdges(
 );
 /**
  * Call EdgeFunc for each mesh EdgeID included in MeshSelection, where MeshSelection has polygroup topology.
- * For Polygroup Faces, all edges in the face are included.
- * For Polygroup Corners, currently includes all edges in any group touching the corner
+ * For Polygroup Faces, all edges (including internal, non-border edges) in the polygroup are included.
+ * For Polygroup Borders, includes all edges (including internal, non-border edges) in both groups that the Border Edge touches
+ * For Polygroup Corners, currently includes all edges (including internal, non-border edges) in all groups touching the Corner
+ *
+ * For rendering purposes and retrieving only BorderEdges, see EnumeratePolygroupSelectionElements
  */
 DYNAMICMESH_API bool EnumeratePolygroupSelectionEdges(
 	const FGeometrySelection& MeshSelection,
@@ -210,10 +222,19 @@ DYNAMICMESH_API bool EnumerateTriangleSelectionElements(
 
 
 /**
- * Call VertexFunc/EdgeFunc/TriangleFunc for the vertices/edges/triangles identified by MeshSelection.
- * Since a MeshSelection only stores vertices, edges, or triangles, but not combined, only one
- * of these functions will be invoked during a call to this function.
- * This function is useful to collect up geometry that needs to be rendered for a given MeshSelection
+ * Enumerates the renderable geometry for a given MeshSelection when in Triangle Topology mode.
+ *
+ * Calls VertexFunc/EdgeFunc/TriangleFunc for the vertices/edges/triangles
+ * identified by MeshSelection. Since a MeshSelection only stores vertices
+ * edges, or triangles, but not combined, only one of these functions will be
+ * invoked during a call to this function.
+ *
+ * This function employs the component specific variants (ex.
+ * EnumerateTriangleSelectionTriangles) but the overall intention of this function
+ * is to collect the renderable geometry representations of the selection.
+ * In triangle topology mode only, there is not a difference between 
+ * all selection elements and the desired rendered elements. However, this function
+ * is still used for render representation for consistency and clarity.
  * @param ApplyTransform if non-null, transform is applied to the 3D geometry
  */
 DYNAMICMESH_API bool EnumerateTriangleSelectionElements(
@@ -223,7 +244,7 @@ DYNAMICMESH_API bool EnumerateTriangleSelectionElements(
 	TFunctionRef<void(int32, const FSegment3d&)> EdgeFunc,
 	TFunctionRef<void(int32, const FTriangle3d&)> TriangleFunc,
 	const FTransform* ApplyTransform = nullptr,
-	const EEnumerateMappingFlags Flags = EEnumerateMappingFlags::Default
+	const EEnumerateSelectionMapping Flags = EEnumerateSelectionMapping::Default
 );
 
 /** Prefer EnumeratePolygroupSelectionElements with Flags parameter. */
@@ -238,12 +259,20 @@ DYNAMICMESH_API bool EnumeratePolygroupSelectionElements(
 	const bool bMapFacesToEdgeLoops = false
 );
 /**
- * Call VertexFunc/EdgeFunc/TriangleFunc for the mesh vertices/edges/triangles identified by MeshSelection,
- * where MeshSelection has polygroup topology referring to the provided GroupTopology
- * Since a MeshSelection only stores vertices, edges, or triangles, but not combined, only one
- * of these functions will be invoked during a call to this function.
- * This function is useful to collect up geometry that needs to be rendered for a given MeshSelection
- * @param ApplyTransform if non-null, transform is applied to the 3D geometry
+ * Enumerates the renderable geometry for a given MeshSelection when in Polygroup Topology mode.
+ *
+ * Calls VertexFunc/EdgeFunc/TriangleFunc for the vertices/edges/triangles
+ * identified by MeshSelection. Since a MeshSelection only stores vertices
+ * edges, or triangles, but not combined, only one of these functions will be
+ * invoked during a call to this function.
+ *
+ * This function differs from the component specific variants (ex.
+ * EnumeratePolygroupSelectionTriangles) by only collecting the renderable
+ * geometry representations of the selection. For instance, a polygroup
+ * face may only want to render the triangles and border edges, while
+ * EnumeratePolygroupSelectionTriangles would collect all triangles and
+ * edges within that polygroup face. This is the difference between render
+ * representation and component conversion.
  */
 DYNAMICMESH_API bool EnumeratePolygroupSelectionElements(
 	const FGeometrySelection& MeshSelection,
@@ -253,7 +282,7 @@ DYNAMICMESH_API bool EnumeratePolygroupSelectionElements(
 	TFunctionRef<void(int32, const FSegment3d&)> EdgeFunc,
 	TFunctionRef<void(int32, const FTriangle3d&)> TriangleFunc,
 	const FTransform* ApplyTransform = nullptr,
-	const EEnumerateMappingFlags Flags = EEnumerateMappingFlags::Default
+	const EEnumerateSelectionMapping Flags = EEnumerateSelectionMapping::Default
 );
 
 
diff --git a/Engine/Plugins/Runtime/MeshModelingToolset/Source/ModelingComponents/Private/Selection/DynamicMeshSelector.cpp b/Engine/Plugins/Runtime/MeshModelingToolset/Source/ModelingComponents/Private/Selection/DynamicMeshSelector.cpp
--- a/Engine/Plugins/Runtime/MeshModelingToolset/Source/ModelingComponents/Private/Selection/DynamicMeshSelector.cpp
+++ b/Engine/Plugins/Runtime/MeshModelingToolset/Source/ModelingComponents/Private/Selection/DynamicMeshSelector.cpp
@@ -646,7 +646,7 @@ void FBaseDynamicMeshSelector::AccumulateSelectionBounds(const FGeometrySelectio
 	{
 		TargetMesh->ProcessMesh([&](const UE::Geometry::FDynamicMesh3& SourceMesh)
 		{
-			UE::Geometry::EnumerateTriangleSelectionVertices(Selection, SourceMesh, UseTransform,
+			UE::Geometry::EnumerateTriangleSelectionVertices(Selection, SourceMesh, &UseTransform,
 				[&](uint32 VertexID, const FVector3d& Position) { 
 					TargetWorldBounds.Contain(Position); 
 					ElementCount++;
@@ -665,20 +665,11 @@ void FBaseDynamicMeshSelector::AccumulateSelectionBounds(const FGeometrySelectio
 
 void FBaseDynamicMeshSelector::AccumulateSelectionElements(const FGeometrySelection& Selection, FGeometrySelectionElements& Elements, bool bTransformToWorld, bool bIsForPreview)
 {
-	const FTransform UseWorldTransform = GetLocalToWorldTransform();
-	const FTransform* ApplyTransform = (bTransformToWorld) ? &UseWorldTransform : nullptr;
-	const bool bMapFacesToEdges = bIsForPreview;
-
-	ToolSelectionUtil::AccumulateSelectionElements(
-		Elements,
-		Selection,
-		TargetMesh->GetMeshRef(),
-		Selection.TopologyType == EGeometryTopologyType::Polygroup ? GetGroupTopology() : nullptr,
-		ApplyTransform,
-		bMapFacesToEdges);
+	// where if bIsForPreview is true, we are mapping faces to edges
+	AccumulateSelectionElements(Selection, Elements, bTransformToWorld, EEnumerateSelectionMapping::Default | (bIsForPreview ? EEnumerateSelectionMapping::FacesToEdges : EEnumerateSelectionMapping::None));
 }
 
-void FBaseDynamicMeshSelector::AccumulateSelectionElements(const FGeometrySelection& Selection, FGeometrySelectionElements& Elements, bool bTransformToWorld, const EEnumerateMappingFlags Flags)
+void FBaseDynamicMeshSelector::AccumulateSelectionElements(const FGeometrySelection& Selection, FGeometrySelectionElements& Elements, bool bTransformToWorld, const EEnumerateSelectionMapping Flags)
 {
 	const FTransform UseWorldTransform = GetLocalToWorldTransform();
 	const FTransform* ApplyTransform = (bTransformToWorld) ? &UseWorldTransform : nullptr;
@@ -700,7 +691,7 @@ void FBaseDynamicMeshSelector::AccumulateElementsFromPredicate(FGeometrySelectio
 		Selection.ElementType = ElementType;
 		Selection.TopologyType = bUseGroupTopology ? EGeometryTopologyType::Polygroup : EGeometryTopologyType::Triangle;
 		InitializeSelectionFromPredicate(Selection, [&Predicate, ElementType](FGeoSelectionID InID){ return Predicate(ElementType, InID); });
-		AccumulateSelectionElements(Selection, Elements, bTransformToWorld, bIsForPreview);
+		AccumulateSelectionElements(Selection, Elements, bTransformToWorld, EEnumerateSelectionMapping::Default);
 	};
 
 	AccumulateElementsOfTypeFromPredicate(EGeometryElementType::Vertex);
@@ -894,7 +885,7 @@ void FBasicDynamicMeshSelectionTransformer::BeginTransform(const FGeometrySelect
 		}
 		else
 		{
-			UE::Geometry::EnumerateTriangleSelectionVertices(Selection, SourceMesh, FTransform::Identity,
+			UE::Geometry::EnumerateTriangleSelectionVertices(Selection, SourceMesh, nullptr,
 				[&](uint32 VertexID, const FVector3d& Position) { VertexIDs.Add((int32)VertexID); }
 			);
 		}
diff --git a/Engine/Plugins/Runtime/MeshModelingToolset/Source/ModelingComponents/Private/Selection/GeometrySelectionManager.cpp b/Engine/Plugins/Runtime/MeshModelingToolset/Source/ModelingComponents/Private/Selection/GeometrySelectionManager.cpp
--- a/Engine/Plugins/Runtime/MeshModelingToolset/Source/ModelingComponents/Private/Selection/GeometrySelectionManager.cpp
+++ b/Engine/Plugins/Runtime/MeshModelingToolset/Source/ModelingComponents/Private/Selection/GeometrySelectionManager.cpp
@@ -1482,7 +1482,7 @@ void UGeometrySelectionManager::RebuildSelectionRenderCache()
 		
 		FGeometrySelectionElements& SelectionElements = CachedSelectionRenderElements[k];
 		SelectionElements.Reset();
-		Target->Selector->AccumulateSelectionElements(Target->Selection, SelectionElements, true, false);
+		Target->Selector->AccumulateSelectionElements(Target->Selection, SelectionElements, true, EEnumerateSelectionMapping::Default);
 	}
 
 	bSelectionRenderCachesDirty = false;
@@ -1530,11 +1530,11 @@ void UGeometrySelectionManager::RebuildPreviewRenderCache()
 	
 	if (ActivePreviewSelection.IsEmpty() == false)
 	{
-		EEnumerateMappingFlags MappingFlags = EEnumerateMappingFlags::Default | EEnumerateMappingFlags::FacesToEdges;
+		EEnumerateSelectionMapping MappingFlags = EEnumerateSelectionMapping::Default | EEnumerateSelectionMapping::FacesToEdges;
 		if (CVarGeometrySelectionManager_FullSelectionHoverHighlights.GetValueOnGameThread() == 0)
 		{
 			// Unset FacesToFaces flag if full hover highlights are disabled
-			MappingFlags &= ~EEnumerateMappingFlags::FacesToFaces;
+			MappingFlags &= ~EEnumerateSelectionMapping::FacesToFaces;
 		}
 		
 		Target->Selector->AccumulateSelectionElements(SelectedActivePreviewSelection, CachedSelectedPreviewRenderElements, true, MappingFlags);
diff --git a/Engine/Plugins/Runtime/MeshModelingToolset/Source/ModelingComponents/Private/Selection/ToolSelectionUtil.cpp b/Engine/Plugins/Runtime/MeshModelingToolset/Source/ModelingComponents/Private/Selection/ToolSelectionUtil.cpp
--- a/Engine/Plugins/Runtime/MeshModelingToolset/Source/ModelingComponents/Private/Selection/ToolSelectionUtil.cpp
+++ b/Engine/Plugins/Runtime/MeshModelingToolset/Source/ModelingComponents/Private/Selection/ToolSelectionUtil.cpp
@@ -56,7 +56,7 @@ bool ToolSelectionUtil::AccumulateSelectionElements(
 		SourceMesh,
 		GroupTopology,
 		ApplyTransform,
-		EEnumerateMappingFlags::Default | (bMapFacesToEdges ? EEnumerateMappingFlags::FacesToEdges : EEnumerateMappingFlags::None));
+		EEnumerateSelectionMapping::Default | (bMapFacesToEdges ? EEnumerateSelectionMapping::FacesToEdges : EEnumerateSelectionMapping::None));
 }
 
 bool ToolSelectionUtil::AccumulateSelectionElements(
@@ -65,7 +65,7 @@ bool ToolSelectionUtil::AccumulateSelectionElements(
 	const FDynamicMesh3& SourceMesh,
 	const FGroupTopology* GroupTopology,
 	const FTransform* ApplyTransform,
-	const EEnumerateMappingFlags Flags)
+	const EEnumerateSelectionMapping Flags)
 {
 	auto AddPoint = [&Elements](uint32 Vid, const FVector3d& Point) { Elements.Points.Add(Point); };
 	auto AddSegment = [&Elements](uint32 Eid, const FSegment3d& Segment) { Elements.Segments.Add(Segment); };
@@ -169,7 +169,7 @@ void FSelectionRenderHelper::Initialize(const FGeometrySelection& Selection,
 	const FGroupTopology* Topology,
 	const FTransform* ApplyTransform)
 {
-	bool bSuccess = ToolSelectionUtil::AccumulateSelectionElements(Elements, Selection, SourceMesh, Topology, ApplyTransform, false);
+	bool bSuccess = ToolSelectionUtil::AccumulateSelectionElements(Elements, Selection, SourceMesh, Topology, ApplyTransform, EEnumerateSelectionMapping::Default | EEnumerateSelectionMapping::None);
 	ensure(bSuccess);
 }
 
diff --git a/Engine/Plugins/Runtime/MeshModelingToolset/Source/ModelingComponents/Public/Selection/DynamicMeshSelector.h b/Engine/Plugins/Runtime/MeshModelingToolset/Source/ModelingComponents/Public/Selection/DynamicMeshSelector.h
--- a/Engine/Plugins/Runtime/MeshModelingToolset/Source/ModelingComponents/Public/Selection/DynamicMeshSelector.h
+++ b/Engine/Plugins/Runtime/MeshModelingToolset/Source/ModelingComponents/Public/Selection/DynamicMeshSelector.h
@@ -118,8 +118,13 @@ public:
 	virtual void GetSelectionFrame(const FGeometrySelection& Selection, UE::Geometry::FFrame3d& SelectionFrame, bool bTransformToWorld) override;
 	virtual void GetTargetFrame(const FGeometrySelection& Selection, UE::Geometry::FFrame3d& SelectionFrame) override;
 	virtual void AccumulateSelectionBounds(const FGeometrySelection& Selection, FGeometrySelectionBounds& BoundsInOut, bool bTransformToWorld) override;
+
+	UE_DEPRECATED(5.5, "AccumulateSelectionElements which takes a bIsForPreview boolean is deprecated."
+				"Please use the function of the same name which takes EEnumerateSelectionMapping flags instead")
 	virtual void AccumulateSelectionElements(const FGeometrySelection& Selection, FGeometrySelectionElements& Elements, bool bTransformToWorld, bool bIsForPreview) override;
-	virtual void AccumulateSelectionElements(const FGeometrySelection& Selection, FGeometrySelectionElements& Elements, bool bTransformToWorld, UE::Geometry::EEnumerateMappingFlags Flags = UE::Geometry::EEnumerateMappingFlags::Default) override;
+
+	/** Prefer AccumulateSelectionElements with Flags parameter. */
+	virtual void AccumulateSelectionElements(const FGeometrySelection& Selection, FGeometrySelectionElements& Elements, bool bTransformToWorld, UE::Geometry::EEnumerateSelectionMapping Flags = UE::Geometry::EEnumerateSelectionMapping::Default) override;
 	
 	virtual void AccumulateElementsFromPredicate(
 		FGeometrySelectionElements& Elements,
diff --git a/Engine/Plugins/Runtime/MeshModelingToolset/Source/ModelingComponents/Public/Selection/GeometrySelector.h b/Engine/Plugins/Runtime/MeshModelingToolset/Source/ModelingComponents/Public/Selection/GeometrySelector.h
--- a/Engine/Plugins/Runtime/MeshModelingToolset/Source/ModelingComponents/Public/Selection/GeometrySelector.h
+++ b/Engine/Plugins/Runtime/MeshModelingToolset/Source/ModelingComponents/Public/Selection/GeometrySelector.h
@@ -410,15 +410,17 @@ public:
 	virtual void AccumulateSelectionBounds(const FGeometrySelection& Selection, FGeometrySelectionBounds& BoundsInOut, bool bTransformToWorld) = 0;
 
 	/** Prefer AccumulateSelectionElements with Flags parameter. */
-	virtual void AccumulateSelectionElements(const FGeometrySelection& Selection, FGeometrySelectionElements& ElementsInOut, bool bTransformToWorld, UE::Geometry::EEnumerateMappingFlags Flags = UE::Geometry::EEnumerateMappingFlags::Default) = 0;
+	UE_DEPRECATED(5.5, "AccumulateSelectionElements which takes a bIsForPreview boolean is deprecated."
+				"Please use the function of the same name which takes EEnumerateSelectionMapping flags instead")
+	virtual void AccumulateSelectionElements(const FGeometrySelection& Selection, FGeometrySelectionElements& ElementsInOut, bool bTransformToWorld, bool bIsForPreview) = 0;
 	
 	/**
 	 * Accumulate geometric elements (currently 3D triangles, line segments, and points) for the provided Selection in the provided ElementsInOut. 
 	 * ElementsInOut is not cleared.
 	 * @param bTransformToWorld if true each geometric element will be transformed to World space, based on GetLocalToWorldTransform()
-	 * @param bIsForPreview if true, geometry is being collected for a preview of selection. Selector may return simplified geometry in this case
+	 * @param Flags determines which pieces of the geometry should be rendered
 	 */
-	virtual void AccumulateSelectionElements(const FGeometrySelection& Selection, FGeometrySelectionElements& ElementsInOut, bool bTransformToWorld, bool bIsForPreview) = 0;
+	virtual void AccumulateSelectionElements(const FGeometrySelection& Selection, FGeometrySelectionElements& ElementsInOut, bool bTransformToWorld, UE::Geometry::EEnumerateSelectionMapping Flags = UE::Geometry::EEnumerateSelectionMapping::Default) = 0;
 
 	/**
 	 * Accumulate all geometric elements (currently 3D triangles, line segments, and points) in the provided ElementsInOut
diff --git a/Engine/Plugins/Runtime/MeshModelingToolset/Source/ModelingComponents/Public/Selection/ToolSelectionUtil.h b/Engine/Plugins/Runtime/MeshModelingToolset/Source/ModelingComponents/Public/Selection/ToolSelectionUtil.h
--- a/Engine/Plugins/Runtime/MeshModelingToolset/Source/ModelingComponents/Public/Selection/ToolSelectionUtil.h
+++ b/Engine/Plugins/Runtime/MeshModelingToolset/Source/ModelingComponents/Public/Selection/ToolSelectionUtil.h
@@ -36,6 +36,8 @@ namespace ToolSelectionUtil
 	MODELINGCOMPONENTS_API void SetNewActorSelection(UInteractiveToolManager* ToolManager, const TArray<AActor*>& Actors);
 
 	/** Prefer AccumulateSelectionElements with Flags parameter. */
+	UE_DEPRECATED(5.5, "AccumulateSelectionElements which takes a bMapFacesToEdges boolean is deprecated."
+				"Please use the function of the same name which takes EEnumerateSelectionMapping flags instead")
 	MODELINGCOMPONENTS_API bool AccumulateSelectionElements(
 		UE::Geometry::FGeometrySelectionElements& Elements,
 		const UE::Geometry::FGeometrySelection& Selection,
@@ -57,7 +59,7 @@ namespace ToolSelectionUtil
 		const UE::Geometry::FDynamicMesh3& SourceMesh,
 		const UE::Geometry::FGroupTopology* Topology = nullptr,
 		const FTransform* ApplyTransform = nullptr,
-		const UE::Geometry::EEnumerateMappingFlags Flags = UE::Geometry::EEnumerateMappingFlags::Default);
+		const UE::Geometry::EEnumerateSelectionMapping Flags = UE::Geometry::EEnumerateSelectionMapping::Default);
 	
 	/**
 	 * Render the given Elements using FPrimitiveDrawInterface
