Modeling Tools: Corrects and clarifys EnumeratePolygroup function descriptions

diff --git a/Engine/Plugins/Runtime/GeometryProcessing/Source/DynamicMesh/Public/Selections/GeometrySelectionUtil.h b/Engine/Plugins/Runtime/GeometryProcessing/Source/DynamicMesh/Public/Selections/GeometrySelectionUtil.h
--- a/Engine/Plugins/Runtime/GeometryProcessing/Source/DynamicMesh/Public/Selections/GeometrySelectionUtil.h
+++ b/Engine/Plugins/Runtime/GeometryProcessing/Source/DynamicMesh/Public/Selections/GeometrySelectionUtil.h
@@ -109,6 +109,10 @@ DYNAMICMESH_API bool EnumerateTriangleSelectionVertices(
 /**
  * Call VertexFunc for each selected Mesh element (vertex/edge/tri) in MeshSelection.
  * ApplyTransform will be applied to Vertex Positions before calling VertexFunc
+ *
+ * For Triangles: includes all 3 points which make up each selected triangle
+ * For Edges: includes both vertex endpoints of each selected edge
+ * For Vertices: includes each selected vertex
  */
 DYNAMICMESH_API bool EnumerateTriangleSelectionVertices(
 	const FGeometrySelection& MeshSelection,
@@ -121,9 +125,12 @@ DYNAMICMESH_API bool EnumerateTriangleSelectionVertices(
  * Call VertexFunc for each selected Mesh element (vertex/edge/tri) in the set
  * of polygroup faces/edges/corners specified by GroupSelection (relative to GroupTopology parameter)
  * ApplyTransform will be applied to Vertex Positions before calling VertexFunc
- * 
- * Enumerates all vertices in the polygroup, including internal, non-Corner vertices
- * For retrieving Corners only, see EnumeratePolygroupSelectionElements
+ *
+ * For Faces: enumerates all vertices in the selected polygroup(s), including internal, non-Corner vertices
+ * For Edges/Boundaries: enumerates the TriVerts which make up the selected PolyEdge(s)
+ * For Verts/Corners: returns the TriVert that corresponds to the selected Corner(s)
+ *
+ * For rendering purposes, use EnumeratePolygroupSelectionElements
  */
 DYNAMICMESH_API bool EnumeratePolygroupSelectionVertices(
 	const FGeometrySelection& GroupSelection,
@@ -197,10 +204,10 @@ DYNAMICMESH_API bool EnumerateTriangleSelectionEdges(
 /**
  * Call EdgeFunc for each mesh EdgeID included in MeshSelection, where MeshSelection has polygroup topology.
  * For Polygroup Faces, all edges (including internal, non-border edges) in the polygroup are included.
- * For Polygroup Borders, includes all edges (including internal, non-border edges) in both groups that the Border Edge touches
- * For Polygroup Corners, currently includes all edges (including internal, non-border edges) in all groups touching the Corner
+ * For Polygroup Borders, includes all mesh edges which are a part of the selected GroupEdge(s)/BorderEdge(s)
+ * For Polygroup Corners, includes all mesh edges in all GroupEdge(s) which touch the selected Corner(s)
  *
- * For rendering purposes and retrieving only BorderEdges, see EnumeratePolygroupSelectionElements
+ * For rendering purposes, use EnumeratePolygroupSelectionElements
  */
 DYNAMICMESH_API bool EnumeratePolygroupSelectionEdges(
 	const FGeometrySelection& MeshSelection,
