fill out more of the selection conversion matrix in geometry selection util

diff --git a/Engine/Plugins/Runtime/GeometryProcessing/Source/DynamicMesh/Private/Selections/GeometrySelectionUtil.cpp b/Engine/Plugins/Runtime/GeometryProcessing/Source/DynamicMesh/Private/Selections/GeometrySelectionUtil.cpp
--- a/Engine/Plugins/Runtime/GeometryProcessing/Source/DynamicMesh/Private/Selections/GeometrySelectionUtil.cpp
+++ b/Engine/Plugins/Runtime/GeometryProcessing/Source/DynamicMesh/Private/Selections/GeometrySelectionUtil.cpp
@@ -1350,6 +1350,98 @@ bool UE::Geometry::ConvertPolygroupSelectionToTopologySelection(
 	return true;
 }
 
+bool UE::Geometry::InitializeSelectionFromPolyGroups(
+	const FDynamicMesh3& Mesh,
+	const FGroupTopology& GroupTopology,
+	TArrayView<const int> PolyGroupIDs,
+	FGeometrySelection& SelectionOut)
+{
+	bool bSuccess = true;
+	if (SelectionOut.TopologyType == EGeometryTopologyType::Triangle)
+	{
+		for (int32 GID : PolyGroupIDs)
+		{
+			const FGroupTopology::FGroup* Group = GroupTopology.FindGroupByID(GID);
+			if (!Group)
+			{
+				bSuccess = false;
+				continue;
+			}
+			for (int32 TID : Group->Triangles)
+			{
+				if (!Mesh.IsTriangle(TID))
+				{
+					continue;
+				}
+				if (SelectionOut.ElementType == EGeometryElementType::Face)
+				{
+					SelectionOut.Selection.Add(FGeoSelectionID::MeshTriangle(TID).Encoded());
+				}
+				else if (SelectionOut.ElementType == EGeometryElementType::Edge)
+				{
+					FIndex3i TriEdges = Mesh.GetTriEdges(TID);
+					SelectionOut.Selection.Add(FGeoSelectionID::MeshEdge(Mesh.GetTriEdgeIDFromEdgeID(TriEdges.A)).Encoded());
+					SelectionOut.Selection.Add(FGeoSelectionID::MeshEdge(Mesh.GetTriEdgeIDFromEdgeID(TriEdges.B)).Encoded());
+					SelectionOut.Selection.Add(FGeoSelectionID::MeshEdge(Mesh.GetTriEdgeIDFromEdgeID(TriEdges.C)).Encoded());
+				}
+				else if (SelectionOut.ElementType == EGeometryElementType::Vertex)
+				{
+					FIndex3i Tri = Mesh.GetTriangle(TID);
+					SelectionOut.Selection.Add(FGeoSelectionID::MeshVertex(Tri.A).Encoded());
+					SelectionOut.Selection.Add(FGeoSelectionID::MeshVertex(Tri.B).Encoded());
+					SelectionOut.Selection.Add(FGeoSelectionID::MeshVertex(Tri.C).Encoded());
+				}
+			}
+		}
+	}
+	else if (SelectionOut.TopologyType == EGeometryTopologyType::Polygroup)
+	{
+		for (int32 GID : PolyGroupIDs)
+		{
+			const FGroupTopology::FGroup* Group = GroupTopology.FindGroupByID(GID);
+			if (!Group)
+			{
+				bSuccess = false;
+				continue;
+			}
+			if (SelectionOut.ElementType == EGeometryElementType::Face)
+			{
+				FGeoSelectionID ID = FGeoSelectionID(Group->Triangles[0], Group->GroupID);
+				SelectionOut.Selection.Add(ID.Encoded());
+			}
+			else
+			{
+				for (const FGroupTopology::FGroupBoundary& Boundary : Group->Boundaries)
+				{
+					for (int GroupEdgeID : Boundary.GroupEdges)
+					{
+						if (GroupTopology.Edges.IsValidIndex(GroupEdgeID))
+						{
+							if (SelectionOut.ElementType == EGeometryElementType::Edge)
+							{
+								const FGroupTopology::FGroupEdge& GroupEdge = GroupTopology.Edges[GroupEdgeID];
+								FMeshTriEdgeID MeshEdgeID = Mesh.GetTriEdgeIDFromEdgeID(GroupEdge.Span.Edges[0]);
+								FGeoSelectionID ID = FGeoSelectionID(MeshEdgeID.Encoded(), GroupEdgeID);
+								SelectionOut.Selection.Add(ID.Encoded());
+							}
+							else if (SelectionOut.ElementType == EGeometryElementType::Vertex)
+							{
+								for (int32 CornerIdx = 0; CornerIdx < 2; ++CornerIdx)
+								{
+									int32 CornerID = GroupTopology.Edges[GroupEdgeID].EndpointCorners[CornerIdx];
+									const FGroupTopology::FCorner& Corner = GroupTopology.Corners[CornerID];
+									FGeoSelectionID ID = FGeoSelectionID(Corner.VertexID, CornerID);
+									SelectionOut.Selection.Add(ID.Encoded());
+								}
+							}
+						}
+					}
+				}
+			}
+		}
+	}
+	return bSuccess;
+}
 
 
 bool UE::Geometry::InitializeSelectionFromTriangles(
@@ -1541,6 +1633,90 @@ bool UE::Geometry::ConvertSelection(
 	};
 
 
+	const auto FromPolyFace = [](
+		const FDynamicMesh3& Mesh,
+		const FGroupTopology* GroupTopology,
+		const FGeometrySelection& FromSelectionIn,
+		FGeometrySelection& ToSelectionOut)
+	{
+		checkSlow(FromSelectionIn.TopologyType == EGeometryTopologyType::Polygroup);
+		checkSlow(FromSelectionIn.ElementType == EGeometryElementType::Face);
+		if (!ensure(GroupTopology))
+		{
+			return false;
+		}
+
+		TArray<int32> Groups;
+		for (uint64 ElemID : FromSelectionIn.Selection)
+		{
+			Groups.Add(FGeoSelectionID(ElemID).TopologyID);
+		}
+
+		return InitializeSelectionFromPolyGroups(Mesh, *GroupTopology, Groups, ToSelectionOut);
+	};
+	
+
+	const auto FromPolyEdgeToTriEdge = [](
+		const FDynamicMesh3& Mesh,
+		const FGroupTopology* GroupTopology,
+		const FGeometrySelection& FromSelectionIn,
+		FGeometrySelection& ToSelectionOut) -> bool
+	{
+		checkSlow(FromSelectionIn.TopologyType == EGeometryTopologyType::Polygroup);
+		checkSlow(FromSelectionIn.ElementType == EGeometryElementType::Edge);
+		checkSlow(ToSelectionOut.TopologyType == EGeometryTopologyType::Triangle);
+		checkSlow(ToSelectionOut.ElementType == EGeometryElementType::Edge);
+
+		if (ensure(GroupTopology))
+		{
+			for (uint64 EncodedID : FromSelectionIn.Selection)
+			{
+				const FGeoSelectionID SelectionID(EncodedID);
+				int32 GroupEID = SelectionID.TopologyID;
+				checkSlow(GroupTopology->Edges.IsValidIndex(GroupEID));
+				for (int32 EID : GroupTopology->Edges[GroupEID].Span.Edges)
+				{
+					FMeshTriEdgeID MeshEdgeID = Mesh.GetTriEdgeIDFromEdgeID(EID);
+					ToSelectionOut.Selection.Add(MeshEdgeID.Encoded());
+				}
+			}
+			return true;
+		}
+		return false;
+	};
+
+
+	const auto FromPolyEdgeToPolyVtx = [](
+		const FDynamicMesh3& Mesh,
+		const FGroupTopology* GroupTopology,
+		const FGeometrySelection& FromSelectionIn,
+		FGeometrySelection& ToSelectionOut) -> bool
+	{
+		checkSlow(FromSelectionIn.TopologyType == EGeometryTopologyType::Polygroup);
+		checkSlow(FromSelectionIn.ElementType == EGeometryElementType::Edge);
+		checkSlow(ToSelectionOut.TopologyType == EGeometryTopologyType::Polygroup);
+		checkSlow(ToSelectionOut.ElementType == EGeometryElementType::Vertex);
+
+		if (ensure(GroupTopology))
+		{
+			for (uint64 EncodedID : FromSelectionIn.Selection)
+			{
+				const FGeoSelectionID SelectionID(EncodedID);
+				int32 GroupEID = SelectionID.TopologyID;
+				checkSlow(GroupTopology->Edges.IsValidIndex(GroupEID));
+				FIndex2i Corners = GroupTopology->Edges[GroupEID].EndpointCorners;
+				for (int32 SubIdx = 0; SubIdx < 2; ++SubIdx)
+				{
+					int32 CornerID = Corners[SubIdx];
+					const FGroupTopology::FCorner& Corner = GroupTopology->Corners[CornerID];
+					FGeoSelectionID ID = FGeoSelectionID(Corner.VertexID, CornerID);
+					ToSelectionOut.Selection.Add(ID.Encoded());
+				}
+			}
+			return true;
+		}
+		return false;
+	};
 
 	const auto FromPolyEdgeToTriVtx = [](
 		const FDynamicMesh3& Mesh,
@@ -1562,7 +1738,33 @@ bool UE::Geometry::ConvertSelection(
 			});
 	};
 
+	const auto FromTriSelectionToPolyVtx = [](
+		const FDynamicMesh3& Mesh,
+		const FGroupTopology* GroupTopology,
+		const FGeometrySelection& FromSelectionIn,
+		FGeometrySelection& ToSelectionOut) -> bool
+	{
+		checkSlow(FromSelectionIn.TopologyType == EGeometryTopologyType::Triangle);
+		checkSlow(ToSelectionOut.TopologyType == EGeometryTopologyType::Polygroup);
+		checkSlow(ToSelectionOut.ElementType == EGeometryElementType::Vertex);
 
+		bool bConverted = false;
+		EnumerateTriangleSelectionVertices(FromSelectionIn, Mesh, nullptr, 
+			[&Mesh, &GroupTopology, &ToSelectionOut, &bConverted](uint64 VID, const FVector& Unused)
+			{
+				if (Mesh.IsVertex((int32)VID))
+				{
+					int32 CornerID = GroupTopology->GetCornerIDFromVertexID((int32)VID);
+					if (CornerID != IndexConstants::InvalidID)
+					{
+						bConverted = true;
+						FGeoSelectionID ID = FGeoSelectionID((int32)VID, CornerID);
+						ToSelectionOut.Selection.Add(ID.Encoded());
+					}
+				}
+			});
+		return bConverted;
+	};
 
 	const auto FromPolyVtxToTriVtx = [](
 		const FDynamicMesh3& Mesh,
@@ -1595,12 +1797,12 @@ bool UE::Geometry::ConvertSelection(
 			FGeometrySelection& ToSelectionOut);
 
 	constexpr ConvertSelectionFunc ConvertFuncs[6][6] = {
-		{FromTypeToSame,        NotImplemented,  NotImplemented,  NotImplemented,  NotImplemented,  NotImplemented},
-		{FromTriEdgeToTriVtx,   FromTypeToSame,  NotImplemented,  NotImplemented,  NotImplemented,  NotImplemented},
-		{FromTriFace,           FromTriFace,     FromTypeToSame,  FromTriFace,     FromTriFace,     FromTriFace   },
-		{FromPolyVtxToTriVtx,   NotImplemented,  NotImplemented,  FromTypeToSame,  NotImplemented,  NotImplemented},
-		{FromPolyEdgeToTriVtx,  NotImplemented,  NotImplemented,  NotImplemented,  FromTypeToSame,  NotImplemented},
-		{NotImplemented,        NotImplemented,  NotImplemented,  NotImplemented,  NotImplemented,  FromTypeToSame}
+		{FromTypeToSame,        NotImplemented,         NotImplemented,  FromTriSelectionToPolyVtx,   NotImplemented,         NotImplemented},
+		{FromTriEdgeToTriVtx,   FromTypeToSame,         NotImplemented,  FromTriSelectionToPolyVtx,   NotImplemented,         NotImplemented},
+		{FromTriFace,           FromTriFace,            FromTypeToSame,  FromTriFace,                 FromTriFace,            FromTriFace   },
+		{FromPolyVtxToTriVtx,   NotImplemented,         NotImplemented,  FromTypeToSame,              NotImplemented,         NotImplemented},
+		{FromPolyEdgeToTriVtx,  FromPolyEdgeToTriEdge,  NotImplemented,  FromPolyEdgeToPolyVtx,       FromTypeToSame,         NotImplemented},
+		{FromPolyFace,          FromPolyFace,           FromPolyFace,    FromPolyFace,                FromPolyFace,           FromTypeToSame}
 	};
 
 	const int FromIndex = GetSelectionTypeAsIndex(FromSelectionIn);
diff --git a/Engine/Plugins/Runtime/GeometryProcessing/Source/DynamicMesh/Public/Selections/GeometrySelectionUtil.h b/Engine/Plugins/Runtime/GeometryProcessing/Source/DynamicMesh/Public/Selections/GeometrySelectionUtil.h
--- a/Engine/Plugins/Runtime/GeometryProcessing/Source/DynamicMesh/Public/Selections/GeometrySelectionUtil.h
+++ b/Engine/Plugins/Runtime/GeometryProcessing/Source/DynamicMesh/Public/Selections/GeometrySelectionUtil.h
@@ -314,6 +314,17 @@ DYNAMICMESH_API bool InitializeSelectionFromTriangles(
 	FGeometrySelection& SelectionOut);
 
 
+/**
+ * Convert PolyGroup IDs to target Selection type
+ */
+DYNAMICMESH_API bool InitializeSelectionFromPolyGroups(
+	const FDynamicMesh3& Mesh,
+	const FGroupTopology& GroupTopology,
+	TArrayView<const int> GroupIDs,
+	FGeometrySelection& SelectionOut);
+
+
+
 /**
  * Convert Selection from one type to another, based on geometry/topology types in FromSelectionIn and ToSelectionOut.
  * 
@@ -323,12 +334,12 @@ DYNAMICMESH_API bool InitializeSelectionFromTriangles(
  *                 To:    Triangle               Polygroup           
  *   From:                Vertex  Edge    Face   Vertex  Edge    Face
  *   ----------------------------------------------------------------
- *   Triangle Vertex      1       .       .      .       .       .   
- *   Triangle Edge        1       1       .      .       .       .   
+ *   Triangle Vertex      1       .       .      4#      .       .   
+ *   Triangle Edge        1       1       .      4#      .       .   
  *   Triangle Face        1       1       1      4#      3#      2#  
  *   Polygroup Vertex     6#      .       .      1       .       .   
- *   Polygroup Edge       5#      .       .      .       1       .   
- *   Polygroup Face       .       .       .      .       .       1   
+ *   Polygroup Edge       5#      1#      .      1#      1       .   
+ *   Polygroup Face       1#      1#      1#     1#      1#      1   
  *   ================================================================
  *
  *   Key:
@@ -337,7 +348,7 @@ DYNAMICMESH_API bool InitializeSelectionFromTriangles(
  *   2  supported. Polygroup faces containing any input triangle are selected
  *   3  supported. Polygroup edges containing any input triangle edge are selected, but
  *                 polygroup edges containing only input triangle vertices are not.
- *   4  supported. Polygroup corners coinciding with any input triangle vertex are selected
+ *   4  supported. Polygroup corners coinciding with any input triangle face/edge/vertex are selected
  *   5  supported. All mesh vertices along the polygroup edge are selected
  *   6  supported. All mesh vertices coinciding with polygroup corners are selected
  *   #  indicates GroupTopology must not be null for this combination. If this symbol is missing GroupTopology is ignored
