fix crash on smooth bevel of a confusingly triangulated door mesh

diff --git a/Engine/Plugins/Runtime/GeometryProcessing/Source/DynamicMesh/Private/Operations/MeshBevel.cpp b/Engine/Plugins/Runtime/GeometryProcessing/Source/DynamicMesh/Private/Operations/MeshBevel.cpp
--- a/Engine/Plugins/Runtime/GeometryProcessing/Source/DynamicMesh/Private/Operations/MeshBevel.cpp
+++ b/Engine/Plugins/Runtime/GeometryProcessing/Source/DynamicMesh/Private/Operations/MeshBevel.cpp
@@ -2105,6 +2105,12 @@ void FMeshBevel::AppendJunctionVertexPolygon_Multi(FDynamicMesh3& Mesh, FBevelVe
 				bFound = true;
 			};
 
+			if (Span.StripQuadPatch.VertexSpans.IsEmpty())
+			{
+				// The no-vertex-span case can happen for cases where the smooth bevel has failed and fallen back to non-smooth bevel
+				continue;
+			}
+			
 			// try start and end columns, and handle case where vertex ordering might be reversed  (should this be possible, because of consistent mesh winding??)
 			Span.StripQuadPatch.GetVertexColumn(0, BevelEdgeQuadStripEndVertices);
 			if (BevelEdgeQuadStripEndVertices[0] == A && BevelEdgeQuadStripEndVertices.Last() == B)
