Add a geometry processing command line program to enable command line execution of geometry processing algorithms
+ add navigation-driven approximate convex decomposition as the first supported algorithm

diff --git a/Engine/Plugins/Runtime/GeometryProcessing/GeometryProcessing.uplugin b/Engine/Plugins/Runtime/GeometryProcessing/GeometryProcessing.uplugin
--- a/Engine/Plugins/Runtime/GeometryProcessing/GeometryProcessing.uplugin
+++ b/Engine/Plugins/Runtime/GeometryProcessing/GeometryProcessing.uplugin
@@ -23,6 +23,11 @@
 			"Name": "DynamicMesh",
 			"Type": "Runtime",
 			"LoadingPhase": "Default"
+		},
+		{
+			"Name": "MeshFileUtils",
+			"Type": "DeveloperTool",
+			"LoadingPhase": "Default"
 		}
 	],
 	"SupportedPrograms": ["ChaosVisualDebugger"]
diff --git a/Engine/Plugins/Runtime/GeometryProcessing/Source/MeshFileUtils/MeshFileUtils.build.cs b/Engine/Plugins/Runtime/GeometryProcessing/Source/MeshFileUtils/MeshFileUtils.build.cs
new file mode 100644
--- /dev/null
+++ b/Engine/Plugins/Runtime/GeometryProcessing/Source/MeshFileUtils/MeshFileUtils.build.cs
@@ -0,0 +1,23 @@
+// Copyright Epic Games, Inc. All Rights Reserved.
+
+using UnrealBuildTool;
+
+public class MeshFileUtils : ModuleRules
+{	
+	public MeshFileUtils(ReadOnlyTargetRules Target) : base(Target)
+	{
+		PCHUsage = ModuleRules.PCHUsageMode.UseExplicitOrSharedPCHs;
+
+		// Note: The module purposefully doesn't have a dependency on CoreUObject.
+		// If possible, we would like avoid having UObjects in GeometryProcessing
+		// modules to keep the door open for writing standalone command-line programs
+		// (which won't have UObject garbage collection).
+        PublicDependencyModuleNames.AddRange(
+			new string[] {
+				"Core",
+				"GeometryCore"
+			}
+		);
+		UnsafeTypeCastWarningLevel = WarningLevel.Error;
+	}
+}
diff --git a/Engine/Plugins/Runtime/GeometryProcessing/Source/MeshFileUtils/Private/MeshFileUtilsModule.cpp b/Engine/Plugins/Runtime/GeometryProcessing/Source/MeshFileUtils/Private/MeshFileUtilsModule.cpp
new file mode 100644
--- /dev/null
+++ b/Engine/Plugins/Runtime/GeometryProcessing/Source/MeshFileUtils/Private/MeshFileUtilsModule.cpp
@@ -0,0 +1,20 @@
+// Copyright Epic Games, Inc. All Rights Reserved.
+
+#include "MeshFileUtilsModule.h"
+
+#define LOCTEXT_NAMESPACE "FMeshFileUtilsModule"
+
+void FMeshFileUtilsModule::StartupModule()
+{
+	// This code will execute after your module is loaded into memory; the exact timing is specified in the .uplugin file per-module
+}
+
+void FMeshFileUtilsModule::ShutdownModule()
+{
+	// This function may be called during shutdown to clean up your module.  For modules that support dynamic reloading,
+	// we call this function before unloading the module.
+}
+
+#undef LOCTEXT_NAMESPACE
+	
+IMPLEMENT_MODULE(FMeshFileUtilsModule, MeshFileUtils)
\ No newline at end of file
diff --git a/Engine/Plugins/Runtime/GeometryProcessing/Source/MeshFileUtils/Private/OBJMeshUtil.cpp b/Engine/Plugins/Runtime/GeometryProcessing/Source/MeshFileUtils/Private/OBJMeshUtil.cpp
new file mode 100644
--- /dev/null
+++ b/Engine/Plugins/Runtime/GeometryProcessing/Source/MeshFileUtils/Private/OBJMeshUtil.cpp
@@ -0,0 +1,443 @@
+// Copyright Epic Games, Inc. All Rights Reserved.
+
+#include "OBJMeshUtil.h"
+
+#include "DynamicMesh/DynamicMesh3.h"
+#include "DynamicMesh/DynamicMeshAttributeSet.h"
+
+#include <fstream>
+#include <string>
+#include <sstream>
+#include <iomanip>
+#include <iostream>
+
+using namespace UE::Geometry;
+
+namespace UE::Private::MeshFileUtilsLocals
+{
+
+/**
+ * Write obj element: vertex, UV, normal or polygonal face (with only vertex indices).
+ * 
+ * @param Out    File stream we are writing into.
+ * @param Token  One of "v", "vn", "vt", "f".
+ * @param Value  Vector type containing the data to write. ith entry will be accessed with the [] operator.
+ * 
+ * @note Assumes Out is open.
+ */
+template<int Dim, class VectorType> 
+void ObjWriteElement(std::ofstream& Out, const char* Token, const VectorType& Value) 
+{
+	Out << Token << ' ';
+	for (int32 Idx = 0; Idx < Dim; ++Idx) 
+	{
+		Out << Value[Idx];
+		if (Idx < Dim - 1) 
+		{
+			Out << ' ';
+		}
+	}
+	Out << '\n';
+}
+
+void ObjWriteString(std::ofstream& Out, const FString& Line) 
+{ 
+	Out << TCHAR_TO_ANSI(*Line);
+}
+
+} // end namespace MeshFileUtilsLocals
+
+namespace UE::MeshFileUtils
+{
+	ELoadOBJStatus ReadOBJ(std::istream& FileStream, FDynamicMesh3& Mesh, const FLoadOBJSettings& Settings)
+	{
+		auto ExtractTripletFromGroup = [](std::stringstream& InLineStream,
+			bool bInParseNormalsAndTextures,
+			int& OutVertexID,
+			int& OutTextureID,
+			int& OutNormalID)
+		{
+			std::string Blob;
+
+			OutVertexID = -1;
+			OutTextureID = -1;
+			OutNormalID = -1;
+
+			if (InLineStream >> Blob)
+			{
+				const char* BlobCStr = Blob.c_str();
+				const char** Token = &BlobCStr;
+
+				OutVertexID = FCStringAnsi::Atoi((*Token));
+				(*Token) += FCStringAnsi::Strcspn((*Token), "/");
+				if ((*Token)[0] != '/')
+				{
+					return true;
+				}
+
+				if (!bInParseNormalsAndTextures)
+				{
+					return true;
+				}
+
+				(*Token)++;
+
+				// v//n case
+				if ((*Token)[0] == '/')
+				{
+					(*Token)++;
+					OutNormalID = FCStringAnsi::Atoi((*Token));
+					(*Token) += FCStringAnsi::Strcspn((*Token), "/");
+					return true;
+				}
+
+				// v/t/n or v/t case
+				OutTextureID = FCStringAnsi::Atoi((*Token));
+				(*Token) += FCStringAnsi::Strcspn((*Token), "/");
+				if ((*Token)[0] != '/')
+				{
+					return true;
+				}
+
+				// v/t/n case
+				(*Token)++;  // skip '/'
+				OutNormalID = FCStringAnsi::Atoi((*Token));
+				(*Token) += FCStringAnsi::Strcspn((*Token), "/");
+
+				return true;
+			}
+
+			return false;
+		};
+
+		Mesh.Clear();
+
+		struct FTriToAdd
+		{
+			FIndex3i Tri, UV, Normal;
+			bool bHasUV = false;
+			bool bHasNormal = false;
+			int32 GroupID = 0;
+			FTriToAdd() = default;
+			FTriToAdd(FIndex3i Tri, int32 GroupID) : Tri(Tri), GroupID(GroupID)
+			{
+			}
+			void SetUV(FIndex3i InUV)
+			{
+				UV = InUV;
+				bHasUV = true;
+			}
+			void SetNormal(FIndex3i InNormal)
+			{
+				Normal = InNormal;
+				bHasNormal = true;
+			}
+		};
+
+		// buffer the triangles in case the faces precede the vertices they reference in the file
+		TArray<FTriToAdd> Triangles;
+
+		const bool bLoadNormalsOrUVs = Settings.bLoadNormals || Settings.bLoadUVs;
+
+		if (bLoadNormalsOrUVs)
+		{
+			Mesh.EnableAttributes();
+		}
+
+		FDynamicMeshNormalOverlay* Normals = (Settings.bLoadNormals) ? Mesh.Attributes()->PrimaryNormals() : nullptr;
+		FDynamicMeshUVOverlay* UVs = (Settings.bLoadUVs) ? Mesh.Attributes()->PrimaryUV() : nullptr;
+
+
+		int32 GroupID = 1;
+		bool bHasGroups = false;
+		std::string Line;
+		while (std::getline(FileStream, Line))
+		{
+			if (Line.empty())
+			{
+				continue;
+			}
+
+			std::stringstream LineStream(Line);
+			std::string Command;
+
+			if (LineStream >> Command)
+			{
+				if (Command.empty() || Command[0] == '#') // skip comment lines
+				{
+					continue;
+				}
+
+				if (Command == "g" || Command == "o")
+				{
+					bHasGroups = true;
+					GroupID++;
+				}
+				else if (Command == "v") // vertex
+				{
+					FVector3d Vert;
+					LineStream >> Vert.X >> Vert.Y >> Vert.Z;
+					Mesh.AppendVertex(Vert);
+				}
+				else if (Settings.bLoadNormals && Command == "vn")
+				{
+					FVector3f Normal;
+					LineStream >> Normal.X >> Normal.Y >> Normal.Z;
+					Normalize(Normal);
+					Normals->AppendElement(Normal);
+				}
+				else if (Settings.bLoadUVs && Command == "vt")
+				{
+					FVector2f UV = FVector2f::Zero(); // v is optional so default to zero
+					LineStream >> UV.X >> UV.Y;
+					UVs->AppendElement(UV);
+				}
+				else if (Command == "f")
+				{
+					int V0, VN_1, VN; // vertices 0, N-1, and N (the last two are reused to make a triangle fan out of any non-triangle faces)
+					int T0, TN_1, TN; // UVs
+					int N0, NN_1, NN; // normals
+
+					if (!ExtractTripletFromGroup(LineStream, bLoadNormalsOrUVs, V0, T0, N0))
+					{
+						continue;
+					}
+					if (!ExtractTripletFromGroup(LineStream, bLoadNormalsOrUVs, VN_1, TN_1, NN_1))
+					{
+						continue;
+					}
+					while (ExtractTripletFromGroup(LineStream, bLoadNormalsOrUVs, VN, TN, NN))
+					{
+						FTriToAdd& ToAdd = Triangles.Emplace_GetRef(FIndex3i(V0 - 1, VN_1 - 1, VN - 1), GroupID);
+						if (Normals)
+						{
+							ToAdd.SetNormal(FIndex3i(N0 - 1, NN_1 - 1, NN - 1));
+						}
+						if (UVs)
+						{
+							ToAdd.SetUV(FIndex3i(T0 - 1, TN_1 - 1, TN - 1));
+						}
+
+						VN_1 = VN;
+						TN_1 = TN;
+						NN_1 = NN;
+					}
+				}
+			}
+		}
+
+		if (bHasGroups)
+		{
+			Mesh.EnableTriangleGroups();
+		}
+		for (const FTriToAdd& ToAdd : Triangles)
+		{
+			int32 TID = Mesh.AppendTriangle(ToAdd.Tri, ToAdd.GroupID);
+			if (Settings.bAddSeparatedTriForNonManifold && TID == FDynamicMesh3::NonManifoldID)
+			{
+				FIndex3i DupeVertsTri;
+				for (int32 SubIdx = 0; SubIdx < 3; ++SubIdx)
+				{
+					DupeVertsTri[SubIdx] = Mesh.AppendVertex(Mesh.GetVertex(ToAdd.Tri[SubIdx]));
+				}
+				TID = Mesh.AppendTriangle(DupeVertsTri, ToAdd.GroupID);
+			}
+			if (TID >= 0)
+			{
+				if (Normals && ToAdd.bHasNormal)
+				{
+					Normals->SetTriangle(TID, ToAdd.Normal);
+				}
+				if (UVs && ToAdd.bHasUV)
+				{
+					UVs->SetTriangle(TID, ToAdd.UV);
+				}
+			}
+		}
+
+		if (Settings.bReverseOrientation)
+		{
+			Mesh.ReverseOrientation();
+		}
+
+		return ELoadOBJStatus::Success;
+	}
+
+	ELoadOBJStatus LoadOBJ(const char* Path, FDynamicMesh3& Mesh, const FLoadOBJSettings& Settings)
+	{
+		std::ifstream FileStream(Path);
+		if (!FileStream)
+		{
+			return ELoadOBJStatus::InvalidPath;
+		}
+
+		ELoadOBJStatus Status = ReadOBJ(FileStream, Mesh, Settings);
+		FileStream.close();
+
+		return Status;
+	}
+
+	FDynamicMesh3 LoadOBJChecked(const char* Path, const FLoadOBJSettings& Settings)
+	{
+		FDynamicMesh3 Mesh;
+		ELoadOBJStatus Status = LoadOBJ(Path, Mesh, Settings);
+		check(Status == ELoadOBJStatus::Success);
+		return Mesh;
+	}
+
+	bool WriteOBJ(const char* Path, const FDynamicMesh3& InMesh, const FWriteOBJSettings& Settings)
+	{
+		// We compact the mesh to make sure the order that we write the vertex/normal/uv indices into obj is 
+		// consistent with the storage order 
+		FDynamicMesh3 Mesh;
+		Mesh.CompactCopy(InMesh);
+
+		if (Settings.bReverseOrientation)
+		{
+			Mesh.ReverseOrientation();
+		}
+
+		std::ofstream FileStream;
+		FileStream.precision(std::numeric_limits<double>::digits10);
+		FileStream.open(Path, std::ofstream::out | std::ofstream::trunc);
+		if (!ensure(FileStream))
+		{
+			return false;
+		}
+
+		bool bHasVertexNormals = Settings.bWritePerVertexValues && Mesh.HasVertexNormals();
+		bool bHasVertexUVs = Settings.bWritePerVertexValues && Mesh.HasVertexUVs();
+
+		for (int32 VID = 0; VID < Mesh.VertexCount(); ++VID)
+		{
+			check(Mesh.IsVertex(VID)); // mesh is not compact
+
+			FVector3d Pos = Mesh.GetVertex(VID);
+			if (!Settings.bWritePerVertexColors || !Mesh.HasVertexColors())
+			{
+				UE::Private::MeshFileUtilsLocals::ObjWriteElement<3>(FileStream, "v", Pos);
+			}
+			else
+			{
+				FVector3f Color = Mesh.GetVertexColor(VID);
+				FileStream << "v " << Pos[0] << " " << Pos[1] << " " << Pos[2] << " " << Color[0] << " " << Color[1] << " " << Color[2] << "\n";
+			}
+
+			if (bHasVertexNormals)
+			{
+				FVector3f Normal = Mesh.GetVertexNormal(VID);
+				UE::Private::MeshFileUtilsLocals::ObjWriteElement<3>(FileStream, "vn", Normal);
+			}
+
+			if (bHasVertexUVs)
+			{
+				FVector2f UV = Mesh.GetVertexUV(VID);
+				UE::Private::MeshFileUtilsLocals::ObjWriteElement<2>(FileStream, "vt", UV);
+			}
+		}
+
+		const FDynamicMeshUVOverlay* UVs = nullptr;
+		const FDynamicMeshNormalOverlay* Normals = nullptr;
+
+		if (Settings.bWritePerVertexValues == false && Mesh.Attributes())
+		{
+			UVs = Mesh.Attributes()->PrimaryUV();
+			if (UVs)
+			{
+				for (int32 UI = 0; UI < UVs->ElementCount(); ++UI)
+				{
+					check(UVs->IsElement(UI))
+						FVector2f UV = UVs->GetElement(UI);
+					UE::Private::MeshFileUtilsLocals::ObjWriteElement<2>(FileStream, "vt", UV);
+				}
+			}
+
+			Normals = Mesh.Attributes()->PrimaryNormals();
+			if (Normals)
+			{
+				for (int32 NI = 0; NI < Normals->ElementCount(); ++NI)
+				{
+					check(Normals->IsElement(NI));
+					FVector3f Normal = Normals->GetElement(NI);
+					UE::Private::MeshFileUtilsLocals::ObjWriteElement<3>(FileStream, "vn", Normal);
+				}
+			}
+		}
+
+		for (int32 TID = 0; TID < Mesh.TriangleCount(); ++TID)
+		{
+			check(Mesh.IsTriangle(TID));
+
+			FIndex3i TriVertices = Mesh.GetTriangle(TID);
+
+			if (Settings.bWritePerVertexValues)
+			{
+				if (bHasVertexNormals == false && bHasVertexUVs == false)
+				{
+					UE::Private::MeshFileUtilsLocals::ObjWriteString(FileStream, FString::Printf(TEXT("f %d %d %d\n"), TriVertices.A + 1, TriVertices.B + 1, TriVertices.C + 1));
+				}
+				else if (bHasVertexNormals == true && bHasVertexUVs == false)
+				{
+					UE::Private::MeshFileUtilsLocals::ObjWriteString(FileStream, FString::Printf(TEXT("f %d//%d %d//%d %d//%d\n"),
+						TriVertices.A + 1, TriVertices.A + 1,
+						TriVertices.B + 1, TriVertices.B + 1,
+						TriVertices.C + 1, TriVertices.C + 1));
+				}
+				else if (bHasVertexNormals == false && bHasVertexUVs == true)
+				{
+					UE::Private::MeshFileUtilsLocals::ObjWriteString(FileStream, FString::Printf(TEXT("f %d/%d %d/%d %d/%d\n"),
+						TriVertices.A + 1, TriVertices.A + 1,
+						TriVertices.B + 1, TriVertices.B + 1,
+						TriVertices.C + 1, TriVertices.C + 1));
+				}
+				else
+				{
+					UE::Private::MeshFileUtilsLocals::ObjWriteString(FileStream, FString::Printf(TEXT("f %d/%d/%d %d/%d/%d %d/%d/%d\n"),
+						TriVertices.A + 1, TriVertices.A + 1, TriVertices.A + 1,
+						TriVertices.B + 1, TriVertices.B + 1, TriVertices.B + 1,
+						TriVertices.C + 1, TriVertices.C + 1, TriVertices.C + 1));
+				}
+			}
+			else
+			{
+				bool bHaveUV = UVs != nullptr && UVs->IsSetTriangle(TID);
+				bool bHaveNormal = Normals != nullptr && Normals->IsSetTriangle(TID);
+
+				FIndex3i TriUVs = bHaveUV ? UVs->GetTriangle(TID) : FIndex3i::Invalid();
+				FIndex3i TriNormals = bHaveNormal ? Normals->GetTriangle(TID) : FIndex3i::Invalid();
+
+				if (bHaveUV && bHaveNormal)
+				{
+					UE::Private::MeshFileUtilsLocals::ObjWriteString(FileStream, FString::Printf(TEXT("f %d/%d/%d %d/%d/%d %d/%d/%d\n"),
+						TriVertices.A + 1, TriUVs.A + 1, TriNormals.A + 1,
+						TriVertices.B + 1, TriUVs.B + 1, TriNormals.B + 1,
+						TriVertices.C + 1, TriUVs.C + 1, TriNormals.C + 1));
+				}
+				else if (bHaveUV)
+				{
+					UE::Private::MeshFileUtilsLocals::ObjWriteString(FileStream, FString::Printf(TEXT("f %d/%d %d/%d %d/%d\n"),
+						TriVertices.A + 1, TriUVs.A + 1,
+						TriVertices.B + 1, TriUVs.B + 1,
+						TriVertices.C + 1, TriUVs.C + 1));
+				}
+				else if (bHaveNormal)
+				{
+					UE::Private::MeshFileUtilsLocals::ObjWriteString(FileStream, FString::Printf(TEXT("f %d//%d %d//%d %d//%d\n"),
+						TriVertices.A + 1, TriNormals.A + 1,
+						TriVertices.B + 1, TriNormals.B + 1,
+						TriVertices.C + 1, TriNormals.C + 1));
+				}
+				else
+				{
+					UE::Private::MeshFileUtilsLocals::ObjWriteString(FileStream, FString::Printf(TEXT("f %d %d %d\n"), TriVertices.A + 1, TriVertices.B + 1, TriVertices.C + 1));
+				}
+			}
+		}
+
+		FileStream.close();
+
+		return true;
+	}
+
+}
+
diff --git a/Engine/Plugins/Runtime/GeometryProcessing/Source/MeshFileUtils/Public/MeshFileUtilsModule.h b/Engine/Plugins/Runtime/GeometryProcessing/Source/MeshFileUtils/Public/MeshFileUtilsModule.h
new file mode 100644
--- /dev/null
+++ b/Engine/Plugins/Runtime/GeometryProcessing/Source/MeshFileUtils/Public/MeshFileUtilsModule.h
@@ -0,0 +1,15 @@
+// Copyright Epic Games, Inc. All Rights Reserved.
+
+#pragma once
+
+#include "CoreMinimal.h"
+#include "Modules/ModuleManager.h"
+
+class FMeshFileUtilsModule : public IModuleInterface
+{
+public:
+
+	/** IModuleInterface implementation */
+	virtual void StartupModule() override;
+	virtual void ShutdownModule() override;
+};
diff --git a/Engine/Plugins/Runtime/GeometryProcessing/Source/MeshFileUtils/Public/OBJMeshUtil.h b/Engine/Plugins/Runtime/GeometryProcessing/Source/MeshFileUtils/Public/OBJMeshUtil.h
new file mode 100644
--- /dev/null
+++ b/Engine/Plugins/Runtime/GeometryProcessing/Source/MeshFileUtils/Public/OBJMeshUtil.h
@@ -0,0 +1,54 @@
+// Copyright Epic Games, Inc. All Rights Reserved.
+
+#pragma once
+
+#include "DynamicMesh/DynamicMesh3.h"
+
+// Developer utilities for managing Wavefront OBJ format mesh files
+// Intended for debugging and for use in tests and utilities; may not support every OBJ file format feature
+namespace UE::MeshFileUtils
+{
+	enum class ELoadOBJStatus
+	{
+		InvalidPath,
+		Success
+	};
+
+	struct FLoadOBJSettings
+	{
+
+		// Dynamic mesh doesn't support non-manifold edges (edges with more than two triangles), but OBJ does
+		// If this setting is true, non-manifold triangles will be added as fully separated triangles;
+		// otherwise, such triangles will be skipped entirely
+		bool bAddSeparatedTriForNonManifold = true;
+
+		// If true, reverses the orientation of the faces
+		bool bReverseOrientation = true;
+
+		// If true and obj contains normal information will compute normal overlay
+		bool bLoadNormals = false;
+
+		// If true and obj contains UVs information will compute UV overlay.
+		bool bLoadUVs = false;
+	};
+
+	struct FWriteOBJSettings
+	{
+		// If true reverses the orientation of the faces
+		bool bReverseOrientation = true;
+
+		// If true will attempt to write the per-vertex normals and UVs to the OBJ instead of the per-element values
+		bool bWritePerVertexValues = true;
+
+		// Whether to write per-vertex colors (when available)
+		bool bWritePerVertexColors = false;
+	};
+
+	// Attempt to load an OBJ file into a Dynamic Mesh
+	ELoadOBJStatus LoadOBJ(const char* Path, UE::Geometry::FDynamicMesh3& Mesh, const FLoadOBJSettings& Settings = FLoadOBJSettings());
+	// Load an OBJ file into a Dynamic Mesh, or fail a check() if it cannot be loaded
+	UE::Geometry::FDynamicMesh3 LoadOBJChecked(const char* Path, const FLoadOBJSettings& Settings = FLoadOBJSettings());
+	// Write an OBJ file to the target path; returns true on success
+	bool WriteOBJ(const char* Path, const UE::Geometry::FDynamicMesh3& InMesh, const FWriteOBJSettings& Settings = FWriteOBJSettings());
+}
+
diff --git a/Engine/Source/Programs/GeometryProcessing/GeometryProcessing.Build.cs b/Engine/Source/Programs/GeometryProcessing/GeometryProcessing.Build.cs
new file mode 100644
--- /dev/null
+++ b/Engine/Source/Programs/GeometryProcessing/GeometryProcessing.Build.cs
@@ -0,0 +1,29 @@
+// Copyright Epic Games, Inc. All Rights Reserved.
+
+using UnrealBuildTool;
+
+public class GeometryProcessing : ModuleRules
+{
+	public GeometryProcessing(ReadOnlyTargetRules Target) : base(Target)
+	{
+		PublicIncludePathModuleNames.Add("Launch");
+
+		PrivateDependencyModuleNames.AddRange(
+			new string[] {
+						"Core",
+						"Projects",
+						"GeometryCore",
+						"GeometryAlgorithms",
+						"DynamicMesh",
+						"MeshFileUtils",
+			}
+		);
+
+		// to link with CoreUObject module:
+		// PrivateDependencyModuleNames.Add("CoreUObject");
+
+		// to enable LLM tracing:
+		// GlobalDefinitions.Add("LLM_ENABLED_IN_CONFIG=1");
+		// GlobalDefinitions.Add("UE_MEMORY_TAGS_TRACE_ENABLED=1");
+	}
+}
diff --git a/Engine/Source/Programs/GeometryProcessing/GeometryProcessing.Target.cs b/Engine/Source/Programs/GeometryProcessing/GeometryProcessing.Target.cs
new file mode 100644
--- /dev/null
+++ b/Engine/Source/Programs/GeometryProcessing/GeometryProcessing.Target.cs
@@ -0,0 +1,33 @@
+// Copyright Epic Games, Inc. All Rights Reserved.
+
+using UnrealBuildTool;
+using System.Collections.Generic;
+
+[SupportedPlatforms(UnrealPlatformClass.All)]
+public class GeometryProcessingTarget : TargetRules
+{
+	public GeometryProcessingTarget(TargetInfo Target) : base(Target)
+	{
+		Type = TargetType.Program;
+		IncludeOrderVersion = EngineIncludeOrderVersion.Latest;
+		LinkType = TargetLinkType.Monolithic;
+		LaunchModuleName = "GeometryProcessing";
+
+		// Lean and mean
+		bBuildDeveloperTools = false;
+
+		bBuildWithEditorOnlyData = true;
+
+		// Currently this app is not linking against the engine, so we'll compile out references from Core to the rest of the engine
+		bCompileAgainstEngine = false;
+		bCompileAgainstCoreUObject = false;
+		bCompileAgainstApplicationCore = false;
+		bCompileICU = false;
+
+		// to enable tracing:
+		// bEnableTrace = true;
+
+		// This app is a console application, not a Windows app (sets entry point to main(), instead of WinMain())
+		bIsBuildingConsoleApplication = true;
+	}
+}
diff --git a/Engine/Source/Programs/GeometryProcessing/Private/Algorithms/NavACD.cpp b/Engine/Source/Programs/GeometryProcessing/Private/Algorithms/NavACD.cpp
new file mode 100644
--- /dev/null
+++ b/Engine/Source/Programs/GeometryProcessing/Private/Algorithms/NavACD.cpp
@@ -0,0 +1,169 @@
+// Copyright Epic Games, Inc. All Rights Reserved.
+
+#include "CompGeom/ConvexDecomposition3.h"
+#include "MeshQueries.h"
+#include "Util/ColorConstants.h"
+#include "Utils/CommandUtils.h"
+
+#include <fstream>
+#include <iomanip>
+
+using namespace UE::Geometry;
+
+// The core navigation-driven approximate convex decomposition algorithm implementation
+FConvexDecomposition3 RunNavACD(const FDynamicMesh3& Mesh, double MinRadiusFrac, double ToleranceFrac, bool bIgnoreUnreachableInternalSpace, 
+	TArrayView<const FVector3d> CustomNavigablePositions = TArrayView<const FVector3d>());
+
+// Algorithm runner (manages command line parameters, reading input and writing output)
+DefineAlgorithm(NavACD)
+{
+	using namespace UE::CommandUtils;
+
+	if (HasTag(TEXT("help")))
+	{
+		UE_LOG(LogGeometryProcessing, Display, TEXT("Navigation-Driven Approximate Convex Decomposition arguments:"));
+		UE_LOG(LogGeometryProcessing, Display, TEXT("-input: Path to input mesh"));
+		UE_LOG(LogGeometryProcessing, Display, TEXT("-output: Path to output mesh; convex hulls will be assigned separate groups and colors"));
+		UE_LOG(LogGeometryProcessing, Display, TEXT("-stats: If this tag is present, will output stats on a successful run (timings, hull counts)"));
+		UE_LOG(LogGeometryProcessing, Display, TEXT("-protect_unreachable: If this tag is present, will protect unreachable space where the min radius sphere could fit."));
+		UE_LOG(LogGeometryProcessing, Display, TEXT("-r: Navigable space min radius parameter, as a fraction of the longest bounding box axis"));
+		UE_LOG(LogGeometryProcessing, Display, TEXT("-t: Navigable space tolerance parameter, as a fraction of the longest bounding box axis"));
+		return true;
+	}
+
+	double MinRadiusFrac = RequireParam<double>("-r");
+	double ToleranceFrac = RequireParam<double>("-t");
+
+	bool bIgnoreUnreachableInternalSpace = !HasTag("protect_unreachable");
+
+	FDynamicMesh3 Mesh = RequireInputMesh();
+	FString OutputPath = RequireParam<FString>("-output");
+
+
+	// Algorithm start time (not including file loading)
+	double StartTime = FPlatformTime::Seconds();
+
+	FConvexDecomposition3 ConvexDecomposition = RunNavACD(Mesh, MinRadiusFrac, ToleranceFrac, bIgnoreUnreachableInternalSpace);
+
+	// Algorithm end time (not including file writing)
+	double EndTime = FPlatformTime::Seconds();
+
+	if (HasTag("stats"))
+	{
+		UE_LOG(LogGeometryProcessing, Display, TEXT("Algorithm time (excluding file read/write): %f seconds"), EndTime - StartTime);
+		UE_LOG(LogGeometryProcessing, Display, TEXT("Number of hulls used: %d"), ConvexDecomposition.NumHulls());
+	}
+
+	// Write the decomposition to a single OBJ file, with a separate object name and vertex coloring per convex hull
+	auto WriteDecomp = [](FConvexDecomposition3& Decomp, const FString& Path)
+	{
+		std::ofstream FileStream(TCHAR_TO_ANSI(*Path));
+		if (!FileStream)
+		{
+			UE_LOG(LogGeometryProcessing, Error, TEXT("Failed to open output path: %s"), *Path);
+			Fail();
+		}
+		FileStream.precision(std::numeric_limits<double>::digits10);
+		int32 LastV = 0;
+		int32 ColorIdx = 0;
+		auto DumpPart = [&FileStream, &LastV, &ColorIdx](const TArray<FIndex3i>& Tris, const TArray<FVector3d>& Verts)
+		{
+			const FVector Color = LinearColors::SelectColor<FVector>(ColorIdx++);
+			FileStream << "o part" << ColorIdx << "\n";
+			for (const FVector3d& V : Verts)
+			{
+				FileStream << "v " << V.X << " " << V.Y << " " << V.Z << " " << Color.X << " " << Color.Y << " " << Color.Z << "\n";
+			}
+			for (const FIndex3i& T : Tris)
+			{
+				FileStream << "f " << T.A + LastV + 1 << " " << T.C + LastV + 1 << " " << T.B + LastV + 1 << "\n";
+			}
+			LastV += Verts.Num();
+			FileStream << "\n\n\n";
+		};
+
+		double SumVolume = 0;
+		for (int32 CvxIdx = 0; CvxIdx < Decomp.NumHulls(); CvxIdx++)
+		{
+			TArray<FIndex3i> Tris = Decomp.GetTriangles(CvxIdx);
+			TArray<FVector3d> Verts = Decomp.GetVertices<double>(CvxIdx);
+			DumpPart(Tris, Verts);
+		}
+		FileStream.close();
+	};
+
+	WriteDecomp(ConvexDecomposition, OutputPath);
+
+	return true;
+}
+
+FConvexDecomposition3 RunNavACD(const FDynamicMesh3& Mesh, double MinRadiusFrac, double ToleranceFrac, bool bIgnoreUnreachableInternalSpace, TArrayView<const FVector3d> CustomNavigablePositions)
+{
+	FAxisAlignedBox3d Bounds = Mesh.GetBounds();
+	double MaxDim = Bounds.MaxDim();
+	double UseMinRadius = MinRadiusFrac * MaxDim;
+	double UseTolerance = ToleranceFrac * MaxDim;
+
+	FConvexDecomposition3::FPreprocessMeshOptions PreprocessOptions;
+	PreprocessOptions.bMergeEdges = true;
+	PreprocessOptions.CustomPreprocess = [](FDynamicMesh3& ProcessMesh, const FAxisAlignedBox3d& Bounds) -> void
+	{
+		// for solid inputs, flip orientation if the initial volume is negative
+		if (ProcessMesh.IsClosed())
+		{
+			double InitialVolume = TMeshQueries<FDynamicMesh3>::GetVolumeArea(ProcessMesh).X;
+			if (InitialVolume < 0)
+			{
+				ProcessMesh.ReverseOrientation();
+			}
+		}
+		// Note: If we add options to simplify the input mesh, should be applied here.
+	};
+
+	FConvexDecomposition3 ConvexDecomposition(Mesh, PreprocessOptions);
+	const bool bIsSolid = ConvexDecomposition.IsInputSolid();
+	ConvexDecomposition.bTreatAsSolid = bIsSolid;
+
+	FNegativeSpaceSampleSettings NegativeSpaceSettings;
+	NegativeSpaceSettings.ApplyDefaults();
+	NegativeSpaceSettings.SampleMethod = FNegativeSpaceSampleSettings::ESampleMethod::NavigableVoxelSearch;
+	NegativeSpaceSettings.bDeterministic = true;
+	NegativeSpaceSettings.bRequireSearchSampleCoverage = true;
+	NegativeSpaceSettings.bOnlyConnectedToHull = bIgnoreUnreachableInternalSpace;
+	NegativeSpaceSettings.TargetNumSamples = 0;
+	NegativeSpaceSettings.bAllowSamplesInsideMesh = !bIsSolid;
+
+	NegativeSpaceSettings.ReduceRadiusMargin = UseTolerance;
+	NegativeSpaceSettings.MinRadius = UseMinRadius;
+	NegativeSpaceSettings.MinSpacing = 0;
+
+	ConvexDecomposition.InitializeNegativeSpace(NegativeSpaceSettings, CustomNavigablePositions);
+
+	ConvexDecomposition.MaxConvexEdgePlanes = 4;
+	ConvexDecomposition.bSplitDisconnectedComponents = false;
+	ConvexDecomposition.ConvexEdgeAngleMoreSamplesThreshold = 180; // a high value to disable this feature
+	ConvexDecomposition.ThickenAfterHullFailure = FMath::Max(FMathd::ZeroTolerance, NegativeSpaceSettings.ReduceRadiusMargin * .01);
+	constexpr int32 MaxAllowedSplits = 1000000; // more parts than any expected / reasonable decomposition
+	for (int32 Split = 0; ; Split++)
+	{
+		int32 NumSplit = ConvexDecomposition.SplitWorst(false, -1, true, NegativeSpaceSettings.ReduceRadiusMargin * .5);
+
+		if (NumSplit == 0)
+		{
+			break;
+		}
+
+		if (!ensureMsgf(Split < MaxAllowedSplits, TEXT("Convex decomposition split the input %d times; likely stuck in a loop"), Split))
+		{
+			break;
+		}
+	}
+
+	ConvexDecomposition.FixHullOverlapsInNegativeSpace();
+
+	int32 NumHullsBefore = ConvexDecomposition.NumHulls();
+	constexpr double MinThicknessToleranceWorldSpace = 0;
+	int32 NumMerged = ConvexDecomposition.MergeBest(-1, 0, MinThicknessToleranceWorldSpace, true);
+
+	return ConvexDecomposition;
+}
\ No newline at end of file
diff --git a/Engine/Source/Programs/GeometryProcessing/Private/GeometryProcessing.cpp b/Engine/Source/Programs/GeometryProcessing/Private/GeometryProcessing.cpp
new file mode 100644
--- /dev/null
+++ b/Engine/Source/Programs/GeometryProcessing/Private/GeometryProcessing.cpp
@@ -0,0 +1,39 @@
+// Copyright Epic Games, Inc. All Rights Reserved.
+
+#include "GeometryProcessing.h"
+
+#include "Utils/CommandUtils.h"
+
+#include "CompGeom/ExactPredicates.h"
+
+#include "RequiredProgramMainCPPInclude.h"
+
+DEFINE_LOG_CATEGORY(LogGeometryProcessing);
+
+IMPLEMENT_APPLICATION(GeometryProcessing, "GeometryProcessing");
+
+
+INT32_MAIN_INT32_ARGC_TCHAR_ARGV()
+{
+	ON_SCOPE_EXIT
+	{ 
+		LLM(FLowLevelMemTracker::Get().UpdateStatsPerFrame());
+		RequestEngineExit(TEXT("Exiting"));
+		FEngineLoop::AppPreExit();
+		FModuleManager::Get().UnloadModulesAtShutdown();
+		FEngineLoop::AppExit();
+	};
+
+	if (int32 Ret = GEngineLoop.PreInit(ArgC, ArgV))
+	{
+		return Ret;
+	}
+
+	// Initialize exact predicates (engine does this automatically on module init, but command line doesn't run the module init code)
+	UE::Geometry::ExactPredicates::GlobalInit();
+
+	// Run an algorithm as specified via the command line
+	bool bSuccess = UE::CommandUtils::FAlgList::Run();
+
+	return (int32)!bSuccess;
+}
diff --git a/Engine/Source/Programs/GeometryProcessing/Private/GeometryProcessing.h b/Engine/Source/Programs/GeometryProcessing/Private/GeometryProcessing.h
new file mode 100644
--- /dev/null
+++ b/Engine/Source/Programs/GeometryProcessing/Private/GeometryProcessing.h
@@ -0,0 +1,7 @@
+// Copyright Epic Games, Inc. All Rights Reserved.
+
+#pragma once
+
+#include "Logging/LogMacros.h"
+
+DECLARE_LOG_CATEGORY_EXTERN(LogGeometryProcessing, Log, All);
\ No newline at end of file
diff --git a/Engine/Source/Programs/GeometryProcessing/Private/PlatformWorkarounds.cpp b/Engine/Source/Programs/GeometryProcessing/Private/PlatformWorkarounds.cpp
new file mode 100644
--- /dev/null
+++ b/Engine/Source/Programs/GeometryProcessing/Private/PlatformWorkarounds.cpp
@@ -0,0 +1,7 @@
+// Copyright Epic Games, Inc. All Rights Reserved.
+
+#include "CoreMinimal.h"
+
+// Hack for missing Launch module dependency that causes link error on a certain platform
+FString GFileRootDirectory;
+FString GSandboxName;
\ No newline at end of file
diff --git a/Engine/Source/Programs/GeometryProcessing/Private/Utils/CommandUtils.cpp b/Engine/Source/Programs/GeometryProcessing/Private/Utils/CommandUtils.cpp
new file mode 100644
--- /dev/null
+++ b/Engine/Source/Programs/GeometryProcessing/Private/Utils/CommandUtils.cpp
@@ -0,0 +1,82 @@
+// Copyright Epic Games, Inc. All Rights Reserved.
+
+#include "Utils/CommandUtils.h"
+
+namespace UE::CommandUtils
+{
+
+	UE::Geometry::FDynamicMesh3 RequireInputMesh(const TCHAR* InputArg, const UE::MeshFileUtils::FLoadOBJSettings& Settings, bool bMustHaveFaces)
+	{
+		using namespace UE::MeshFileUtils;
+		using namespace UE::Geometry;
+
+		FString Path = RequireParam<FString>(InputArg);
+		UE::Geometry::FDynamicMesh3 InputMesh;
+		ELoadOBJStatus Status = LoadOBJ(TCHAR_TO_ANSI(*Path), InputMesh, Settings);
+		if (Status != ELoadOBJStatus::Success)
+		{
+			UE_LOG(LogGeometryProcessing, Error, TEXT("Failed to load OBJ from path: %s"), *Path);
+			Fail();
+		}
+		if (bMustHaveFaces && InputMesh.TriangleCount() == 0)
+		{
+			UE_LOG(LogGeometryProcessing, Error, TEXT("OBJ at path %s had no triangles; cannot process."), *Path);
+			Fail();
+		}
+		return InputMesh;
+	}
+
+	bool OutputResult(const UE::Geometry::FDynamicMesh3& Mesh, const TCHAR* OutputPathArg, const UE::MeshFileUtils::FWriteOBJSettings& Settings)
+	{
+		FString Path = RequireParam<FString>(OutputPathArg);
+		bool bSuccess = UE::MeshFileUtils::WriteOBJ(TCHAR_TO_ANSI(*Path), Mesh, Settings);
+		if (!bSuccess)
+		{
+			UE_LOG(LogGeometryProcessing, Error, TEXT("Failed to write mesh output parameter %s to path %s"), OutputPathArg, *Path);
+		}
+		return bSuccess;
+	}
+
+	static TMap<FString, TFunction<bool()>>& GetAlgorithms()
+	{
+		static TMap<FString, TFunction<bool()>> Algorithms;
+		return Algorithms;
+	}
+
+	bool FAlgList::Register(FString Name, TFunction<bool()> Alg)
+	{
+		GetAlgorithms().Add(Name, Alg);
+		return true;
+	}
+
+	bool FAlgList::Run()
+	{
+		TMap<FString, TFunction<bool()>>& Algorithms = GetAlgorithms();
+		FString AlgName;
+		if (FParse::Value(FCommandLine::Get(), TEXT("-alg"), AlgName))
+		{
+			TFunction<bool(void)>* AlgFunc = Algorithms.Find(AlgName);
+			if (!AlgFunc)
+			{
+				UE_LOG(LogGeometryProcessing, Error, TEXT("Unknown algorithm (%s); -alg must be one of the following:"), *AlgName);
+				for (const auto& It : Algorithms)
+				{
+					UE_LOG(LogGeometryProcessing, Error, TEXT("  -alg %s"), *It.Key);
+				}
+			}
+			else
+			{
+				return (*AlgFunc)();
+			}
+		}
+		else
+		{
+			UE_LOG(LogGeometryProcessing, Error, TEXT("Must specify one of the following arguments to choose which algorithm to run:"), *AlgName);
+			for (const auto& It : Algorithms)
+			{
+				UE_LOG(LogGeometryProcessing, Error, TEXT("  -alg %s"), *It.Key);
+			}
+		}
+		return false;
+	}
+}
\ No newline at end of file
diff --git a/Engine/Source/Programs/GeometryProcessing/Private/Utils/CommandUtils.h b/Engine/Source/Programs/GeometryProcessing/Private/Utils/CommandUtils.h
new file mode 100644
--- /dev/null
+++ b/Engine/Source/Programs/GeometryProcessing/Private/Utils/CommandUtils.h
@@ -0,0 +1,92 @@
+// Copyright Epic Games, Inc. All Rights Reserved.
+
+#pragma once
+
+#include "DynamicMesh/DynamicMesh3.h"
+#include "GeometryProcessing.h"
+#include "OBJMeshUtil.h"
+
+#include "Misc/CommandLine.h"
+#include "Misc/Parse.h"
+
+namespace UE::CommandUtils
+{
+	// For errors that we can't continue through -- immediately exits the program 
+	inline void Fail()
+	{
+		check(false);
+		FPlatformMisc::RequestExit(true);
+	}
+
+	// Read a required parameter from the command line, or exit if we cannot
+	// Note that parameters are not required to have a preceding dash, so must be specified if desired (e.g., RequireParam("-input") if you expect "-input inputvalue")
+	template<typename T>
+	T RequireParam(const TCHAR* Arg)
+	{
+		T Value;
+		if (FParse::Value(FCommandLine::Get(), Arg, Value))
+		{
+			return Value;
+		}
+		else
+		{
+			UE_LOG(LogGeometryProcessing, Error, TEXT("Must specify a valid %s parameter! Use -help to see algorithm parameters."), Arg);
+			Fail();
+			return T();
+		}
+	}
+	template<typename T>
+	T RequireParam(const char* Arg)
+	{
+		return RequireParam<T>(ANSI_TO_TCHAR(Arg));
+	}
+
+	// Request an optional parameter from the command line
+	template<typename T>
+	bool RequestParam(const TCHAR* Arg, T& Value)
+	{
+		return FParse::Value(FCommandLine::Get(), Arg, Value);
+	}
+	template<typename T>
+	bool RequestParam(const char* Arg, T& Value)
+	{
+		return RequestParam<T>(ANSI_TO_TCHAR(Arg), Value);
+	}
+
+	// Check whether a tag is on the command line
+	// Note the preceding "-" is implicit for this method, so HasTag("tag") will check for "-tag"
+	inline bool HasTag(const TCHAR* Tag)
+	{
+		return FParse::Param(FCommandLine::Get(), Tag);
+	}
+	inline bool HasTag(const char* Tag)
+	{
+		return HasTag(ANSI_TO_TCHAR(Tag));
+	}
+
+	// Reads an OBJ mesh from path specified by the input parameter
+	UE::Geometry::FDynamicMesh3 RequireInputMesh(const TCHAR* InputArg = TEXT("-input"), const UE::MeshFileUtils::FLoadOBJSettings& Settings = UE::MeshFileUtils::FLoadOBJSettings(), bool bMustHaveFaces = true);
+	inline UE::Geometry::FDynamicMesh3 RequireInputMesh(const char* InputArg, const UE::MeshFileUtils::FLoadOBJSettings& Settings = UE::MeshFileUtils::FLoadOBJSettings(), bool bMustHaveFaces = true)
+	{
+		return RequireInputMesh(ANSI_TO_TCHAR(InputArg), Settings, bMustHaveFaces);
+	}
+
+	// Writes an OBJ mesh to the output parameter
+	bool OutputResult(const UE::Geometry::FDynamicMesh3& Mesh, const TCHAR* OutputPathArg = TEXT("-output"), const UE::MeshFileUtils::FWriteOBJSettings& Settings = UE::MeshFileUtils::FWriteOBJSettings());
+	inline bool OutputResult(const UE::Geometry::FDynamicMesh3& Mesh, const char* OutputPathArg, const UE::MeshFileUtils::FWriteOBJSettings& Settings = UE::MeshFileUtils::FWriteOBJSettings())
+	{
+		return OutputResult(Mesh, ANSI_TO_TCHAR(OutputPathArg), Settings);
+	}
+
+	// Struct to manage the list of known algorithms selectable with the -alg parameter
+	struct FAlgList
+	{
+		static bool Register(FString Name, TFunction<bool()> Alg);
+		static bool Run();
+	};
+
+}
+
+// Macro to create an algorithm that will automatically register as an -alg parameter
+#define DefineAlgorithm(Name) bool Alg_##Name(); static bool bAlgRegistered_##Name = UE::CommandUtils::FAlgList::Register(TEXT(#Name), &Alg_##Name); bool Alg_##Name()
+
