Modeling Tools: fixes crash in conversion to PolyVerts

diff --git a/Engine/Plugins/Runtime/GeometryProcessing/Source/DynamicMesh/Private/Selections/GeometrySelectionUtil.cpp b/Engine/Plugins/Runtime/GeometryProcessing/Source/DynamicMesh/Private/Selections/GeometrySelectionUtil.cpp
--- a/Engine/Plugins/Runtime/GeometryProcessing/Source/DynamicMesh/Private/Selections/GeometrySelectionUtil.cpp
+++ b/Engine/Plugins/Runtime/GeometryProcessing/Source/DynamicMesh/Private/Selections/GeometrySelectionUtil.cpp
@@ -1429,9 +1429,12 @@ bool UE::Geometry::InitializeSelectionFromPolyGroups(
 								for (int32 CornerIdx = 0; CornerIdx < 2; ++CornerIdx)
 								{
 									int32 CornerID = GroupTopology.Edges[GroupEdgeID].EndpointCorners[CornerIdx];
-									const FGroupTopology::FCorner& Corner = GroupTopology.Corners[CornerID];
-									FGeoSelectionID ID = FGeoSelectionID(Corner.VertexID, CornerID);
-									SelectionOut.Selection.Add(ID.Encoded());
+									if (CornerID != IndexConstants::InvalidID)
+									{
+										const FGroupTopology::FCorner& Corner = GroupTopology.Corners[CornerID];
+										FGeoSelectionID ID = FGeoSelectionID(Corner.VertexID, CornerID);
+										SelectionOut.Selection.Add(ID.Encoded());
+									}
 								}
 							}
 						}
@@ -1708,9 +1711,12 @@ bool UE::Geometry::ConvertSelection(
 				for (int32 SubIdx = 0; SubIdx < 2; ++SubIdx)
 				{
 					int32 CornerID = Corners[SubIdx];
-					const FGroupTopology::FCorner& Corner = GroupTopology->Corners[CornerID];
-					FGeoSelectionID ID = FGeoSelectionID(Corner.VertexID, CornerID);
-					ToSelectionOut.Selection.Add(ID.Encoded());
+					if (CornerID != IndexConstants::InvalidID)
+					{
+						const FGroupTopology::FCorner& Corner = GroupTopology->Corners[CornerID];
+						FGeoSelectionID ID = FGeoSelectionID(Corner.VertexID, CornerID);
+						ToSelectionOut.Selection.Add(ID.Encoded());
+					}
 				}
 			}
 			return true;
