Modeling Tools: Fixed bug after TriFace->TriEdge conversion not recognizing selection during hover

diff --git a/Engine/Plugins/Runtime/GeometryProcessing/Source/DynamicMesh/Private/Selections/GeometrySelectionUtil.cpp b/Engine/Plugins/Runtime/GeometryProcessing/Source/DynamicMesh/Private/Selections/GeometrySelectionUtil.cpp
--- a/Engine/Plugins/Runtime/GeometryProcessing/Source/DynamicMesh/Private/Selections/GeometrySelectionUtil.cpp
+++ b/Engine/Plugins/Runtime/GeometryProcessing/Source/DynamicMesh/Private/Selections/GeometrySelectionUtil.cpp
@@ -1379,10 +1379,11 @@ bool UE::Geometry::InitializeSelectionFromPolyGroups(
 				}
 				else if (SelectionOut.ElementType == EGeometryElementType::Edge)
 				{
-					FIndex3i TriEdges = Mesh.GetTriEdges(TID);
-					SelectionOut.Selection.Add(FGeoSelectionID::MeshEdge(Mesh.GetTriEdgeIDFromEdgeID(TriEdges.A)).Encoded());
-					SelectionOut.Selection.Add(FGeoSelectionID::MeshEdge(Mesh.GetTriEdgeIDFromEdgeID(TriEdges.B)).Encoded());
-					SelectionOut.Selection.Add(FGeoSelectionID::MeshEdge(Mesh.GetTriEdgeIDFromEdgeID(TriEdges.C)).Encoded());
+					Mesh.EnumerateTriEdgeIDsFromTriID(TID,
+						[&SelectionOut](const FMeshTriEdgeID TriEdgeID)
+						{
+							SelectionOut.Selection.Add(FGeoSelectionID::MeshEdge(TriEdgeID).Encoded());
+						});
 				}
 				else if (SelectionOut.ElementType == EGeometryElementType::Vertex)
 				{
@@ -1472,14 +1473,15 @@ bool UE::Geometry::InitializeSelectionFromTriangles(
 		}
 		else if (SelectionOut.ElementType == EGeometryElementType::Edge)
 		{
-			for (int32 tid : Triangles)
+			for (const int32 TID : Triangles)
 			{
-				if (Mesh.IsTriangle(tid))
+				if (Mesh.IsTriangle(TID))
 				{
-					FIndex3i TriEdges = Mesh.GetTriEdges(tid);
-					SelectionOut.Selection.Add(FGeoSelectionID::MeshEdge(Mesh.GetTriEdgeIDFromEdgeID(TriEdges.A)).Encoded());
-					SelectionOut.Selection.Add(FGeoSelectionID::MeshEdge(Mesh.GetTriEdgeIDFromEdgeID(TriEdges.B)).Encoded());
-					SelectionOut.Selection.Add(FGeoSelectionID::MeshEdge(Mesh.GetTriEdgeIDFromEdgeID(TriEdges.C)).Encoded());
+					Mesh.EnumerateTriEdgeIDsFromTriID(TID,
+						[&SelectionOut](const FMeshTriEdgeID TriEdgeID)
+						{
+							SelectionOut.Selection.Add(FGeoSelectionID::MeshEdge(TriEdgeID).Encoded());
+						});
 				}
 			}
 		}
diff --git a/Engine/Source/Runtime/GeometryCore/Public/DynamicMesh/DynamicMesh3.h b/Engine/Source/Runtime/GeometryCore/Public/DynamicMesh/DynamicMesh3.h
--- a/Engine/Source/Runtime/GeometryCore/Public/DynamicMesh/DynamicMesh3.h
+++ b/Engine/Source/Runtime/GeometryCore/Public/DynamicMesh/DynamicMesh3.h
@@ -704,6 +704,16 @@ public:
 		return TriangleEdges[TriangleID][j];
 	}
 
+	/**  Applies a given function to both TriEdgeIDs which each EdgeID in a given Triangle is associated with */
+	void EnumerateTriEdgeIDsFromTriID(const int TriID, const TFunctionRef<void(FMeshTriEdgeID TriEdgeID)>& TriEdgeFunc) const
+	{
+		FIndex3i TriEdges = GetTriEdges(TriID);
+		for (int TriEdgesIndex = 0 ; TriEdgesIndex <=2; TriEdgesIndex++)
+		{
+			EnumerateTriEdgeIDsFromEdgeID(TriEdges[TriEdgesIndex], TriEdgeFunc);
+		}
+	}
+
 	/** Find the neighbour triangles of a triangle (any of them might be InvalidID) */
 	GEOMETRYCORE_API FIndex3i GetTriNeighbourTris(int TriangleID) const;
 
@@ -782,6 +792,30 @@ public:
 		}
 	}
 
+	/** Applies a given function to both TriEdgeIDs which a given EdgeID is associated with*/
+	void EnumerateTriEdgeIDsFromEdgeID(const int32 EdgeID, const TFunctionRef<void(FMeshTriEdgeID TriEdgeID)>& TriEdgeFunc) const
+	{
+		const FMeshTriEdgeID FirstTriEdgeID = GetTriEdgeIDFromEdgeID(EdgeID); // function gets MeshTriEdgeID for edge included in EdgeTri.A only
+		TriEdgeFunc(FirstTriEdgeID);
+
+		// have to get MeshTriEdgeID for edge included in EdgeTri.B
+		const int OtherTriID = GetEdgeT(EdgeID).B;
+		if (OtherTriID != IndexConstants::InvalidID)
+		{
+			FMeshTriEdgeID SecondTriEdgeID;
+			const FIndex3i SecondTriEdges = GetTriEdges(OtherTriID);
+			if (SecondTriEdges.A == EdgeID)
+			{
+				SecondTriEdgeID = FMeshTriEdgeID(OtherTriID, 0);
+			}
+			else
+			{
+				SecondTriEdgeID = FMeshTriEdgeID(OtherTriID, ( SecondTriEdges.B == EdgeID ) ? 1 : 2 );
+			}
+			TriEdgeFunc(SecondTriEdgeID);
+		}
+	}
+
 	//
 	// Vertex and Triangle attribute arrays
 	//
