Add an open-boundary snapping algorithm, similar to 'resolve t-junctions' but just snapping without adding vertices

diff --git a/Engine/Plugins/Runtime/GeometryProcessing/Source/DynamicMesh/Private/Operations/MeshResolveTJunctions.cpp b/Engine/Plugins/Runtime/GeometryProcessing/Source/DynamicMesh/Private/Operations/MeshResolveTJunctions.cpp
--- a/Engine/Plugins/Runtime/GeometryProcessing/Source/DynamicMesh/Private/Operations/MeshResolveTJunctions.cpp
+++ b/Engine/Plugins/Runtime/GeometryProcessing/Source/DynamicMesh/Private/Operations/MeshResolveTJunctions.cpp
@@ -6,6 +6,7 @@
 #include "DynamicMesh/DynamicMeshAttributeSet.h"
 #include "Spatial/PointSetHashTable.h"
 #include "SegmentTypes.h"
+#include "TriangleTypes.h"
 
 using namespace UE::Geometry;
 
@@ -131,5 +132,174 @@ bool FMeshResolveTJunctions::Apply()
 		}
 	}
 
+	return true;
+}
+
+
+
+
+
+const double FMeshSnapOpenBoundaries::DEFAULT_TOLERANCE = FMathf::ZeroTolerance;
+
+bool FMeshSnapOpenBoundaries::Apply()
+{
+	NumVertexSnaps = 0;
+
+	// make clean boundary edge set, either for the whole mesh, or based on input set
+	if (BoundaryEdges.Num() == 0)
+	{
+		for (int32 EID : Mesh->BoundaryEdgeIndicesItr())
+		{
+			BoundaryEdges.Add(EID);
+		}
+	}
+	else
+	{
+		// filter out any invalid / non-boundary edges from the input set
+		TArray<int32> ToRemove;
+		for (int32 EID : BoundaryEdges)
+		{
+			if (!Mesh->IsEdge(EID) || !Mesh->IsBoundaryEdge(EID))
+			{
+				ToRemove.Add(EID);
+			}
+		}
+		for (int32 EID : ToRemove)
+		{
+			BoundaryEdges.Remove(EID);
+		}
+	}
+
+
+	TSet<int32> BoundaryVertices;
+	for (int32 EID : BoundaryEdges)
+	{
+		FIndex2i EdgeV = Mesh->GetEdgeV(EID);
+		BoundaryVertices.Add(EdgeV.A);
+		BoundaryVertices.Add(EdgeV.B);
+	}
+
+	// TODO: this should probably use a hash table like MeshMergeCoincidentEdges.
+	// However in this case adding the midpoints won't work, the entire edge boxes
+	// needed to be added. And when edges are moved via snapping, the hash will need to
+	// be updated. So for now this is just O(N*M).
+
+	// TODO: Write a 'fast path' for the bSnapToEdge==false case?
+	
+	// TODO: Optimize subsequent iterations w/ correspondences from previous iteration 
+	// and perhaps optionally report the correspondences back to caller?
+
+	double DistanceToleranceSq = DistanceTolerance * DistanceTolerance;
+	double VertexDistanceToleranceSq = DistanceToleranceSq * VertexSnapToleranceFactor * VertexSnapToleranceFactor;
+
+	auto TestIfFlips = [this](int32 TID, int32 VID, FVector3d NewPos, double FlipThreshold = 0) -> bool
+	{
+		FIndex3i TriVIDs = Mesh->GetTriangle(TID);
+		int32 SubIdx = TriVIDs.IndexOf(VID);
+		FTriangle3d Tri(Mesh->GetVertex(TriVIDs.A), Mesh->GetVertex(TriVIDs.B), Mesh->GetVertex(TriVIDs.C));
+		FVector3d InitialNormal = Tri.Normal();
+		Tri.V[SubIdx] = NewPos;
+		FVector3d NewNormal = Tri.Normal();
+		return InitialNormal.Dot(NewNormal) < FlipThreshold;
+	};
+
+	bool bSnappedLastIteration = true;
+	int32 LastNumSnapped = NumVertexSnaps;
+	for (int32 Iters = 0; Iters < MaxIterations && bSnappedLastIteration; ++Iters)
+	{
+		for (int32 VID : BoundaryVertices)
+		{
+			FVector3d Position = Mesh->GetVertex(VID);
+			int E0, E1;
+			Mesh->GetVtxBoundaryEdges(VID, E0, E1);
+
+			int32 OnEdgeID = -1;
+			double MinDistSqr = TNumericLimits<double>::Max();
+			FSegment3d OnEdgeSegment;
+			FVector3d SegmentA, SegmentB;
+
+			// find the edge in the boundary edge set that is closest to this vertex, 
+			// and not connected to this vertex
+			for (int32 EID : BoundaryEdges)
+			{
+				if (EID == E0 || EID == E1) continue;
+
+				FVector3d A, B;
+				Mesh->GetEdgeV(EID, A, B);
+				FSegment3d EdgeSegment(A, B);
+				double DistSqr = EdgeSegment.DistanceSquared(Position);
+				if (DistSqr < MinDistSqr)
+				{
+					MinDistSqr = DistSqr;
+					OnEdgeID = EID;
+					OnEdgeSegment = EdgeSegment;
+					SegmentA = A;
+					SegmentB = B;
+				}
+			}
+
+			// if we did not find an edge, or we are too far from any edge, do not snap
+			if (OnEdgeID == -1 || MinDistSqr > DistanceToleranceSq)
+			{
+				continue;
+			}
+
+			// if we are within vertex tolerance of either edge endpoint, snap to the closest edge endpoint
+			double ToStartSq = DistanceSquared(OnEdgeSegment.StartPoint(), Position);
+			double ToEndSq = DistanceSquared(OnEdgeSegment.EndPoint(), Position);
+			FVector3d SnapToPt = Position;
+			if (ToStartSq < VertexDistanceToleranceSq || ToEndSq < VertexDistanceToleranceSq)
+			{
+				if (ToStartSq < ToEndSq)
+				{
+					SnapToPt = SegmentA;
+				}
+				else
+				{
+					SnapToPt = SegmentB;
+				}
+			}
+			else if (bSnapToEdges) // otherwise snap to the edge
+			{
+				// check that the position is within the span of the edge
+				double SegmentT = OnEdgeSegment.Project(Position);
+				if (SegmentT <= OnEdgeSegment.Extent && SegmentT >= -OnEdgeSegment.Extent)
+				{
+					SnapToPt = OnEdgeSegment.NearestPoint(Position);
+				}
+				else
+				{
+					// outside of span, would be a vertex snap so should use the (possibly lower) vertex snap tolerance
+					continue;
+				}
+			}
+
+			// if vertex was moved farther than a very small tolerance, test for triangle flips and then count it as a snap
+			constexpr double CountsAsSnapDistance = FMathd::ZeroTolerance;
+			constexpr double CountsAsSnapDistanceSq = CountsAsSnapDistance * CountsAsSnapDistance;
+			if (DistanceSquared(Position, SnapToPt) > CountsAsSnapDistanceSq)
+			{
+				if (bPreventFlips)
+				{
+					bool bFlippedTri = false;
+					Mesh->EnumerateVertexTriangles(VID, [&bFlippedTri, &TestIfFlips, VID, SnapToPt](int32 TID)
+					{
+						bFlippedTri = bFlippedTri || TestIfFlips(TID, VID, SnapToPt);
+					});
+					if (bFlippedTri)
+					{
+						continue;
+					}
+				}
+				NumVertexSnaps++;
+			}
+
+			// snap the vertex
+			Mesh->SetVertex(VID, SnapToPt);
+		}
+		bSnappedLastIteration = NumVertexSnaps > LastNumSnapped;
+		LastNumSnapped = NumVertexSnaps;
+	}
+
 	return true;
 }
\ No newline at end of file
diff --git a/Engine/Plugins/Runtime/GeometryProcessing/Source/DynamicMesh/Public/Operations/MeshResolveTJunctions.h b/Engine/Plugins/Runtime/GeometryProcessing/Source/DynamicMesh/Public/Operations/MeshResolveTJunctions.h
--- a/Engine/Plugins/Runtime/GeometryProcessing/Source/DynamicMesh/Public/Operations/MeshResolveTJunctions.h
+++ b/Engine/Plugins/Runtime/GeometryProcessing/Source/DynamicMesh/Public/Operations/MeshResolveTJunctions.h
@@ -71,6 +71,53 @@ public:
 
 };
 
+/**
+ * Similar to FMeshResolveTJunctions, but does not add any vertices to the mesh.
+ * Supports running multiple snapping iterations, because snapped-to edges may move in subsequent snaps.
+ */
+class FMeshSnapOpenBoundaries
+{
+public:
+	/** default tolerance is float ZeroTolerance */
+	DYNAMICMESH_API static const double DEFAULT_TOLERANCE;  // = FMathf::ZeroTolerance;
+
+	/** The mesh that we are modifying */
+	FDynamicMesh3* Mesh;
+
+	/** Subset of mesh boundary edges (otherwise all boundary edges are processed) */
+	TSet<int32> BoundaryEdges;
+
+	/** Distance threshold used for various checks (eg is vertex on edge, end endpoint tolerance, etc) */
+	double DistanceTolerance = DEFAULT_TOLERANCE;
+
+	/** Scalar multiple of DistanceTolerance at which we snap a vertex directly to another vertex, rather than an edge */
+	double VertexSnapToleranceFactor = 1.0;
+
+	/** Number of vertex snaps performed (cumulative over iterations) */
+	int32 NumVertexSnaps = 0;
+
+	/** Maximum number of snapping iterations to perform */
+	int32 MaxIterations = 1;
+
+	/** Whether vertices can be snapped to edges; otherwise, vertices are only snapped to other vertices */
+	bool bSnapToEdges = true;
+
+	/** Whether to avoid snapping in cases where doing so would locally flip a triangle */
+	bool bPreventFlips = true;
+
+
+public:
+	FMeshSnapOpenBoundaries(FDynamicMesh3* MeshIn) : Mesh(MeshIn)
+	{
+	}
+
+	/**
+	 * Run the resolve operation and modify .Mesh
+	 * @return true if the algorithm succeeds
+	 */
+	DYNAMICMESH_API bool Apply();
+
+};
 
 } // end namespace UE::Geometry
 } // end namespace UE
\ No newline at end of file
diff --git a/Engine/Plugins/Runtime/GeometryScripting/Source/GeometryScriptingCore/Private/MeshRepairFunctions.cpp b/Engine/Plugins/Runtime/GeometryScripting/Source/GeometryScriptingCore/Private/MeshRepairFunctions.cpp
--- a/Engine/Plugins/Runtime/GeometryScripting/Source/GeometryScriptingCore/Private/MeshRepairFunctions.cpp
+++ b/Engine/Plugins/Runtime/GeometryScripting/Source/GeometryScriptingCore/Private/MeshRepairFunctions.cpp
@@ -65,6 +65,36 @@ UDynamicMesh* UGeometryScriptLibrary_MeshRepairFunctions::RemoveUnusedVertices(
 }
 
 
+
+
+UDynamicMesh* UGeometryScriptLibrary_MeshRepairFunctions::SnapMeshOpenBoundaries(
+	UDynamicMesh* TargetMesh,
+	FGeometryScriptSnapBoundariesOptions SnapOptions,
+	UGeometryScriptDebug* Debug)
+{
+	if (TargetMesh == nullptr)
+	{
+		UE::Geometry::AppendError(Debug, EGeometryScriptErrorType::InvalidInputs, LOCTEXT("SnapMeshOpenBoundaries_InvalidInput", "SnapMeshOpenBoundaries: TargetMesh is Null"));
+		return TargetMesh;
+	}
+
+	TargetMesh->EditMesh([&](FDynamicMesh3& EditMesh) 
+	{
+		FMeshSnapOpenBoundaries Snapper(&EditMesh);
+		Snapper.DistanceTolerance = 2 * SnapOptions.Tolerance;
+		Snapper.MaxIterations = SnapOptions.MaxIterations;
+		Snapper.bSnapToEdges = SnapOptions.bSnapToEdges;
+		bool bResolveOK = Snapper.Apply();
+		if (!bResolveOK)
+		{
+			UE::Geometry::AppendError(Debug, EGeometryScriptErrorType::OperationFailed, LOCTEXT("SnapMeshOpenBoundaries_Error", "SnapMeshOpenBoundaries: SnapMeshOpenBoundaries Operation returned error flag"));
+		}
+
+	}, EDynamicMeshChangeType::GeneralEdit, EDynamicMeshAttributeChangeFlags::Unknown, false);
+
+	return TargetMesh;
+}
+
 UDynamicMesh* UGeometryScriptLibrary_MeshRepairFunctions::ResolveMeshTJunctions(
 	UDynamicMesh* TargetMesh,
 	FGeometryScriptResolveTJunctionOptions ResolveOptions,
diff --git a/Engine/Plugins/Runtime/GeometryScripting/Source/GeometryScriptingCore/Public/GeometryScript/MeshRepairFunctions.h b/Engine/Plugins/Runtime/GeometryScripting/Source/GeometryScriptingCore/Public/GeometryScript/MeshRepairFunctions.h
--- a/Engine/Plugins/Runtime/GeometryScripting/Source/GeometryScriptingCore/Public/GeometryScript/MeshRepairFunctions.h
+++ b/Engine/Plugins/Runtime/GeometryScripting/Source/GeometryScriptingCore/Public/GeometryScript/MeshRepairFunctions.h
@@ -35,6 +35,23 @@ public:
 	float Tolerance = 1e-03f;
 };
 
+USTRUCT(BlueprintType)
+struct GEOMETRYSCRIPTINGCORE_API FGeometryScriptSnapBoundariesOptions
+{
+	GENERATED_BODY()
+public:
+	/** Snapping tolerance */
+	UPROPERTY(EditAnywhere, BlueprintReadWrite, Category = Options)
+	float Tolerance = 1e-03f;
+
+	/** Whether to snap vertices to open edges. If false, will only snap together vertices */
+	UPROPERTY(EditAnywhere, BlueprintReadWrite, Category = Options)
+	bool bSnapToEdges = true;
+
+	/** Maximum number of iterations of boundary snapping to apply. Will stop earlier if an iteration applies no snapping. */
+	UPROPERTY(EditAnywhere, BlueprintReadWrite, Category = Options)
+	int MaxIterations = 5;
+};
 
 UENUM(BlueprintType)
 enum class EGeometryScriptFillHolesMethod : uint8
@@ -183,6 +200,18 @@ public:
 		UDynamicMesh* TargetMesh,
 		UGeometryScriptDebug* Debug = nullptr);
 
+	/**
+	 * Snap vertices on open edges to the closest compatible open boundary, if found within the tolerance distance
+	 * Unlike ResolveMeshTJunctions, does not introduce new vertices to the mesh
+	 */
+	UFUNCTION(BlueprintCallable, Category = "GeometryScript|Repair", meta = (ScriptMethod))
+	static UPARAM(DisplayName = "Target Mesh") UDynamicMesh*
+	SnapMeshOpenBoundaries(
+		UDynamicMesh* TargetMesh,
+		FGeometryScriptSnapBoundariesOptions SnapOptions,
+		UGeometryScriptDebug* Debug = nullptr
+	);
+
 	/**
 	* Attempts to resolve T-Junctions in the mesh by addition of vertices and welding.
 	*/
