CIS fix for 'potentially uninitialized local variable' warning in open boundary snapping algorithm

diff --git a/Engine/Plugins/Runtime/GeometryProcessing/Source/DynamicMesh/Private/Operations/MeshResolveTJunctions.cpp b/Engine/Plugins/Runtime/GeometryProcessing/Source/DynamicMesh/Private/Operations/MeshResolveTJunctions.cpp
--- a/Engine/Plugins/Runtime/GeometryProcessing/Source/DynamicMesh/Private/Operations/MeshResolveTJunctions.cpp
+++ b/Engine/Plugins/Runtime/GeometryProcessing/Source/DynamicMesh/Private/Operations/MeshResolveTJunctions.cpp
@@ -216,7 +216,7 @@ bool FMeshSnapOpenBoundaries::Apply()
 			int32 OnEdgeID = -1;
 			double MinDistSqr = TNumericLimits<double>::Max();
 			FSegment3d OnEdgeSegment;
-			FVector3d SegmentA, SegmentB;
+			FVector3d SegmentA = FVector3d::ZeroVector, SegmentB = FVector3d::ZeroVector;
 
 			// find the edge in the boundary edge set that is closest to this vertex, 
 			// and not connected to this vertex
