Modeling Tools: Complete conversion matrix for mesh element selection

diff --git a/Engine/Plugins/Runtime/GeometryProcessing/Source/DynamicMesh/Private/Selections/GeometrySelectionUtil.cpp b/Engine/Plugins/Runtime/GeometryProcessing/Source/DynamicMesh/Private/Selections/GeometrySelectionUtil.cpp
--- a/Engine/Plugins/Runtime/GeometryProcessing/Source/DynamicMesh/Private/Selections/GeometrySelectionUtil.cpp
+++ b/Engine/Plugins/Runtime/GeometryProcessing/Source/DynamicMesh/Private/Selections/GeometrySelectionUtil.cpp
@@ -1050,6 +1050,103 @@ bool UE::Geometry::EnumeratePolygroupSelectionEdges(
 	return true;
 }
 
+bool UE::Geometry::EnumeratePolygroupSelectionEdges(
+	const FGeometrySelection& MeshSelection,
+	const UE::Geometry::FDynamicMesh3& Mesh,
+	const UE::Geometry::FGroupTopology& GroupTopology,
+	TFunctionRef<void(int32)> EdgeFunc
+)
+{
+	if ( ensure( MeshSelection.TopologyType == EGeometryTopologyType::Polygroup ) == false )
+	{
+		return false;
+	}
+
+
+	TArray<int32> GroupEdgeIDs;
+	
+	if (MeshSelection.ElementType == EGeometryElementType::Face)
+	{
+		TSet<int32> SeedGroups;
+		
+		for (const uint64 EncodedID : MeshSelection.Selection)
+		{
+			const FGeoSelectionID SelectionID(EncodedID);
+			const int32 SeedTriangleID = (int32)SelectionID.GeometryID;
+			if (Mesh.IsTriangle(SeedTriangleID))
+			{
+				const int32 GroupID = GroupTopology.GetGroupID(SeedTriangleID);
+				check(GroupID == (int32)SelectionID.TopologyID);		// sanity-check that we are using the right group
+				if ( GroupID >= 0)
+				{
+					SeedGroups.Add(GroupID);
+				}
+			}
+		}
+
+		for (const int32 GroupID : SeedGroups)
+		{
+			TArray<int32> GroupTriangles = GroupTopology.GetGroupTriangles(GroupID);
+			for (const int32 TriangleID : GroupTriangles)
+			{
+				FIndex3i TriEdges = Mesh.GetTriEdges(TriangleID);
+				EdgeFunc(TriEdges[0]);
+				EdgeFunc(TriEdges[1]);
+				EdgeFunc(TriEdges[2]);
+			}
+		}
+	}
+	else if (MeshSelection.ElementType == EGeometryElementType::Edge)
+	{
+		for (const uint64 EncodedID : MeshSelection.Selection)
+		{
+			const FMeshTriEdgeID TriEdgeID( FGeoSelectionID(EncodedID).GeometryID );
+			const int32 SeedEdgeID = Mesh.IsTriangle(TriEdgeID.TriangleID) ? Mesh.GetTriEdge(TriEdgeID.TriangleID, TriEdgeID.TriEdgeIndex) : IndexConstants::InvalidID;
+			if (Mesh.IsEdge(SeedEdgeID))
+			{
+				const int32 GroupEdgeID = GroupTopology.FindGroupEdgeID(SeedEdgeID);
+				if (GroupEdgeID != IndexConstants::InvalidID)
+				{
+					GroupEdgeIDs.Add (GroupEdgeID);
+				}
+			}
+		}
+	}
+	else if (MeshSelection.ElementType == EGeometryElementType::Vertex)
+	{
+		for (const uint64 VertexID : MeshSelection.Selection)
+		{
+			if (Mesh.IsVertex((int32)VertexID))
+			{
+				Mesh.EnumerateVertexEdges((int32)VertexID, [&GroupTopology, &GroupEdgeIDs](const int32 EdgeID)
+				{
+					const int32 GroupEdgeID = GroupTopology.FindGroupEdgeID(EdgeID);
+					if (GroupEdgeID != IndexConstants::InvalidID)
+					{
+						GroupEdgeIDs.Add(GroupEdgeID);
+					}
+				});
+			}
+		}
+	}
+	else
+	{
+		return false;
+	}
+
+	if (MeshSelection.ElementType == EGeometryElementType::Vertex || MeshSelection.ElementType == EGeometryElementType::Edge)
+	{
+		for (const int32 GroupEdgeID : GroupEdgeIDs)
+		{
+			const FGroupTopology::FGroupEdge& GroupEdge = GroupTopology.Edges[GroupEdgeID];
+			for (const int32 EID : GroupEdge.Span.Edges)
+			{
+				EdgeFunc(EID);
+			}
+		}
+	}
+	return true;
+}
 
 
 bool UE::Geometry::EnumerateTriangleSelectionElements(
@@ -1350,103 +1447,6 @@ bool UE::Geometry::ConvertPolygroupSelectionToTopologySelection(
 	return true;
 }
 
-bool UE::Geometry::InitializeSelectionFromPolyGroups(
-	const FDynamicMesh3& Mesh,
-	const FGroupTopology& GroupTopology,
-	TArrayView<const int> PolyGroupIDs,
-	FGeometrySelection& SelectionOut)
-{
-	bool bSuccess = true;
-	if (SelectionOut.TopologyType == EGeometryTopologyType::Triangle)
-	{
-		for (int32 GID : PolyGroupIDs)
-		{
-			const FGroupTopology::FGroup* Group = GroupTopology.FindGroupByID(GID);
-			if (!Group)
-			{
-				bSuccess = false;
-				continue;
-			}
-			for (int32 TID : Group->Triangles)
-			{
-				if (!Mesh.IsTriangle(TID))
-				{
-					continue;
-				}
-				if (SelectionOut.ElementType == EGeometryElementType::Face)
-				{
-					SelectionOut.Selection.Add(FGeoSelectionID::MeshTriangle(TID).Encoded());
-				}
-				else if (SelectionOut.ElementType == EGeometryElementType::Edge)
-				{
-					Mesh.EnumerateTriEdgeIDsFromTriID(TID,
-						[&SelectionOut](const FMeshTriEdgeID TriEdgeID)
-						{
-							SelectionOut.Selection.Add(FGeoSelectionID::MeshEdge(TriEdgeID).Encoded());
-						});
-				}
-				else if (SelectionOut.ElementType == EGeometryElementType::Vertex)
-				{
-					FIndex3i Tri = Mesh.GetTriangle(TID);
-					SelectionOut.Selection.Add(FGeoSelectionID::MeshVertex(Tri.A).Encoded());
-					SelectionOut.Selection.Add(FGeoSelectionID::MeshVertex(Tri.B).Encoded());
-					SelectionOut.Selection.Add(FGeoSelectionID::MeshVertex(Tri.C).Encoded());
-				}
-			}
-		}
-	}
-	else if (SelectionOut.TopologyType == EGeometryTopologyType::Polygroup)
-	{
-		for (int32 GID : PolyGroupIDs)
-		{
-			const FGroupTopology::FGroup* Group = GroupTopology.FindGroupByID(GID);
-			if (!Group)
-			{
-				bSuccess = false;
-				continue;
-			}
-			if (SelectionOut.ElementType == EGeometryElementType::Face)
-			{
-				FGeoSelectionID ID = FGeoSelectionID(Group->Triangles[0], Group->GroupID);
-				SelectionOut.Selection.Add(ID.Encoded());
-			}
-			else
-			{
-				for (const FGroupTopology::FGroupBoundary& Boundary : Group->Boundaries)
-				{
-					for (int GroupEdgeID : Boundary.GroupEdges)
-					{
-						if (GroupTopology.Edges.IsValidIndex(GroupEdgeID))
-						{
-							if (SelectionOut.ElementType == EGeometryElementType::Edge)
-							{
-								const FGroupTopology::FGroupEdge& GroupEdge = GroupTopology.Edges[GroupEdgeID];
-								FMeshTriEdgeID MeshEdgeID = Mesh.GetTriEdgeIDFromEdgeID(GroupEdge.Span.Edges[0]);
-								FGeoSelectionID ID = FGeoSelectionID(MeshEdgeID.Encoded(), GroupEdgeID);
-								SelectionOut.Selection.Add(ID.Encoded());
-							}
-							else if (SelectionOut.ElementType == EGeometryElementType::Vertex)
-							{
-								for (int32 CornerIdx = 0; CornerIdx < 2; ++CornerIdx)
-								{
-									int32 CornerID = GroupTopology.Edges[GroupEdgeID].EndpointCorners[CornerIdx];
-									if (CornerID != IndexConstants::InvalidID)
-									{
-										const FGroupTopology::FCorner& Corner = GroupTopology.Corners[CornerID];
-										FGeoSelectionID ID = FGeoSelectionID(Corner.VertexID, CornerID);
-										SelectionOut.Selection.Add(ID.Encoded());
-									}
-								}
-							}
-						}
-					}
-				}
-			}
-		}
-	}
-	return bSuccess;
-}
-
 
 bool UE::Geometry::InitializeSelectionFromTriangles(
 	const UE::Geometry::FDynamicMesh3& Mesh,
@@ -1573,22 +1573,11 @@ bool UE::Geometry::ConvertSelection(
 {
 	using namespace GeometrySelectionUtilLocals;
 
-	const auto NotImplemented = [](
-		const FDynamicMesh3& Mesh,
-		const FGroupTopology* GroupTopology,
-		const FGeometrySelection& FromSelectionIn,
-		FGeometrySelection& ToSelectionOut)
-	{
-		return false;
-	};
-
-
-
 	const auto FromTypeToSame = [](
 		const FDynamicMesh3& Mesh,
 		const FGroupTopology* GroupTopology,
 		const FGeometrySelection& FromSelectionIn,
-		FGeometrySelection& ToSelectionOut)
+		FGeometrySelection& ToSelectionOut) -> bool
 	{
 		checkSlow(FromSelectionIn.IsSameType(ToSelectionOut));
 
@@ -1597,207 +1586,197 @@ bool UE::Geometry::ConvertSelection(
 		return true;
 	};
 
-
-
-	const auto FromTriFace = [](
+	const auto ToTriFace = [](
 		const FDynamicMesh3& Mesh,
 		const FGroupTopology* GroupTopology,
 		const FGeometrySelection& FromSelectionIn,
-		FGeometrySelection& ToSelectionOut)
+		FGeometrySelection& ToSelectionOut) -> bool
 	{
-		checkSlow(FromSelectionIn.TopologyType == EGeometryTopologyType::Triangle);
-		checkSlow(FromSelectionIn.ElementType == EGeometryElementType::Face);
+		checkSlow(ToSelectionOut.TopologyType == EGeometryTopologyType::Triangle);
+		checkSlow(ToSelectionOut.ElementType == EGeometryElementType::Face);
 
-		TArray<int32> Triangles;
-		for (uint64 ElemID : FromSelectionIn.Selection)
+		const FPolygroupSet GroupSet = FPolygroupSet(&Mesh);
+		return EnumerateSelectionTriangles(FromSelectionIn, Mesh,
+		[&ToSelectionOut](const int32 TID)
 		{
-			Triangles.Add( FGeoSelectionID(ElemID).GeometryID );
-		}
-
-		return InitializeSelectionFromTriangles(Mesh, GroupTopology, Triangles, ToSelectionOut);
+			ToSelectionOut.Selection.Add(FGeoSelectionID::MeshTriangle(TID).Encoded());
+		}, &GroupSet);
 	};
 
-
-
-	const auto FromTriEdgeToTriVtx = [](
+	const auto ToTriEdge = [](
 		const FDynamicMesh3& Mesh,
 		const FGroupTopology* GroupTopology,
 		const FGeometrySelection& FromSelectionIn,
-		FGeometrySelection& ToSelectionOut)
+		FGeometrySelection& ToSelectionOut) -> bool
 	{
-		checkSlow(FromSelectionIn.TopologyType == EGeometryTopologyType::Triangle);
-		checkSlow(FromSelectionIn.ElementType == EGeometryElementType::Edge);
 		checkSlow(ToSelectionOut.TopologyType == EGeometryTopologyType::Triangle);
-		checkSlow(ToSelectionOut.ElementType == EGeometryElementType::Vertex);
-
-		return EnumerateTriangleSelectionVertices(FromSelectionIn, Mesh, nullptr,
-			[&ToSelectionOut](uint64 Vid, const FVector3d& Unused)
-			{
-				ToSelectionOut.Selection.Add( FGeoSelectionID::MeshVertex((int32)Vid).Encoded() );
-			});
-	};
+		checkSlow(ToSelectionOut.ElementType == EGeometryElementType::Edge);
 
+		auto SelectEdgesFunc = [&Mesh, &GroupTopology, &ToSelectionOut] (const int32 EdgeID)
+		{
+			Mesh.EnumerateTriEdgeIDsFromEdgeID(EdgeID,
+				[&ToSelectionOut](const FMeshTriEdgeID TriEdgeID)
+				{
+					ToSelectionOut.Selection.Add(FGeoSelectionID::MeshEdge(TriEdgeID).Encoded());
+				});
+		};
 
-	const auto FromPolyFace = [](
-		const FDynamicMesh3& Mesh,
-		const FGroupTopology* GroupTopology,
-		const FGeometrySelection& FromSelectionIn,
-		FGeometrySelection& ToSelectionOut)
-	{
-		checkSlow(FromSelectionIn.TopologyType == EGeometryTopologyType::Polygroup);
-		checkSlow(FromSelectionIn.ElementType == EGeometryElementType::Face);
-		if (!ensure(GroupTopology))
+		if (FromSelectionIn.TopologyType == EGeometryTopologyType::Triangle)
 		{
-			return false;
+			return EnumerateTriangleSelectionEdges(FromSelectionIn, Mesh, SelectEdgesFunc);
 		}
-
-		TArray<int32> Groups;
-		for (uint64 ElemID : FromSelectionIn.Selection)
+		else if (FromSelectionIn.TopologyType == EGeometryTopologyType::Polygroup)
 		{
-			Groups.Add(FGeoSelectionID(ElemID).TopologyID);
+			return EnumeratePolygroupSelectionEdges(FromSelectionIn, Mesh, *GroupTopology, SelectEdgesFunc);
 		}
-
-		return InitializeSelectionFromPolyGroups(Mesh, *GroupTopology, Groups, ToSelectionOut);
+		return true;
 	};
-	
 
-	const auto FromPolyEdgeToTriEdge = [](
+	const auto ToTriVtx = [](
 		const FDynamicMesh3& Mesh,
 		const FGroupTopology* GroupTopology,
 		const FGeometrySelection& FromSelectionIn,
 		FGeometrySelection& ToSelectionOut) -> bool
 	{
-		checkSlow(FromSelectionIn.TopologyType == EGeometryTopologyType::Polygroup);
-		checkSlow(FromSelectionIn.ElementType == EGeometryElementType::Edge);
 		checkSlow(ToSelectionOut.TopologyType == EGeometryTopologyType::Triangle);
-		checkSlow(ToSelectionOut.ElementType == EGeometryElementType::Edge);
+		checkSlow(ToSelectionOut.ElementType == EGeometryElementType::Vertex);
 
-		if (ensure(GroupTopology))
+		if (FromSelectionIn.TopologyType == EGeometryTopologyType::Triangle)
 		{
-			for (uint64 EncodedID : FromSelectionIn.Selection)
+			return EnumerateTriangleSelectionVertices(FromSelectionIn, Mesh, nullptr,
+			[&ToSelectionOut](const uint64 Vid, const FVector3d& Unused)
 			{
-				const FGeoSelectionID SelectionID(EncodedID);
-				int32 GroupEID = SelectionID.TopologyID;
-				checkSlow(GroupTopology->Edges.IsValidIndex(GroupEID));
-				for (int32 EID : GroupTopology->Edges[GroupEID].Span.Edges)
+				ToSelectionOut.Selection.Add( FGeoSelectionID::MeshVertex((int32)Vid).Encoded() );
+			});
+		}
+		else if (FromSelectionIn.TopologyType == EGeometryTopologyType::Polygroup)
+		{
+			// TODO Add a function which gets the Vids only, to remove the matrix-vector multiplication we just ignore
+			const FTransform Transform = FTransform::Identity;
+			return EnumeratePolygroupSelectionVertices(FromSelectionIn, Mesh, GroupTopology, Transform,
+				[&ToSelectionOut](const uint64 VID, const FVector3d& Unused)
 				{
-					FMeshTriEdgeID MeshEdgeID = Mesh.GetTriEdgeIDFromEdgeID(EID);
-					ToSelectionOut.Selection.Add(MeshEdgeID.Encoded());
-				}
-			}
-			return true;
+					ToSelectionOut.Selection.Add( FGeoSelectionID::MeshVertex((int32)VID).Encoded() );
+				});
 		}
-		return false;
+		return true;
 	};
 
-
-	const auto FromPolyEdgeToPolyVtx = [](
+	const auto ToPolyFace = [](
 		const FDynamicMesh3& Mesh,
 		const FGroupTopology* GroupTopology,
 		const FGeometrySelection& FromSelectionIn,
 		FGeometrySelection& ToSelectionOut) -> bool
 	{
-		checkSlow(FromSelectionIn.TopologyType == EGeometryTopologyType::Polygroup);
-		checkSlow(FromSelectionIn.ElementType == EGeometryElementType::Edge);
 		checkSlow(ToSelectionOut.TopologyType == EGeometryTopologyType::Polygroup);
-		checkSlow(ToSelectionOut.ElementType == EGeometryElementType::Vertex);
+		checkSlow(ToSelectionOut.ElementType == EGeometryElementType::Face);
 
-		if (ensure(GroupTopology))
+		const FPolygroupSet GroupSet = FPolygroupSet(&Mesh);
+		EnumerateSelectionTriangles(FromSelectionIn, Mesh,
+		[&ToSelectionOut, &GroupTopology](const int32 TID)
 		{
-			for (uint64 EncodedID : FromSelectionIn.Selection)
+			const int GroupID = GroupTopology->GetGroupID(TID);
+			const TArray<int>& GroupTris = GroupTopology->GetGroupTriangles(GroupID);
+			for (const int GroupTri : GroupTris)
 			{
-				const FGeoSelectionID SelectionID(EncodedID);
-				int32 GroupEID = SelectionID.TopologyID;
-				checkSlow(GroupTopology->Edges.IsValidIndex(GroupEID));
-				FIndex2i Corners = GroupTopology->Edges[GroupEID].EndpointCorners;
-				for (int32 SubIdx = 0; SubIdx < 2; ++SubIdx)
-				{
-					int32 CornerID = Corners[SubIdx];
-					if (CornerID != IndexConstants::InvalidID)
-					{
-						const FGroupTopology::FCorner& Corner = GroupTopology->Corners[CornerID];
-						FGeoSelectionID ID = FGeoSelectionID(Corner.VertexID, CornerID);
-						ToSelectionOut.Selection.Add(ID.Encoded());
-					}
-				}
+				ToSelectionOut.Selection.Add(FGeoSelectionID(GroupTri, GroupID).Encoded());
 			}
-			return true;
-		}
-		return false;
+		}, &GroupSet);
+		return true;
 	};
 
-	const auto FromPolyEdgeToTriVtx = [](
+	const auto ToPolyEdge = [](
 		const FDynamicMesh3& Mesh,
 		const FGroupTopology* GroupTopology,
 		const FGeometrySelection& FromSelectionIn,
 		FGeometrySelection& ToSelectionOut)
 	{
-		checkSlow(FromSelectionIn.TopologyType == EGeometryTopologyType::Polygroup);
-		checkSlow(FromSelectionIn.ElementType == EGeometryElementType::Edge);
-		checkSlow(ToSelectionOut.TopologyType == EGeometryTopologyType::Triangle);
-		checkSlow(ToSelectionOut.ElementType == EGeometryElementType::Vertex);
-
-		// TODO Add a function which gets the Vids only, to remove the matrix-vector multiplication we just ignore
-		const FTransform Transform = FTransform::Identity;
-		return EnumeratePolygroupSelectionVertices(FromSelectionIn, Mesh, GroupTopology, Transform,
-			[&ToSelectionOut](uint64 Vid, const FVector3d& Unused)
-			{
-				ToSelectionOut.Selection.Add( FGeoSelectionID::MeshVertex((int32)Vid).Encoded() );
-			});
-	};
-
-	const auto FromTriSelectionToPolyVtx = [](
-		const FDynamicMesh3& Mesh,
-		const FGroupTopology* GroupTopology,
-		const FGeometrySelection& FromSelectionIn,
-		FGeometrySelection& ToSelectionOut) -> bool
-	{
-		checkSlow(FromSelectionIn.TopologyType == EGeometryTopologyType::Triangle);
 		checkSlow(ToSelectionOut.TopologyType == EGeometryTopologyType::Polygroup);
-		checkSlow(ToSelectionOut.ElementType == EGeometryElementType::Vertex);
+		checkSlow(ToSelectionOut.ElementType == EGeometryElementType::Edge);
 
 		bool bConverted = false;
-		EnumerateTriangleSelectionVertices(FromSelectionIn, Mesh, nullptr, 
-			[&Mesh, &GroupTopology, &ToSelectionOut, &bConverted](uint64 VID, const FVector& Unused)
+
+		auto SelectEdgesFunc = [&Mesh, &GroupTopology, &ToSelectionOut, &bConverted] (const int32 EdgeID)
+		{
+			// similar but simplified version of code in GroupTopology->IsGroupEdge()
+			const FIndex2i EdgeT = Mesh.GetEdgeT(EdgeID);
+			if (EdgeT.B == IndexConstants::InvalidID)
+			{
+				bConverted = true;
+			}
+			const bool bIsGroupEdge = GroupTopology->GetGroupID(EdgeT.A) !=  GroupTopology->GetGroupID(EdgeT.B);
+			bConverted = bConverted || bIsGroupEdge;
+
+			if (bIsGroupEdge)
 			{
-				if (Mesh.IsVertex((int32)VID))
+				const int32 GroupEdgeID = GroupTopology->FindGroupEdgeID(EdgeID);
+				for (const int32 EID : GroupTopology->Edges[GroupEdgeID].Span.Edges)
 				{
-					int32 CornerID = GroupTopology->GetCornerIDFromVertexID((int32)VID);
-					if (CornerID != IndexConstants::InvalidID)
+					Mesh.EnumerateTriEdgeIDsFromEdgeID(EID,
+				[&ToSelectionOut, &GroupEdgeID](const FMeshTriEdgeID TriEdgeID)
 					{
-						bConverted = true;
-						FGeoSelectionID ID = FGeoSelectionID((int32)VID, CornerID);
-						ToSelectionOut.Selection.Add(ID.Encoded());
-					}
+						ToSelectionOut.Selection.Add(FGeoSelectionID(TriEdgeID.Encoded(), GroupEdgeID).Encoded());
+					});
 				}
-			});
+			}
+		};
+
+		if (FromSelectionIn.TopologyType == EGeometryTopologyType::Triangle)
+		{
+			EnumerateTriangleSelectionEdges(FromSelectionIn, Mesh, SelectEdgesFunc);
+		}
+		else if (FromSelectionIn.TopologyType == EGeometryTopologyType::Polygroup)
+		{
+			EnumeratePolygroupSelectionEdges(FromSelectionIn, Mesh, *GroupTopology, SelectEdgesFunc);
+		}
+		
 		return bConverted;
 	};
 
-	const auto FromPolyVtxToTriVtx = [](
+	const auto ToPolyVtx = [](
 		const FDynamicMesh3& Mesh,
 		const FGroupTopology* GroupTopology,
 		const FGeometrySelection& FromSelectionIn,
-		FGeometrySelection& ToSelectionOut)
+		FGeometrySelection& ToSelectionOut) -> bool
 	{
-		checkSlow(FromSelectionIn.TopologyType == EGeometryTopologyType::Polygroup);
-		checkSlow(FromSelectionIn.ElementType == EGeometryElementType::Vertex);
-		checkSlow(ToSelectionOut.TopologyType == EGeometryTopologyType::Triangle);
+		checkSlow(ToSelectionOut.TopologyType == EGeometryTopologyType::Polygroup);
 		checkSlow(ToSelectionOut.ElementType == EGeometryElementType::Vertex);
 
-		// TODO Add a function which gets the Vids only, to remove the matrix-vector multiplication we just ignore
-		const FTransform Transform = FTransform::Identity;
-		EnumeratePolygroupSelectionVertices(FromSelectionIn, Mesh, GroupTopology, Transform,
-			[&ToSelectionOut](uint64 Vid, const FVector3d& Unused)
+		auto ApplyPolyVertFunc = [&Mesh, GroupTopology, &ToSelectionOut](const uint64 VID, bool &bConverted)
+		{
+			if (Mesh.IsVertex((int32)VID))
 			{
-				ToSelectionOut.Selection.Add( FGeoSelectionID::MeshVertex((int32)Vid).Encoded() );
+				const int32 CornerID = GroupTopology->GetCornerIDFromVertexID((int32)VID);
+				if (CornerID != IndexConstants::InvalidID)
+				{
+					bConverted = true;
+					const FGeoSelectionID ID = FGeoSelectionID((int32)VID, CornerID);
+					ToSelectionOut.Selection.Add(ID.Encoded());
+				}
+			}
+		};
+		
+		bool bConverted = false;
+		if (FromSelectionIn.TopologyType == EGeometryTopologyType::Triangle)
+		{
+			EnumerateTriangleSelectionVertices(FromSelectionIn, Mesh, nullptr, 
+			[&bConverted, &ApplyPolyVertFunc](const uint64 VID, const FVector& Unused)
+			{
+				ApplyPolyVertFunc(VID, bConverted);
 			});
-
-		return true;
+		}
+		else if (FromSelectionIn.TopologyType == EGeometryTopologyType::Polygroup)
+		{
+			const FTransform Transform = FTransform::Identity;
+			EnumeratePolygroupSelectionVertices(FromSelectionIn, Mesh, GroupTopology, Transform, 
+				[&bConverted, &ApplyPolyVertFunc](const uint64 VID, const FVector& Unused)
+				{
+					ApplyPolyVertFunc(VID, bConverted);
+				});
+		}
+		return bConverted;
 	};
 
-
-
 	typedef bool (*ConvertSelectionFunc)(
 			const FDynamicMesh3& Mesh,
 			const FGroupTopology* GroupTopology,
@@ -1805,12 +1784,12 @@ bool UE::Geometry::ConvertSelection(
 			FGeometrySelection& ToSelectionOut);
 
 	constexpr ConvertSelectionFunc ConvertFuncs[6][6] = {
-		{FromTypeToSame,        NotImplemented,         NotImplemented,  FromTriSelectionToPolyVtx,   NotImplemented,         NotImplemented},
-		{FromTriEdgeToTriVtx,   FromTypeToSame,         NotImplemented,  FromTriSelectionToPolyVtx,   NotImplemented,         NotImplemented},
-		{FromTriFace,           FromTriFace,            FromTypeToSame,  FromTriFace,                 FromTriFace,            FromTriFace   },
-		{FromPolyVtxToTriVtx,   NotImplemented,         NotImplemented,  FromTypeToSame,              NotImplemented,         NotImplemented},
-		{FromPolyEdgeToTriVtx,  FromPolyEdgeToTriEdge,  NotImplemented,  FromPolyEdgeToPolyVtx,       FromTypeToSame,         NotImplemented},
-		{FromPolyFace,          FromPolyFace,           FromPolyFace,    FromPolyFace,                FromPolyFace,           FromTypeToSame}
+		{FromTypeToSame,  ToTriEdge,      ToTriFace,      ToPolyVtx,     ToPolyEdge,     ToPolyFace },
+		{ToTriVtx,		  FromTypeToSame, ToTriFace,      ToPolyVtx,     ToPolyEdge,     ToPolyFace },
+		{ToTriVtx,		  ToTriEdge,      FromTypeToSame, ToPolyVtx,     ToPolyEdge,     ToPolyFace },
+		{ToTriVtx,		  ToTriEdge,      ToTriFace,      FromTypeToSame,ToPolyEdge,     ToPolyFace },
+		{ToTriVtx,		  ToTriEdge,      ToTriFace,      ToPolyVtx,     FromTypeToSame, ToPolyFace },
+		{ToTriVtx,		  ToTriEdge,      ToTriFace,      ToPolyVtx,     ToPolyEdge,     FromTypeToSame }
 	};
 
 	const int FromIndex = GetSelectionTypeAsIndex(FromSelectionIn);
diff --git a/Engine/Plugins/Runtime/GeometryProcessing/Source/DynamicMesh/Public/Selections/GeometrySelectionUtil.h b/Engine/Plugins/Runtime/GeometryProcessing/Source/DynamicMesh/Public/Selections/GeometrySelectionUtil.h
--- a/Engine/Plugins/Runtime/GeometryProcessing/Source/DynamicMesh/Public/Selections/GeometrySelectionUtil.h
+++ b/Engine/Plugins/Runtime/GeometryProcessing/Source/DynamicMesh/Public/Selections/GeometrySelectionUtil.h
@@ -215,6 +215,21 @@ DYNAMICMESH_API bool EnumeratePolygroupSelectionEdges(
 	const UE::Geometry::FPolygroupSet& GroupSet,
 	TFunctionRef<void(int32)> EdgeFunc
 );
+
+/**
+ * Version of EnumeratePolygroupSelectionEdges which uses GroupTopology
+ *
+* Call EdgeFunc for each mesh EdgeID included in MeshSelection, where MeshSelection has polygroup topology.
+ * For Polygroup Faces, all edges (including internal, non-border edges) in the polygroup are included.
+ * For Polygroup Borders, includes all mesh edges which are a part of the selected GroupEdge(s)/BorderEdge(s)
+ * For Polygroup Corners, includes all mesh edges in all GroupEdge(s) which touch the selected Corner(s)
+ */
+DYNAMICMESH_API bool EnumeratePolygroupSelectionEdges(
+	const FGeometrySelection& MeshSelection,
+	const UE::Geometry::FDynamicMesh3& Mesh,
+	const UE::Geometry::FGroupTopology& GroupTopology,
+	TFunctionRef<void(int32)> EdgeFunc
+);	
 	
 /** Prefer EnumerateTriangleSelectionElements with Flags parameter. */
 DYNAMICMESH_API bool EnumerateTriangleSelectionElements(
@@ -314,17 +329,6 @@ DYNAMICMESH_API bool InitializeSelectionFromTriangles(
 	FGeometrySelection& SelectionOut);
 
 
-/**
- * Convert PolyGroup IDs to target Selection type
- */
-DYNAMICMESH_API bool InitializeSelectionFromPolyGroups(
-	const FDynamicMesh3& Mesh,
-	const FGroupTopology& GroupTopology,
-	TArrayView<const int> GroupIDs,
-	FGeometrySelection& SelectionOut);
-
-
-
 /**
  * Convert Selection from one type to another, based on geometry/topology types in FromSelectionIn and ToSelectionOut.
  * 
@@ -334,23 +338,28 @@ DYNAMICMESH_API bool InitializeSelectionFromPolyGroups(
  *                 To:    Triangle               Polygroup           
  *   From:                Vertex  Edge    Face   Vertex  Edge    Face
  *   ----------------------------------------------------------------
- *   Triangle Vertex      1       .       .      4#      .       .   
- *   Triangle Edge        1       1       .      4#      .       .   
- *   Triangle Face        1       1       1      4#      3#      2#  
- *   Polygroup Vertex     6#      .       .      1       .       .   
- *   Polygroup Edge       5#      1#      .      1#      1       .   
- *   Polygroup Face       1#      1#      1#     1#      1#      1   
+ *   Triangle Vertex      1       1       7      4#      7#       7#   
+ *   Triangle Edge        1       1       8      4#      1#       8#   
+ *   Triangle Face        1       1       1      4#      3#       2#  
+ *   Polygroup Vertex     6#      9#      9      1       7#       7#   
+ *   Polygroup Edge       5#      1#      8      1#      1        8#   
+ *   Polygroup Face       1#      1#      1      1#      1#       1   
  *   ================================================================
  *
  *   Key:
  *   .  These conversions are not implemented... yet? GroupTopology is ignored
- *   1  supported. The implementation is obvious/unambiguous
- *   2  supported. Polygroup faces containing any input triangle are selected
- *   3  supported. Polygroup edges containing any input triangle edge are selected, but
+ *   1  - The implementation is obvious/unambiguous
+ *   2  - Polygroup faces containing any input triangle are selected
+ *   3  - Polygroup edges containing any input triangle edge are selected, but
  *                 polygroup edges containing only input triangle vertices are not.
- *   4  supported. Polygroup corners coinciding with any input triangle face/edge/vertex are selected
- *   5  supported. All mesh vertices along the polygroup edge are selected
- *   6  supported. All mesh vertices coinciding with polygroup corners are selected
+ *   4  - Polygroup corners coinciding with any input triangle face/edge/vertex are selected
+ *   5  - All mesh vertices along the polygroup edge are selected
+ *   6  - All mesh vertices coinciding with polygroup corners are selected
+ *   7  - All Elements which immediately touch any input vertex/corner are selected
+ *   8  - All triangles or faces on either side of the any input edge are selected
+ *			PolyEdge->TriFace includes all triangles in both groups which the edge borders
+ *   9  - All edges in each PolyEdge with touch any input Corner are selected
+ *		  All triangles in each PolyFace which touch any input Corner are selected
  *   #  indicates GroupTopology must not be null for this combination. If this symbol is missing GroupTopology is ignored
  *
  * @return true if conversion is supported and was computed successfully, return false otherwise
