Modeling Tools: light cleaning up functions relating to selection, no functionality change

diff --git a/Engine/Plugins/Runtime/GeometryProcessing/Source/DynamicMesh/Private/Selections/GeometrySelectionUtil.cpp b/Engine/Plugins/Runtime/GeometryProcessing/Source/DynamicMesh/Private/Selections/GeometrySelectionUtil.cpp
--- a/Engine/Plugins/Runtime/GeometryProcessing/Source/DynamicMesh/Private/Selections/GeometrySelectionUtil.cpp
+++ b/Engine/Plugins/Runtime/GeometryProcessing/Source/DynamicMesh/Private/Selections/GeometrySelectionUtil.cpp
@@ -572,15 +572,15 @@ bool UE::Geometry::EnumeratePolygroupSelectionVertices(
 
 	if (MeshSelection.ElementType == EGeometryElementType::Face)
 	{
-		for (uint64 EncodedID : MeshSelection.Selection)
+		for (const uint64 EncodedID : MeshSelection.Selection)
 		{
-			FGeoSelectionID GroupTriID(EncodedID);
-			int32 SeedTriangleID = (int32)GroupTriID.GeometryID, GroupID = (int32)GroupTriID.TopologyID;
+			const FGeoSelectionID GroupTriID(EncodedID);
+			const int32 SeedTriangleID = (int32)GroupTriID.GeometryID, GroupID = (int32)GroupTriID.TopologyID;
 			if (Mesh.IsTriangle(SeedTriangleID))
 			{
-				for (int32 TriangleID : GroupTopology->GetGroupFaces(GroupID))
+				for (const int32 TriangleID : GroupTopology->GetGroupFaces(GroupID))
 				{
-					FIndex3i Triangle = Mesh.GetTriangle((int32)TriangleID);
+					const FIndex3i Triangle = Mesh.GetTriangle((int32)TriangleID);
 					VertexFunc((uint64)Triangle.A, ApplyTransform.TransformPosition(Mesh.GetVertex(Triangle.A)));
 					VertexFunc((uint64)Triangle.B, ApplyTransform.TransformPosition(Mesh.GetVertex(Triangle.B)));
 					VertexFunc((uint64)Triangle.C, ApplyTransform.TransformPosition(Mesh.GetVertex(Triangle.C)));
@@ -590,14 +590,14 @@ bool UE::Geometry::EnumeratePolygroupSelectionVertices(
 	}
 	else if (MeshSelection.ElementType == EGeometryElementType::Edge)
 	{
-		for (uint64 EncodedID : MeshSelection.Selection)
+		for (const uint64 EncodedID : MeshSelection.Selection)
 		{
-			FMeshTriEdgeID TriEdgeID( FGeoSelectionID(EncodedID).GeometryID );
-			int32 SeedEdgeID = Mesh.IsTriangle(TriEdgeID.TriangleID) ? Mesh.GetTriEdge(TriEdgeID.TriangleID, TriEdgeID.TriEdgeIndex) : IndexConstants::InvalidID;
+			const FMeshTriEdgeID TriEdgeID( FGeoSelectionID(EncodedID).GeometryID );
+			const int32 SeedEdgeID = Mesh.IsTriangle(TriEdgeID.TriangleID) ? Mesh.GetTriEdge(TriEdgeID.TriangleID, TriEdgeID.TriEdgeIndex) : IndexConstants::InvalidID;
 			if (Mesh.IsEdge(SeedEdgeID))
 			{
-				int32 GroupEdgeID = GroupTopology->FindGroupEdgeID(SeedEdgeID);
-				for (int32 VertexID : GroupTopology->GetGroupEdgeVertices(GroupEdgeID))
+				const int32 GroupEdgeID = GroupTopology->FindGroupEdgeID(SeedEdgeID);
+				for (const int32 VertexID : GroupTopology->GetGroupEdgeVertices(GroupEdgeID))
 				{
 					FVector3d V = Mesh.GetVertex(VertexID);
 					VertexFunc((uint64)VertexID, ApplyTransform.TransformPosition(V));
@@ -607,9 +607,9 @@ bool UE::Geometry::EnumeratePolygroupSelectionVertices(
 	}
 	else if (MeshSelection.ElementType == EGeometryElementType::Vertex)
 	{
-		for (uint64 EncodedID : MeshSelection.Selection)
+		for (const uint64 EncodedID : MeshSelection.Selection)
 		{
-			int32 VertexID = (int32)FGeoSelectionID(EncodedID).GeometryID;
+			const int32 VertexID = (int32)FGeoSelectionID(EncodedID).GeometryID;
 			if (Mesh.IsVertex(VertexID))
 			{
 				VertexFunc((uint64)VertexID, ApplyTransform.TransformPosition(Mesh.GetVertex(VertexID)));
@@ -659,7 +659,7 @@ bool UE::Geometry::EnumerateTriangleSelectionTriangles(
 
 	if (MeshSelection.ElementType == EGeometryElementType::Face)
 	{
-		for (uint64 TriangleID : MeshSelection.Selection)
+		for (const uint64 TriangleID : MeshSelection.Selection)
 		{
 			if (Mesh.IsTriangle((int32)TriangleID))
 			{
@@ -669,11 +669,11 @@ bool UE::Geometry::EnumerateTriangleSelectionTriangles(
 	}
 	else if (MeshSelection.ElementType == EGeometryElementType::Edge)
 	{
-		for (uint64 EncodedID : MeshSelection.Selection)
+		for (const uint64 EncodedID : MeshSelection.Selection)
 		{
-			FMeshTriEdgeID TriEdgeID( FGeoSelectionID(EncodedID).GeometryID );
-			int32 EdgeID = Mesh.IsTriangle(TriEdgeID.TriangleID) ? Mesh.GetTriEdge(TriEdgeID.TriangleID, TriEdgeID.TriEdgeIndex) : IndexConstants::InvalidID;
-			Mesh.EnumerateEdgeTriangles(EdgeID, [&](int32 TriangleID)
+			const FMeshTriEdgeID TriEdgeID( FGeoSelectionID(EncodedID).GeometryID );
+			const int32 EdgeID = Mesh.IsTriangle(TriEdgeID.TriangleID) ? Mesh.GetTriEdge(TriEdgeID.TriangleID, TriEdgeID.TriEdgeIndex) : IndexConstants::InvalidID;
+			Mesh.EnumerateEdgeTriangles(EdgeID, [&TriangleFunc](const int32 TriangleID)
 			{
 				TriangleFunc(TriangleID);
 			});
@@ -681,9 +681,9 @@ bool UE::Geometry::EnumerateTriangleSelectionTriangles(
 	}
 	else if (MeshSelection.ElementType == EGeometryElementType::Vertex)
 	{
-		for (uint64 VertexID : MeshSelection.Selection)
+		for (const uint64 VertexID : MeshSelection.Selection)
 		{
-			Mesh.EnumerateVertexTriangles((int32)VertexID, [&](int32 TriangleID)
+			Mesh.EnumerateVertexTriangles((int32)VertexID, [&TriangleFunc](const int32 TriangleID)
 			{
 				TriangleFunc(TriangleID);
 			});
@@ -722,13 +722,13 @@ bool UE::Geometry::EnumeratePolygroupSelectionTriangles(
 
 	if (MeshSelection.ElementType == EGeometryElementType::Face)
 	{
-		for (uint64 EncodedID : MeshSelection.Selection)
+		for (const uint64 EncodedID : MeshSelection.Selection)
 		{
-			FGeoSelectionID SelectionID(EncodedID);
-			int32 SeedTriangleID = (int32)SelectionID.GeometryID;
+			const FGeoSelectionID SelectionID(EncodedID);
+			const int32 SeedTriangleID = (int32)SelectionID.GeometryID;
 			if (Mesh.IsTriangle(SeedTriangleID))
 			{
-				int32 GroupID = GroupSet.GetGroup(SeedTriangleID);
+				const int32 GroupID = GroupSet.GetGroup(SeedTriangleID);
 				check(GroupID == (int32)SelectionID.TopologyID);		// sanity-check that we are using the right group
 				if ( GroupID >= 0 && UniqueSeedGroups.Contains(GroupID) == false)
 				{
@@ -741,43 +741,37 @@ bool UE::Geometry::EnumeratePolygroupSelectionTriangles(
 	}
 	else if (MeshSelection.ElementType == EGeometryElementType::Edge)
 	{
-		for (uint64 EncodedID : MeshSelection.Selection)
+		for (const uint64 EncodedID : MeshSelection.Selection)
 		{
-			FMeshTriEdgeID TriEdgeID( FGeoSelectionID(EncodedID).GeometryID );
-			int32 SeedEdgeID = Mesh.IsTriangle(TriEdgeID.TriangleID) ? Mesh.GetTriEdge(TriEdgeID.TriangleID, TriEdgeID.TriEdgeIndex) : IndexConstants::InvalidID;
-			if (Mesh.IsEdge(SeedEdgeID))
+			const FMeshTriEdgeID TriEdgeID( FGeoSelectionID(EncodedID).GeometryID );
+			const int32 SeedEdgeID = Mesh.IsTriangle(TriEdgeID.TriangleID) ? Mesh.GetTriEdge(TriEdgeID.TriangleID, TriEdgeID.TriEdgeIndex) : IndexConstants::InvalidID;
+			Mesh.EnumerateEdgeTriangles(SeedEdgeID, [&GroupSet, &UniqueSeedGroups, &SeedGroups, &SeedTriangles](const int32 TriangleID)
 			{
-				Mesh.EnumerateEdgeTriangles(SeedEdgeID, [&](int32 TriangleID)
+				const int32 GroupID = GroupSet.GetGroup(TriangleID);
+				if (GroupID >= 0 && UniqueSeedGroups.Contains(GroupID) == false)
 				{
-					int32 GroupID = GroupSet.GetGroup(TriangleID);
-					if (GroupID >= 0 && UniqueSeedGroups.Contains(GroupID) == false)
-					{
-						UniqueSeedGroups.Add(GroupID);
-						SeedGroups.Add(GroupID);
-						SeedTriangles.Add(TriangleID);
-					}
-				});
-			}
+					UniqueSeedGroups.Add(GroupID);
+					SeedGroups.Add(GroupID);
+					SeedTriangles.Add(TriangleID);
+				}
+			});
 		}
 	}
 	else if (MeshSelection.ElementType == EGeometryElementType::Vertex)
 	{
-		for (uint64 EncodedID : MeshSelection.Selection)
+		for (const uint64 EncodedID : MeshSelection.Selection)
 		{
-			int32 VertexID = (int32)FGeoSelectionID(EncodedID).GeometryID;
-			if (Mesh.IsVertex(VertexID))
+			const int32 VertexID = (int32)FGeoSelectionID(EncodedID).GeometryID;
+			Mesh.EnumerateVertexTriangles(VertexID, [&GroupSet, &UniqueSeedGroups, &SeedGroups, &SeedTriangles](const int32 TriangleID)
 			{
-				Mesh.EnumerateVertexTriangles(VertexID, [&](int32 TriangleID)
+				const int32 GroupID = GroupSet.GetGroup(TriangleID);
+				if (GroupID >= 0 && UniqueSeedGroups.Contains(GroupID) == false)
 				{
-					int32 GroupID = GroupSet.GetGroup(TriangleID);
-					if (GroupID >= 0 && UniqueSeedGroups.Contains(GroupID) == false)
-					{
-						UniqueSeedGroups.Add(GroupID);
-						SeedGroups.Add(GroupID);
-						SeedTriangles.Add(TriangleID);
-					}
-				});
-			}
+					UniqueSeedGroups.Add(GroupID);
+					SeedGroups.Add(GroupID);
+					SeedTriangles.Add(TriangleID);
+				}
+			});
 		}
 	}
 	else
@@ -788,17 +782,20 @@ bool UE::Geometry::EnumeratePolygroupSelectionTriangles(
 
 	TSet<int> TempROI;		// if we could provide this as input we would not need a temporary roi...
 	TArray<int32> QueueBuffer;
-	int32 NumGroups = SeedGroups.Num();
+	const int32 NumGroups = SeedGroups.Num();
 	for (int32 k = 0; k < NumGroups; ++k)
 	{
 		check(GroupSet.GetGroup(SeedTriangles[k]) == SeedGroups[k]);
-		int32 GroupID = SeedGroups[k];
+		const int32 GroupID = SeedGroups[k];
 		FMeshConnectedComponents::GrowToConnectedTriangles(&Mesh, 
 			TArray<int>{SeedTriangles[k]}, TempROI, &QueueBuffer, 
-			[&](int32 T1, int32 T2) { return GroupSet.GetGroup(T2) == GroupID; });
-		for (int32 tid : TempROI)
+			[&GroupSet, &GroupID](const int32 T1, const int32 T2)
+			{
+				return GroupSet.GetGroup(T2) == GroupID;
+			});
+		for (const int32 TID : TempROI)
 		{
-			TriangleFunc(tid);
+			TriangleFunc(TID);
 		}
 	}
 
@@ -842,7 +839,7 @@ bool UE::Geometry::EnumerateTriangleSelectionEdges(
 		{
 			if (Mesh.IsTriangle((int32)TriangleID))
 			{
-				FIndex3i TriEdges = Mesh.GetTriEdges((int32)TriangleID);
+				const FIndex3i TriEdges = Mesh.GetTriEdges((int32)TriangleID);
 				EdgeFunc(TriEdges[0]);
 				EdgeFunc(TriEdges[1]);
 				EdgeFunc(TriEdges[2]);
@@ -865,7 +862,7 @@ bool UE::Geometry::EnumerateTriangleSelectionEdges(
 	{
 		for (const uint64 VertexID : MeshSelection.Selection)
 		{
-			Mesh.EnumerateVertexEdges((int32)VertexID, [&EdgeFunc](int32 EdgeID)
+			Mesh.EnumerateVertexEdges((int32)VertexID, [&EdgeFunc](const int32 EdgeID)
 			{
 				EdgeFunc(EdgeID);
 			});
@@ -894,8 +891,7 @@ bool UE::Geometry::EnumeratePolygroupSelectionEdges(
 
 	TArray<int32> SeedTriGroups;
 	TArray<int32> SeedTriangles;
-	TSet<int32> UniqueSeedTriGroups;
-	
+
 	TArray<int32> SeedEdges;
 
 	// TODO: the face code below will not work correctly if the selection contains
@@ -906,13 +902,14 @@ bool UE::Geometry::EnumeratePolygroupSelectionEdges(
 
 	if (MeshSelection.ElementType == EGeometryElementType::Face)
 	{
+		TSet<int32> UniqueSeedTriGroups;
 		for (const uint64 EncodedID : MeshSelection.Selection)
 		{
 			const FGeoSelectionID SelectionID(EncodedID);
 			const int32 SeedTriangleID = (int32)SelectionID.GeometryID;
 			if (Mesh.IsTriangle(SeedTriangleID))
 			{
-				int32 GroupID = GroupSet.GetGroup(SeedTriangleID);
+				const int32 GroupID = GroupSet.GetGroup(SeedTriangleID);
 				check(GroupID == (int32)SelectionID.TopologyID);		// sanity-check that we are using the right group
 				if ( GroupID >= 0 && UniqueSeedTriGroups.Contains(GroupID) == false)
 				{
@@ -972,7 +969,7 @@ bool UE::Geometry::EnumeratePolygroupSelectionEdges(
 		{
 			TArray<int32> EdgeGroups; // the 1 or 2 groups which an edge belongs to
 			// retrieves the 1 or 2 triangles and groups to which a selected edge belongs, and finds which group(s) the triangles belong to
-			Mesh.EnumerateEdgeTriangles(SeedEdges[j], [&GroupSet, &EdgeGroups](int32 TriangleID)
+			Mesh.EnumerateEdgeTriangles(SeedEdges[j], [&GroupSet, &EdgeGroups](const int32 TriangleID)
 			{
 				const int32 GroupID = GroupSet.GetGroup(TriangleID);
 				if (GroupID >= 0)
@@ -1304,7 +1301,7 @@ if ( ensure( MeshSelection.TopologyType == EGeometryTopologyType::Polygroup ) ==
 	{
 		for (const uint64 EncodedID : MeshSelection.Selection)
 		{
-			FGeoSelectionID SelectionID(EncodedID);
+			const FGeoSelectionID SelectionID(EncodedID);
 			const int32 SeedTriangleID = (int32)SelectionID.GeometryID, GroupID = (int32)SelectionID.TopologyID;
 			if (Mesh.IsTriangle(SeedTriangleID))
 			{
@@ -1460,11 +1457,11 @@ bool UE::Geometry::InitializeSelectionFromTriangles(
 	{
 		if (SelectionOut.ElementType == EGeometryElementType::Vertex)
 		{
-			for (int32 tid : Triangles)
+			for (const int32 TID : Triangles)
 			{
-				if (Mesh.IsTriangle(tid))
+				if (Mesh.IsTriangle(TID))
 				{
-					FIndex3i TriVertices = Mesh.GetTriangle(tid);
+					const FIndex3i TriVertices = Mesh.GetTriangle(TID);
 					SelectionOut.Selection.Add(FGeoSelectionID::MeshVertex(TriVertices.A).Encoded() );
 					SelectionOut.Selection.Add(FGeoSelectionID::MeshVertex(TriVertices.B).Encoded());
 					SelectionOut.Selection.Add(FGeoSelectionID::MeshVertex(TriVertices.C).Encoded());
@@ -1487,11 +1484,11 @@ bool UE::Geometry::InitializeSelectionFromTriangles(
 		}
 		else if (SelectionOut.ElementType == EGeometryElementType::Face)
 		{
-			for (int32 tid : Triangles)
+			for (const int32 TID : Triangles)
 			{
-				if (Mesh.IsTriangle(tid))
+				if (Mesh.IsTriangle(TID))
 				{
-					SelectionOut.Selection.Add(FGeoSelectionID::MeshTriangle(tid).Encoded());
+					SelectionOut.Selection.Add(FGeoSelectionID::MeshTriangle(TID).Encoded());
 				}
 			}
 		}
@@ -1512,13 +1509,13 @@ bool UE::Geometry::InitializeSelectionFromTriangles(
 		{
 			FMeshVertexSelection VertSelection(&Mesh);
 			VertSelection.SelectTriangleVertices(Triangles);
-			for (int32 vid : VertSelection)
+			for (const int32 VID : VertSelection)
 			{
-				int32 CornerID = GroupTopology->GetCornerIDFromVertexID(vid);
+				const int32 CornerID = GroupTopology->GetCornerIDFromVertexID(VID);
 				if (CornerID != IndexConstants::InvalidID)
 				{
 					const FGroupTopology::FCorner& Corner = GroupTopology->Corners[CornerID];
-					FGeoSelectionID ID = FGeoSelectionID(Corner.VertexID, CornerID);
+					const FGeoSelectionID ID = FGeoSelectionID(Corner.VertexID, CornerID);
 					SelectionOut.Selection.Add(ID.Encoded());
 				}
 			}
@@ -1527,29 +1524,29 @@ bool UE::Geometry::InitializeSelectionFromTriangles(
 		{
 			FMeshEdgeSelection EdgeSelection(&Mesh);
 			EdgeSelection.SelectTriangleEdges(Triangles);
-			for (int32 eid : EdgeSelection)
+			for (const int32 EID : EdgeSelection)
 			{
-				int32 GroupEdgeID = GroupTopology->FindGroupEdgeID(eid);
+				const int32 GroupEdgeID = GroupTopology->FindGroupEdgeID(EID);
 				if (GroupEdgeID != IndexConstants::InvalidID)
 				{
 					const FGroupTopology::FGroupEdge& GroupEdge = GroupTopology->Edges[GroupEdgeID];
-					FMeshTriEdgeID MeshEdgeID = Mesh.GetTriEdgeIDFromEdgeID(GroupEdge.Span.Edges[0]);
-					FGeoSelectionID ID = FGeoSelectionID(MeshEdgeID.Encoded(), GroupEdgeID);
+					const FMeshTriEdgeID MeshEdgeID = Mesh.GetTriEdgeIDFromEdgeID(GroupEdge.Span.Edges[0]);
+					const FGeoSelectionID ID = FGeoSelectionID(MeshEdgeID.Encoded(), GroupEdgeID);
 					SelectionOut.Selection.Add(ID.Encoded());
 				}
 			}
 		}
 		else if (SelectionOut.ElementType == EGeometryElementType::Face)
 		{
-			for (int32 tid : Triangles)
+			for (const int32 TID : Triangles)
 			{
-				if (Mesh.IsTriangle(tid))
+				if (Mesh.IsTriangle(TID))
 				{
-					int32 GroupID = GroupTopology->GetGroupID(tid);
+					const int32 GroupID = GroupTopology->GetGroupID(TID);
 					const FGroupTopology::FGroup* GroupFace = GroupTopology->FindGroupByID(GroupID);
 					if ( GroupFace )
 					{
-						FGeoSelectionID ID = FGeoSelectionID(GroupFace->Triangles[0], GroupFace->GroupID);
+						const FGeoSelectionID ID = FGeoSelectionID(GroupFace->Triangles[0], GroupFace->GroupID);
 						SelectionOut.Selection.Add(ID.Encoded());
 					}
 				}
diff --git a/Engine/Plugins/Runtime/MeshModelingToolset/Source/ModelingComponents/Private/Selection/DynamicMeshSelector.cpp b/Engine/Plugins/Runtime/MeshModelingToolset/Source/ModelingComponents/Private/Selection/DynamicMeshSelector.cpp
--- a/Engine/Plugins/Runtime/MeshModelingToolset/Source/ModelingComponents/Private/Selection/DynamicMeshSelector.cpp
+++ b/Engine/Plugins/Runtime/MeshModelingToolset/Source/ModelingComponents/Private/Selection/DynamicMeshSelector.cpp
@@ -291,10 +291,12 @@ void FBaseDynamicMeshSelector::UpdateSelectionViaRaycast_MeshTopology(
 	UE::Geometry::UpdateTriangleSelectionViaRaycast(
 		GetColliderMesh(), &SelectionEditor,
 		LocalRay, UpdateConfig, ResultOut);
-	
+
+	// necessary so that edge is recognized from both 'sides' since edge is stored as MeshTriEdgeID.Encoded in SelectionEditor
+	// therefore need to include both of the TriEdgeIDs which the edge belongs to to the SelectionEditor
 	if (SelectionEditor.GetElementType() == EGeometryElementType::Edge)
 	{
-		for (uint64 AddedElement : ResultOut.SelectionDelta.Added)
+		for (const uint64 AddedElement : ResultOut.SelectionDelta.Added)
 		{
 			FMeshTriEdgeID TriEdgeID(FGeoSelectionID(AddedElement).GeometryID);
 			TargetMesh->ProcessMesh([TriEdgeID, &SelectionEditor](const UE::Geometry::FDynamicMesh3& SourceMesh)
@@ -303,28 +305,16 @@ void FBaseDynamicMeshSelector::UpdateSelectionViaRaycast_MeshTopology(
 				const int32 EdgeID = SourceMesh.IsTriangle(TriEdgeID.TriangleID) ? SourceMesh.GetTriEdge(TriEdgeID.TriangleID, TriEdgeID.TriEdgeIndex) : IndexConstants::InvalidID;
 				if (SourceMesh.IsEdge(EdgeID))
 				{
-					const FIndex2i EdgeV = SourceMesh.GetEdgeV(EdgeID);
-					// necessary so that edge is recognized from both 'sides' since edge is stored as MeshTriEdgeID.Encoded in SelectionEditor
-					// therefore need to include both of the Tris which the edge belongs to to the SelectionEditor
-					SourceMesh.EnumerateEdgeTriangles(EdgeID, [&SourceMesh, &EdgeV, &SelectionEditor, &TriEdgeID](const int32 TriangleID)
-					{
-						// avoid adding to selection edge already exists in selection
-						if (TriangleID == TriEdgeID.TriangleID)
-						{
-							return;
-						}
-
-						// iterate through all edges of the triangle
-						for (int i = 0; i <= 2; i++)
+					SourceMesh.EnumerateTriEdgeIDsFromEdgeID(EdgeID,
+						[&SelectionEditor, &TriEdgeID](const FMeshTriEdgeID OtherTriEdgeID)
 						{
-							const int32 OtherEdgeID = SourceMesh.GetTriEdge(TriangleID, i);
-							if (EdgeV == SourceMesh.GetEdgeV(OtherEdgeID))
+							// avoid adding to selection edge already exists in selection
+							if (OtherTriEdgeID.TriangleID == TriEdgeID.TriangleID)
 							{
-								FMeshTriEdgeID MeshTriEdgeToSelect = FMeshTriEdgeID(TriangleID, i);
-								SelectionEditor.Select((uint64)MeshTriEdgeToSelect.Encoded());
+								return;
 							}
-						}
-					});
+							SelectionEditor.Select(OtherTriEdgeID.Encoded());
+						});
 				}
 			});
 		}
