make the exterior visibility occlusion tester slightly perturb sampling directions to reduce the chance of perfectly axis-aligned rays, which are more likely to pass through tiny cracks

diff --git a/Engine/Plugins/Runtime/GeometryProcessing/Source/DynamicMesh/Private/Operations/DetectExteriorVisibility.cpp b/Engine/Plugins/Runtime/GeometryProcessing/Source/DynamicMesh/Private/Operations/DetectExteriorVisibility.cpp
--- a/Engine/Plugins/Runtime/GeometryProcessing/Source/DynamicMesh/Private/Operations/DetectExteriorVisibility.cpp
+++ b/Engine/Plugins/Runtime/GeometryProcessing/Source/DynamicMesh/Private/Operations/DetectExteriorVisibility.cpp
@@ -69,10 +69,13 @@ void ComputePerTriangleOcclusionHelper(const FDynamicMesh3& TargetMesh, const Sp
 
 	// array of (+/-)X/Y/Z directions
 	TArray<FVector3d> CardinalDirections;
+	// Small arbitrary offset vector, used to reduce chance of rays exactly fitting through small axis-aligned cracks
+	const FVector3d TinyPerturb(UE_DOUBLE_KINDA_SMALL_NUMBER, .17 * UE_DOUBLE_KINDA_SMALL_NUMBER, .3*UE_DOUBLE_KINDA_SMALL_NUMBER);
 	for (int32 k = 0; k < 3; ++k)
 	{
-		FVector3d Direction(0, 0, 0);
+		FVector3d Direction = TinyPerturb;
 		Direction[k] = 1.0;
+		Direction.Normalize();
 		CardinalDirections.Add(Direction);
 		CardinalDirections.Add(-Direction);
 	}
@@ -142,7 +145,10 @@ void ComputePerTriangleOcclusionHelper(const FDynamicMesh3& TargetMesh, const Sp
 				}
 			}
 		}
-		if (FindHitTriangleTest(Centroid, Normal, Centroid + Radius * Normal, TID, UnusedBaryCoords) == TID)
+		
+		// Slightly nudge the normal direction to reduce the chance of edge cases for meshes w/ features aligned to the normals
+		const FVector3d NormalRayDir = Normal + TinyPerturb;
+		if (FindHitTriangleTest(Centroid, Normal, Centroid + (1.01 * Radius) * NormalRayDir, TID, UnusedBaryCoords) == TID)
 		{
 			TriStatusKnown[TID] = EVisStatus::Visible;
 			return;
@@ -219,7 +225,7 @@ void ComputePerTriangleOcclusionHelper(const FDynamicMesh3& TargetMesh, const Sp
 	uint32 DirectionIndex = 0;
 	while (ModuloIter.GetNextIndex(DirectionIndex))
 	{
-		VisibilityDirections.Add(Normalized(SphereSampler[DirectionIndex]));
+		VisibilityDirections.Add(Normalized(SphereSampler[DirectionIndex] + TinyPerturb));
 	}
 	// Fibonacci set generally does not include the cardinal directions, but they are highly useful to check
 	VisibilityDirections.Append(CardinalDirections);
