Fix bevel of meshes with bowtie vertices on the bevel path + fix internal bevel processing data being exposed as public, and remove some unused fields

diff --git a/Engine/Plugins/Runtime/GeometryProcessing/Source/DynamicMesh/Private/Operations/MeshBevel.cpp b/Engine/Plugins/Runtime/GeometryProcessing/Source/DynamicMesh/Private/Operations/MeshBevel.cpp
--- a/Engine/Plugins/Runtime/GeometryProcessing/Source/DynamicMesh/Private/Operations/MeshBevel.cpp
+++ b/Engine/Plugins/Runtime/GeometryProcessing/Source/DynamicMesh/Private/Operations/MeshBevel.cpp
@@ -14,6 +14,7 @@
 #include "Operations/PolyEditingEdgeUtil.h"
 #include "Operations/PolyEditingUVUtil.h"
 #include "Algo/Count.h"
+#include "Algo/RemoveIf.h"
 #include "Distance/DistLine3Line3.h"
 #include "Operations/UniformTessellate.h"
 #include "DynamicSubmesh3.h"
@@ -207,10 +208,171 @@ bool FMeshBevel::InitializeFromTriangleSet(const FDynamicMesh3& Mesh, const TArr
 	return true;
 }
 
+void FMeshBevel::FixBowties(FDynamicMesh3& Mesh, FDynamicMeshChangeTracker* ChangeTracker)
+{
+	TMultiMap<int32, int32> SplitMeshVIDs; // map from original VID to new VIDs for any bowtie vertices that were split
+	auto SplitBowtie = [this, &Mesh, &SplitMeshVIDs, &ChangeTracker](int32 VID, int32 EID) -> int32
+	{
+		FIndex2i EdgeV = Mesh.GetEdgeV(EID);
+		int32 SubIdx = EdgeV.IndexOf(VID);
+		check(SubIdx != INDEX_NONE);
+		int32 OtherVID = EdgeV[1 - SubIdx];
+
+		if (ChangeTracker)
+		{
+			ChangeTracker->SaveVertexOneRingTriangles(VID, true);
+		}
+		FDynamicMeshEditor Edit(&Mesh);
+		FDynamicMeshEditResult EditResult;
+		Edit.SplitBowties(VID, EditResult);
+		FIndex2i UpdatedEdgeV = Mesh.GetEdgeV(EID);
+		int32 OtherVIDSubIdx = UpdatedEdgeV.IndexOf(OtherVID);
+		check(OtherVIDSubIdx != INDEX_NONE);
+		int32 NewVID = UpdatedEdgeV[1 - OtherVIDSubIdx];
+		
+		for (int32 AddedVID : EditResult.NewVertices)
+		{
+			SplitMeshVIDs.Add(VID, AddedVID);
+		}
+		return NewVID;
+	};
+
+	auto RemapVertexID = [this, &Mesh, &SplitMeshVIDs, &SplitBowtie](int32 OrigVID, int32 EID, int32& RemapVID) -> bool
+	{
+		bool bWasBowtie = SplitMeshVIDs.Contains(OrigVID);
+		RemapVID = OrigVID;
+		if (bWasBowtie)
+		{
+			FIndex2i EdgeV = Mesh.GetEdgeV(EID);
+			if (!EdgeV.Contains(OrigVID))
+			{
+				for (auto KeyIter = SplitMeshVIDs.CreateConstKeyIterator(OrigVID); KeyIter; ++KeyIter)
+				{
+					int32 NewVID = KeyIter.Value();
+					if (EdgeV.Contains(NewVID))
+					{
+						RemapVID = NewVID;
+						return true;
+					}
+				}
+				// the edge should always include one of the vertices that the bowtie was split in to
+				checkSlow(false);
+			}
+			return true;
+		}
+		else if (Mesh.IsBowtieVertex(OrigVID))
+		{
+			 RemapVID = SplitBowtie(OrigVID, EID);
+			 return true;
+		}
+		return false;
+	};
 
+	for (FBevelLoop& Loop : Loops)
+	{
+		checkSlow(Loop.MeshVertices.Num() == Loop.MeshEdges.Num());
+		for (int32 VertPathIdx = 0; VertPathIdx < Loop.MeshVertices.Num(); ++VertPathIdx)
+		{
+			int32 OrigVID = Loop.MeshVertices[VertPathIdx];
+			int32 EdgePathIdx = VertPathIdx;
+			int32 EID = Loop.MeshEdges[EdgePathIdx];
+			int32 RemapVID = OrigVID;
+			if (RemapVertexID(OrigVID, EID, RemapVID))
+			{
+				Loop.MeshVertices[VertPathIdx] = RemapVID;
+			}
+		}
+	}
+
+	int32 OrigNumVertices = Vertices.Num();
+	TSet<int32> BevelVerticesToRelink;
+	TMap<int32, int32> MeshVIDToBevelVertexIdx;
+	for (FBevelEdge& Edge : Edges)
+	{
+		checkSlow(Edge.MeshVertices.Num() == Edge.MeshEdges.Num() + 1);
+		for (int32 VertPathIdx = 0; VertPathIdx < Edge.MeshVertices.Num(); ++VertPathIdx)
+		{
+			int32 OrigVID = Edge.MeshVertices[VertPathIdx];
+			int32 EdgePathIdx = FMath::Min(VertPathIdx, Edge.MeshEdges.Num() - 1);
+			int32 EID = Edge.MeshEdges[EdgePathIdx];
+			int32 RemapVID = OrigVID;
+			if (RemapVertexID(OrigVID, EID, RemapVID))
+			{
+				Edge.MeshVertices[VertPathIdx] = RemapVID;
+
+				// if we're at an endpoint of the bevel edge, also update the bevel vertex, and make the edge point to the updated bevel vertex
+				if (VertPathIdx == 0 || VertPathIdx + 1 == Edge.MeshVertices.Num())
+				{
+					int32 BevelVertexSubIdx = VertPathIdx == 0 ? 0 : 1;
+					int32 BevelVertexIdx = Edge.BevelVertices[BevelVertexSubIdx];
+					VertexIDToIndexMap.Remove(OrigVID);
+					bool bWasAlreadyRelinked;
+					BevelVerticesToRelink.Add(BevelVertexIdx, &bWasAlreadyRelinked);
+
+					// after splitting bowties, recompute whether the split vertex is still a boundary vertex
+					Edge.bEndpointBoundaryFlag[BevelVertexSubIdx] = Mesh.IsBoundaryVertex(RemapVID);
+
+					if (!bWasAlreadyRelinked)
+					{
+						// First time encountering this bevel vertex; just remap it to the split vertex on the current edge
+						Vertices[BevelVertexIdx].VertexID = RemapVID;
+						MeshVIDToBevelVertexIdx.Add(RemapVID, BevelVertexIdx);
+					}
+					else
+					{
+						// We've encountered this bevel vertex before; use a previously-found mapping, or copy the vertex to create a new one
+						int32* FoundBevelVert = MeshVIDToBevelVertexIdx.Find(RemapVID);
+						int32 NewBevelVertIdx = INDEX_NONE;
+						if (FoundBevelVert)
+						{
+							NewBevelVertIdx = *FoundBevelVert;
+						}
+						else
+						{
+							// Add a copy of the existing bevel vertex
+							NewBevelVertIdx = Vertices.Add(FBevelVertex(Vertices[BevelVertexIdx]));
+							Vertices[NewBevelVertIdx].VertexID = RemapVID;
+							BevelVerticesToRelink.Add(NewBevelVertIdx);
+							MeshVIDToBevelVertexIdx.Add(RemapVID, NewBevelVertIdx);
+						}
+						Edge.BevelVertices[BevelVertexSubIdx] = NewBevelVertIdx;
+					}
+				}
+			}
+		}
+	}
+
+	// Fix up edge and triangle references on any relinked bevel vertices
+	for (int32 BevelVertexIdx : BevelVerticesToRelink)
+	{
+		FBevelVertex& Vertex = Vertices[BevelVertexIdx];
+		VertexIDToIndexMap.Add(Vertices[BevelVertexIdx].VertexID, BevelVertexIdx);
+		Vertex.VertexType = EBevelVertexType::Unknown; // reset type to unknown; will be set by InitVertexSet below
+		Vertex.IncomingBevelEdgeIndices.SetNum(
+			Algo::RemoveIf(Vertex.IncomingBevelEdgeIndices, [this, BevelVertexIdx](int32 BevelEdgeIdx)
+			{
+				return !Edges[BevelEdgeIdx].BevelVertices.Contains(BevelVertexIdx);
+			}));
+		int32 MeshVertexID = Vertex.VertexID;
+		Vertex.IncomingBevelMeshEdges.SetNum(
+			Algo::RemoveIf(Vertex.IncomingBevelMeshEdges, [this, MeshVertexID, &Mesh](int32 EdgeID)
+			{
+				return !Mesh.GetEdgeV(EdgeID).Contains(MeshVertexID);
+			}));
+		InitVertexSet(Mesh, Vertices[BevelVertexIdx]);
+	}
+
+	for (int32 BevelVertexIdx : BevelVerticesToRelink)
+	{
+		FBevelVertex& Vertex = Vertices[BevelVertexIdx];
+		FinalizeTerminatorVertex(Mesh, Vertex);
+	}
+}
 
 bool FMeshBevel::Apply(FDynamicMesh3& Mesh, FDynamicMeshChangeTracker* ChangeTracker)
 {
+	FixBowties(Mesh, ChangeTracker);
+
 	// disconnect along bevel graph edges/vertices and save necessary info
 	UnlinkEdges(Mesh, ChangeTracker);
 	if (ResultInfo.CheckAndSetCancelled(Progress))
@@ -343,7 +505,9 @@ void FMeshBevel::AddBevelGroupEdge(const FDynamicMesh3& Mesh, const FGroupTopolo
 		if (VertInfo == nullptr)
 		{
 			FBevelVertex NewVertex;
+			PRAGMA_DISABLE_DEPRECATION_WARNINGS
 			NewVertex.CornerID = CornerID;
+			PRAGMA_ENABLE_DEPRECATION_WARNINGS
 			NewVertex.VertexID = VertexID;
 			BevelVertexIndex = Vertices.Num();
 			Vertices.Add(NewVertex);
@@ -351,15 +515,19 @@ void FMeshBevel::AddBevelGroupEdge(const FDynamicMesh3& Mesh, const FGroupTopolo
 			VertInfo = &Vertices[BevelVertexIndex];
 		}
 		VertInfo->IncomingBevelMeshEdges.Add(IncomingEdgeID);
+		PRAGMA_DISABLE_DEPRECATION_WARNINGS
 		VertInfo->IncomingBevelTopoEdges.Add(GroupEdgeID);
+		PRAGMA_ENABLE_DEPRECATION_WARNINGS
 		VertInfo->IncomingBevelEdgeIndices.Add(NewBevelEdgeIndex);
 		Edge.BevelVertices[ci] = BevelVertexIndex;
 	}
 
 	Edge.MeshEdges.Append(MeshEdgeList);
 	Edge.MeshVertices.Append(Topology.Edges[GroupEdgeID].Span.Vertices);
+	PRAGMA_DISABLE_DEPRECATION_WARNINGS
 	Edge.GroupEdgeID = GroupEdgeID;
 	Edge.GroupIDs = Topology.Edges[GroupEdgeID].Groups;
+	PRAGMA_ENABLE_DEPRECATION_WARNINGS
 
 	Edge.MeshEdgeTris.Reserve(Edge.MeshEdges.Num());
 	for (int32 eid : Edge.MeshEdges)
@@ -409,51 +577,82 @@ void FMeshBevel::AddBevelEdgeLoop(const FDynamicMesh3& Mesh, const FEdgeLoop& Me
 }
 
 
+void FMeshBevel::InitVertexSet(const FDynamicMesh3& Mesh, FMeshBevel::FBevelVertex& Vertex)
+{
+	// get sorted list of triangles around the vertex
+	TArray<int> GroupLengths;
+	TArray<bool> bGroupIsLoop;
+	EMeshResult Result = Mesh.GetVtxContiguousTriangles(Vertex.VertexID, Vertex.SortedTriangles, GroupLengths, bGroupIsLoop);
+	if (Result != EMeshResult::Ok || GroupLengths.Num() != 1 || Vertex.SortedTriangles.Num() < 2)
+	{
+		Vertex.VertexType = EBevelVertexType::Unknown;
+		return;
+	}
 
+	// GetVtxContiguousTriangles does not return triangles sorted in a consistent direction. This check will
+	// reverse the ordering such that it is consistently walking counter-clockwise around the vertex (I think...)
+	FIndex3i Tri0 = Mesh.GetTriangle(Vertex.SortedTriangles[0]).GetCycled(Vertex.VertexID);
+	FIndex3i Tri1 = Mesh.GetTriangle(Vertex.SortedTriangles[1]).GetCycled(Vertex.VertexID);
+	if (Tri0.C == Tri1.B)
+	{
+		Algo::Reverse(Vertex.SortedTriangles);
+	}
 
-void FMeshBevel::BuildVertexSets(const FDynamicMesh3& Mesh)
-{
-	// can be parallel
-	for (FBevelVertex& Vertex : Vertices)
+	if (Mesh.IsBoundaryVertex(Vertex.VertexID))
 	{
-		// get sorted list of triangles around the vertex
-		TArray<int> GroupLengths;
-		TArray<bool> bGroupIsLoop;
-		EMeshResult Result = Mesh.GetVtxContiguousTriangles(Vertex.VertexID, Vertex.SortedTriangles, GroupLengths, bGroupIsLoop);
-		if ( Result != EMeshResult::Ok || GroupLengths.Num() != 1 || Vertex.SortedTriangles.Num() < 2)
-		{
-			Vertex.VertexType = EBevelVertexType::Unknown;
-			continue;
-		}
+		Vertex.VertexType = EBevelVertexType::BoundaryVertex;
+		// TODO: we should have a BuildBoundaryVertex function here that correctly populates the 
+		// Wedges for the boundary vertex. The currently BuildJunctionVertex will not be able to do
+		// this because it assumes it can just walk forward from any edge
+		return;
+	}
 
-		// GetVtxContiguousTriangles does not return triangles sorted in a consistent direction. This check will
-		// reverse the ordering such that it is consistently walking counter-clockwise around the vertex (I think...)
-		FIndex3i Tri0 = Mesh.GetTriangle(Vertex.SortedTriangles[0]).GetCycled(Vertex.VertexID);
-		FIndex3i Tri1 = Mesh.GetTriangle(Vertex.SortedTriangles[1]).GetCycled(Vertex.VertexID);
-		if (Tri0.C == Tri1.B)
-		{
-			Algo::Reverse(Vertex.SortedTriangles);
-		}
 
-		if (Mesh.IsBoundaryVertex(Vertex.VertexID))
-		{
-			Vertex.VertexType = EBevelVertexType::BoundaryVertex;
-			// TODO: we should have a BuildBoundaryVertex function here that correctly populates the 
-			// Wedges for the boundary vertex. The currently BuildJunctionVertex will not be able to do
-			// this because it assumes it can just walk forward from any edge
-			continue;
-		}
+	MESH_BEVEL_DEBUG_CHECK(Vertex.IncomingBevelMeshEdges.Num() != 0);		// shouldn't ever happen
+	if (Vertex.IncomingBevelMeshEdges.Num() == 1)
+	{
+		BuildTerminatorVertex(Vertex, Mesh);
+	}
+	else
+	{
+		BuildJunctionVertex(Vertex, Mesh);
+	}
+}
 
+void FMeshBevel::FinalizeTerminatorVertex(const FDynamicMesh3& Mesh, FBevelVertex& Vertex)
+{
+	if (Vertex.VertexType == EBevelVertexType::TerminatorVertex)
+	{
+		int32 OtherVertexID = Vertex.TerminatorInfo.B;
+		int32* OtherBevelVtxIdx = VertexIDToIndexMap.Find(OtherVertexID);
+		if (OtherBevelVtxIdx != nullptr)
+		{
+			// does other vertex have to be a terminator? or can this also happen w/ a junction?
+			FBevelVertex& OtherVertex = Vertices[*OtherBevelVtxIdx];
+			if (OtherVertex.VertexType == EBevelVertexType::TerminatorVertex)
+			{
+				// want to skip this if the ring-split edge is already a bevel edge
+				int32 MeshEdgeID = Mesh.FindEdge(Vertex.VertexID, OtherVertex.VertexID);
+				MESH_BEVEL_DEBUG_CHECK(MeshEdgeID >= 0);
+				if (Mesh.IsEdge(MeshEdgeID) &&
+					Vertex.TerminatorInfo.A == MeshEdgeID &&
+					OtherVertex.TerminatorInfo.A == MeshEdgeID &&		// do we need the other vertex to use the same edge here?  (is this actually a hard constraint on that edge that we should be enforcing??)
+					Vertex.IncomingBevelMeshEdges.Contains(MeshEdgeID) == false)
+				{
+					Vertex.ConnectedBevelVertex = *OtherBevelVtxIdx;
+				}
+			}
 
-		MESH_BEVEL_DEBUG_CHECK(Vertex.IncomingBevelMeshEdges.Num() != 0);		// shouldn't ever happen
-		if (Vertex.IncomingBevelMeshEdges.Num() == 1)
-		{
-			BuildTerminatorVertex(Vertex, Mesh);
-		}
-		else
-		{
-			BuildJunctionVertex(Vertex, Mesh);
 		}
+	}
+}
+
+void FMeshBevel::BuildVertexSets(const FDynamicMesh3& Mesh)
+{
+	// can be parallel
+	for (FBevelVertex& Vertex : Vertices)
+	{
+		InitVertexSet(Mesh, Vertex);
 
 		if (ResultInfo.CheckAndSetCancelled(Progress))
 		{
@@ -467,30 +666,7 @@ void FMeshBevel::BuildVertexSets(const FDynamicMesh3& Mesh)
 	// to detect this case here /before/ we split the mesh up into disconnected parts...
 	for (FBevelVertex& Vertex : Vertices)
 	{
-		if (Vertex.VertexType == EBevelVertexType::TerminatorVertex)
-		{
-			int32 OtherVertexID = Vertex.TerminatorInfo.B;
-			int32* OtherBevelVtxIdx = VertexIDToIndexMap.Find(OtherVertexID);
-			if (OtherBevelVtxIdx != nullptr)
-			{
-				// does other vertex have to be a terminator? or can this also happen w/ a junction?
-				FBevelVertex& OtherVertex = Vertices[*OtherBevelVtxIdx];
-				if (OtherVertex.VertexType == EBevelVertexType::TerminatorVertex)
-				{
-					// want to skip this if the ring-split edge is already a bevel edge
-					int32 MeshEdgeID = Mesh.FindEdge(Vertex.VertexID, OtherVertex.VertexID);
-					MESH_BEVEL_DEBUG_CHECK(MeshEdgeID >= 0);
-					if (Mesh.IsEdge(MeshEdgeID) && 
-						Vertex.TerminatorInfo.A == MeshEdgeID &&
-						OtherVertex.TerminatorInfo.A == MeshEdgeID &&		// do we need the other vertex to use the same edge here?  (is this actually a hard constraint on that edge that we should be enforcing??)
-						Vertex.IncomingBevelMeshEdges.Contains(MeshEdgeID) == false )
-					{
-						Vertex.ConnectedBevelVertex = *OtherBevelVtxIdx;
-					}
-				}
-
-			}
-		}
+		FinalizeTerminatorVertex(Mesh, Vertex);
 	}
 
 }
diff --git a/Engine/Plugins/Runtime/GeometryProcessing/Source/DynamicMesh/Public/Operations/MeshBevel.h b/Engine/Plugins/Runtime/GeometryProcessing/Source/DynamicMesh/Public/Operations/MeshBevel.h
--- a/Engine/Plugins/Runtime/GeometryProcessing/Source/DynamicMesh/Public/Operations/MeshBevel.h
+++ b/Engine/Plugins/Runtime/GeometryProcessing/Source/DynamicMesh/Public/Operations/MeshBevel.h
@@ -166,12 +166,20 @@ public:
 	// at the vertices on either end of the span. Each mesh-edge of the bevel-edge will become a quad
 	struct FBevelEdge
 	{
+		PRAGMA_DISABLE_DEPRECATION_WARNINGS
+		FBevelEdge() = default;
+		FBevelEdge(const FBevelEdge&) = default;
+		FBevelEdge(FBevelEdge&&) = default;
+		PRAGMA_ENABLE_DEPRECATION_WARNINGS
+
 		// initial topological information that defines what happens in unlink/displace/mesh steps
 		int32 EdgeIndex;				// index of this BevelEdge in Edges array
 		TArray<int32> MeshVertices;		// sequential list of mesh vertex IDs along edge
 		TArray<int32> MeshEdges;		// sequential list of mesh edge IDs along edge
 		TArray<FIndex2i> MeshEdgeTris;	// the one or two triangles associated w/ each MeshEdges element in the input mesh
+		UE_DEPRECATED(5.5, "Mapping back to source topology is not used.")
 		int32 GroupEdgeID;				// ID of this edge in external topology (eg FGroupTopology)
+		UE_DEPRECATED(5.5, "Mapping back to source topology is not used.")
 		FIndex2i GroupIDs;				// topological IDs of groups on either side of topological edge
 		bool bEndpointBoundaryFlag[2];	// flag defining whether vertex at start/end of MeshVertices was a boundary vertex
 		TArray<FVector3d> InitialPositions;		// initial vertex positions
@@ -232,11 +240,19 @@ public:
 	// A FBevelVertex may be expanded out into a polygon or just an edge, depending on its Type
 	struct FBevelVertex
 	{
+		PRAGMA_DISABLE_DEPRECATION_WARNINGS
+		FBevelVertex() = default;
+		FBevelVertex(const FBevelVertex&) = default;
+		FBevelVertex(FBevelVertex&&) = default;
+		PRAGMA_ENABLE_DEPRECATION_WARNINGS
+
 		int32 VertexID;												// Initial Mesh Vertex ID for the Bevel Vertex
+		UE_DEPRECATED(5.5, "Mapping back to source topology is not used.")
 		int32 CornerID;												// Initial Group Topology Corner ID for the Bevel Vertex (if exists)
 		EBevelVertexType VertexType = EBevelVertexType::Unknown;	// Type of the Bevel Vertex
 
 		TArray<int32> IncomingBevelMeshEdges;						// Set of (unsorted) Mesh Edges that are destined to be Beveled, coming into the vertex
+		UE_DEPRECATED(5.5, "Mapping back to source topology is not used.")
 		TArray<int32> IncomingBevelTopoEdges;						// Set of (unsorted) Group Topology Edges that are to be Beveled, coming into the vertex
 		TArray<int32> IncomingBevelEdgeIndices;						// Set of (unsorted) indices of FBevelEdge, coming into the vertex
 
@@ -256,7 +272,7 @@ public:
 
 
 
-public:
+protected:
 	TMap<int32, int32> VertexIDToIndexMap;		// map of mesh-vertex-IDs to indices into Vertices list
 	TArray<FBevelVertex> Vertices;				// list of FBevelVertex data structures for mesh vertices that need beveling
 
@@ -328,6 +344,14 @@ protected:
 	void ComputeNormals(FDynamicMesh3& Mesh);
 	void ComputeUVs(FDynamicMesh3& Mesh);
 	void ComputeMaterialIDs(FDynamicMesh3& Mesh);
+
+private:
+	// Detect and fix any bowtie vertices in the bevel operation set
+	// Called by "Apply" (because the setup methods all operate on a const Mesh) before the actual bevel operation
+	void FixBowties(FDynamicMesh3& Mesh, FDynamicMeshChangeTracker* ChangeTracker);
+
+	void InitVertexSet(const FDynamicMesh3& Mesh, FBevelVertex& Vertex);
+	void FinalizeTerminatorVertex(const FDynamicMesh3& Mesh, FBevelVertex& Vertex);
 };
 
 
