Transfer Bone Weights: added the possibility of propagating the skin weights to a subset of vertices only
    if this is straightforward for the closest point algorithm, inpaint can extract data from regions outside the target vertex subset.
    a temporary skin weight attribute is used to modify the weights, then only the subset is copied into the full skin weights attribute.

diff --git a/Engine/Plugins/Runtime/GeometryProcessing/Source/DynamicMesh/Private/Operations/TransferBoneWeights.cpp b/Engine/Plugins/Runtime/GeometryProcessing/Source/DynamicMesh/Private/Operations/TransferBoneWeights.cpp
--- a/Engine/Plugins/Runtime/GeometryProcessing/Source/DynamicMesh/Private/Operations/TransferBoneWeights.cpp
+++ b/Engine/Plugins/Runtime/GeometryProcessing/Source/DynamicMesh/Private/Operations/TransferBoneWeights.cpp
@@ -12,6 +12,7 @@
 #include "BoneIndices.h"
 #include "IndexTypes.h"
 #include "TransformTypes.h"
+#include "Algo/Count.h"
 #include "Solvers/Internal/QuadraticProgramming.h"
 #include "Solvers/LaplacianMatrixAssembly.h"
 #include "Operations/SmoothBoneWeights.h"
@@ -246,18 +247,23 @@ bool FTransferBoneWeights::TransferWeightsToMesh(FDynamicMesh3& InOutTargetMesh,
 	}
 
 	bool bFailed = false;
-	
-	MatchedVertices.Init(false, InOutTargetMesh.MaxVertexID());
 
+	// compute the transfer only for the subset of vertices if necessary  
+	const bool bUseSubset = !TargetVerticesSubset.IsEmpty();
+	const int32 NumVerticesToTransfer = bUseSubset ? TargetVerticesSubset.Num() : InOutTargetMesh.MaxVertexID();
+	
 	if (TransferMethod == ETransferBoneWeightsMethod::ClosestPointOnSurface)
 	{
-		ParallelFor(InOutTargetMesh.MaxVertexID(), [this, &InOutTargetMesh, &TargetBoneToIndex, &TargetSkinWeights, &InternalTargetMeshNormals](int32 VertexID)
+		MatchedVertices.Init(false, NumVerticesToTransfer);
+		
+		ParallelFor(NumVerticesToTransfer, [this, &InOutTargetMesh, &TargetBoneToIndex, &TargetSkinWeights, &InternalTargetMeshNormals, bUseSubset](int32 InVertexID)
 		{
 			if (Cancelled()) 
 			{
 				return;
 			}
-			
+
+			const int32 VertexID = bUseSubset ? TargetVerticesSubset[InVertexID] : InVertexID;
 			if (InOutTargetMesh.IsVertex(VertexID)) 
 			{
 				const FVector3d Point = InOutTargetMesh.GetVertex(VertexID);
@@ -294,7 +300,7 @@ bool FTransferBoneWeights::TransferWeightsToMesh(FDynamicMesh3& InOutTargetMesh,
 					NumMatched++;
 				}
 			}
-			bFailed = NumMatched != InOutTargetMesh.VertexCount();
+			bFailed = NumMatched != NumVerticesToTransfer;
 		}
 	} 
 	else if (TransferMethod == ETransferBoneWeightsMethod::InpaintWeights)
@@ -321,6 +327,8 @@ bool FTransferBoneWeights::TransferWeightsToMesh(FDynamicMesh3& InOutTargetMesh,
          *      Any vertex for which we found a match must have fixed weights that can't be changed, 
          *      i.e. W(i,j) = KnownWeights(i,j) where i is a vertex for which we found a match on the body.
 		 */
+	
+		MatchedVertices.Init(false, InOutTargetMesh.MaxVertexID());
 		
 		// Check if the target mesh contains the user specifed force inpaint weight map 
 		const FDynamicMeshWeightAttribute* ForceInpaintLayer = nullptr;
@@ -337,8 +345,17 @@ bool FTransferBoneWeights::TransferWeightsToMesh(FDynamicMesh3& InOutTargetMesh,
 			}
 		}
 
+		// because the inpaint algorithm can extract data from regions outside the target vertex subset, a temporary attribute profile is used to modify the weights.
+		// NOTE: make sure to copy the weights of the vertex subset into the complete TargetSkinWeights attribute before exciting the function. (see CopySubsetWeightsIfNeeded) 
+		FDynamicMeshVertexSkinWeightsAttribute SubsetTargetSkinWeights;
+		if (bUseSubset)
+		{
+			SubsetTargetSkinWeights.Copy(*TargetSkinWeights);
+		}
+		FDynamicMeshVertexSkinWeightsAttribute* EditedSkinWeights = bUseSubset ? &SubsetTargetSkinWeights : TargetSkinWeights;
+
 		// For every vertex on the target mesh try to find the match on the source mesh using the distance and normal checks
-		ParallelFor(InOutTargetMesh.MaxVertexID(), [this, &InOutTargetMesh, &ForceInpaintLayer, &TargetBoneToIndex, &TargetSkinWeights, &InternalTargetMeshNormals](int32 VertexID)
+		ParallelFor(InOutTargetMesh.MaxVertexID(), [this, &InOutTargetMesh, &ForceInpaintLayer, &TargetBoneToIndex, &EditedSkinWeights, &InternalTargetMeshNormals](int32 VertexID)
 		{
 			if (Cancelled()) 
 			{
@@ -376,7 +393,7 @@ bool FTransferBoneWeights::TransferWeightsToMesh(FDynamicMesh3& InOutTargetMesh,
 				FBoneWeights Weights;
 				if (TransferWeightsToPoint(Weights, Point, TargetBoneToIndex.Get(), Normal))
 				{
-					TargetSkinWeights->SetValue(VertexID, Weights);
+					EditedSkinWeights->SetValue(VertexID, Weights);
 					MatchedVertices[VertexID] = true;
 				}
 			}
@@ -387,24 +404,23 @@ bool FTransferBoneWeights::TransferWeightsToMesh(FDynamicMesh3& InOutTargetMesh,
 			return false;
 		}
 
-		// Compute linearization so we can store constraints at linearized indices
-		FVertexLinearization VtxLinearization(InOutTargetMesh, false);
-		const TArray<int32>& ToMeshV = VtxLinearization.ToId();
-		const TArray<int32>& ToIndex = VtxLinearization.ToIndex();
-		
 		int32 NumMatched = 0;
-		for (bool Flag : MatchedVertices)
+		if (!bUseSubset)
 		{
-			if (Flag) 
-			{ 
-				NumMatched++;
+			for (bool Flag : MatchedVertices)
+			{
+				if (Flag) 
+				{ 
+					NumMatched++;
+				}
 			}
 		}
-
-		// If all vertices were matched then nothing else to do
-		if (NumMatched == InOutTargetMesh.VertexCount())
+		else
 		{
-			return true;
+			NumMatched = (int32)Algo::CountIf(TargetVerticesSubset, [this](int32 VertexID)
+			{
+				return MatchedVertices.IsValidIndex(VertexID) && MatchedVertices[VertexID];
+			});
 		}
 
 		// If no vertices matched, we have nothing to inpaint.
@@ -413,6 +429,35 @@ bool FTransferBoneWeights::TransferWeightsToMesh(FDynamicMesh3& InOutTargetMesh,
 			return false;
 		}
 
+		auto CopySubsetWeightsIfNeeded = [Subset = TargetVerticesSubset, &EditedSkinWeights, &TargetSkinWeights, &InOutTargetMesh]()
+		{
+			if (EditedSkinWeights && EditedSkinWeights != TargetSkinWeights)
+			{
+				for (const int32 VertexID: Subset)
+				{
+					if (InOutTargetMesh.IsVertex(VertexID))
+					{
+						FBoneWeights Weights;
+						EditedSkinWeights->GetValue(VertexID, Weights);
+						TargetSkinWeights->SetValue(VertexID, Weights);
+					}
+				}
+			}
+		};
+		
+		// If all vertices were matched then nothing else to do
+		if (NumMatched == NumVerticesToTransfer)
+		{
+			// copy weights from the subset skin weights if using subset
+			CopySubsetWeightsIfNeeded();
+			return true;
+		}
+
+		// Compute linearization so we can store constraints at linearized indices
+		FVertexLinearization VtxLinearization(InOutTargetMesh, false);
+		const TArray<int32>& ToMeshV = VtxLinearization.ToId();
+		const TArray<int32>& ToIndex = VtxLinearization.ToIndex();
+		
 		// Setup the sparse matrix FixedValues of known (matched) weight values and the array (FixedIndices) of the matched vertex IDs
 		const int32 TargetNumBones = InOutTargetMesh.Attributes()->GetBoneNames()->GetAttribValues().Num();
 		FSparseMatrixD FixedValues;
@@ -428,7 +473,7 @@ bool FTransferBoneWeights::TransferWeightsToMesh(FDynamicMesh3& InOutTargetMesh,
 			if (InOutTargetMesh.IsVertex(VertexID) && MatchedVertices[VertexID])
 			{
 				FBoneWeights Data;
-				TargetSkinWeights->GetValue(VertexID, Data);
+				EditedSkinWeights->GetValue(VertexID, Data);
 
 				const int32 NumBones = Data.Num();
 				checkSlow(NumBones > 0);
@@ -519,9 +564,12 @@ bool FTransferBoneWeights::TransferWeightsToMesh(FDynamicMesh3& InOutTargetMesh,
 
 				const int32 VertexIDLinearalized = bVariablesOnly ? static_cast<int32>(VaribleRows[ColIdx]) : ColIdx; // linearized vertex ID (matrix row) of the variable in the Energy matrix
 				const int32 VertexID = static_cast<int32>(ToMeshV[VertexIDLinearalized]);
-				TargetSkinWeights->SetValue(VertexID, WeightArray);
+				EditedSkinWeights->SetValue(VertexID, WeightArray);
 			}
 
+			// copy weights from the subset skin weights if using subset
+			CopySubsetWeightsIfNeeded();
+			
 			// Optional post-processing smoothing of the weights at the vertices without a match
 			if (NumSmoothingIterations > 0 && SmoothingStrength > 0)
 			{
diff --git a/Engine/Plugins/Runtime/GeometryProcessing/Source/DynamicMesh/Public/Operations/TransferBoneWeights.h b/Engine/Plugins/Runtime/GeometryProcessing/Source/DynamicMesh/Public/Operations/TransferBoneWeights.h
--- a/Engine/Plugins/Runtime/GeometryProcessing/Source/DynamicMesh/Public/Operations/TransferBoneWeights.h
+++ b/Engine/Plugins/Runtime/GeometryProcessing/Source/DynamicMesh/Public/Operations/TransferBoneWeights.h
@@ -163,7 +163,7 @@ public:
 	 * automatically.
 	 * 
 	 * @note Only used when TransferMethod == ETransferBoneWeightsMethod::InpaintWeights.
-	 * 		 The size must be equal to the InSourceMesh.MaxVertexID(), otherwise the mask is ignored.
+	 * 		 The size must be equal to the InTargetMesh.MaxVertexID(), otherwise the mask is ignored.
 	 */
 	TArray<float> ForceInpaint;
 
@@ -173,6 +173,11 @@ public:
 	 */
 	FName ForceInpaintWeightMapName = NAME_None;
 
+	/** 
+	 * Optional subset of target mesh vertices to transfer weights to.
+	 * If left empty, skin weights will be transferred to all target mesh vertices.
+	 */
+	TArray<int32> TargetVerticesSubset;
 
 	//
 	// Outputs
diff --git a/Engine/Plugins/Runtime/GeometryScripting/Source/GeometryScriptingCore/Private/MeshBoneWeightFunctions.cpp b/Engine/Plugins/Runtime/GeometryScripting/Source/GeometryScriptingCore/Private/MeshBoneWeightFunctions.cpp
--- a/Engine/Plugins/Runtime/GeometryScripting/Source/GeometryScriptingCore/Private/MeshBoneWeightFunctions.cpp
+++ b/Engine/Plugins/Runtime/GeometryScripting/Source/GeometryScriptingCore/Private/MeshBoneWeightFunctions.cpp
@@ -364,6 +364,7 @@ UDynamicMesh* UGeometryScriptLibrary_MeshBoneWeightFunctions::TransferBoneWeight
 	UDynamicMesh* SourceMesh,
 	UDynamicMesh* TargetMesh,
 	FGeometryScriptTransferBoneWeightsOptions Options,
+	FGeometryScriptMeshSelection InSelection,
 	UGeometryScriptDebug* Debug)
 {
 	using namespace UE::Geometry;
@@ -396,9 +397,13 @@ UDynamicMesh* UGeometryScriptLibrary_MeshBoneWeightFunctions::TransferBoneWeight
 		TransferBoneWeights.TransferMethod = static_cast<FTransferBoneWeights::ETransferBoneWeightsMethod>(Options.TransferMethod);
 		TransferBoneWeights.bUseParallel = true;
 
-
 		TargetMesh->EditMesh([&](FDynamicMesh3& EditMesh)
 		{
+			if (!InSelection.IsEmpty())
+			{
+				InSelection.ConvertToMeshIndexArray(EditMesh, TransferBoneWeights.TargetVerticesSubset, EGeometryScriptIndexType::Vertex);
+			}
+			
 			if (!EditMesh.HasAttributes())
 			{
 				EditMesh.EnableAttributes();
diff --git a/Engine/Plugins/Runtime/GeometryScripting/Source/GeometryScriptingCore/Public/GeometryScript/MeshBoneWeightFunctions.h b/Engine/Plugins/Runtime/GeometryScripting/Source/GeometryScriptingCore/Public/GeometryScript/MeshBoneWeightFunctions.h
--- a/Engine/Plugins/Runtime/GeometryScripting/Source/GeometryScriptingCore/Public/GeometryScript/MeshBoneWeightFunctions.h
+++ b/Engine/Plugins/Runtime/GeometryScripting/Source/GeometryScriptingCore/Public/GeometryScript/MeshBoneWeightFunctions.h
@@ -3,6 +3,7 @@
 #pragma once
 
 #include "CoreMinimal.h"
+#include "GeometryScriptSelectionTypes.h"
 #include "Kismet/BlueprintFunctionLibrary.h"
 #include "GeometryScript/GeometryScriptTypes.h"
 #include "SkeletalMeshAttributes.h"
@@ -330,6 +331,8 @@ public:
 	 * @param SourceMesh The mesh we are transferring the weights from.
 	 * @param TargetMesh The mesh we are transferring the weights to.
 	 * @param Options The options to set for the transfer weight algorithm.
+	 * @param Selection Optional subset of target mesh vertices to transfer weights to.
+	 * If left empty, skin weights will be transferred to all target mesh vertices.
 	 */
 	UFUNCTION(BlueprintCallable, Category = "GeometryScript|MeshQueries|BoneWeights", meta = (ScriptMethod))
 	static UPARAM(DisplayName = "Target Mesh") UDynamicMesh* 
@@ -337,6 +340,7 @@ public:
 		UDynamicMesh* SourceMesh,
 		UDynamicMesh* TargetMesh,
 		FGeometryScriptTransferBoneWeightsOptions Options = FGeometryScriptTransferBoneWeightsOptions(),
+		FGeometryScriptMeshSelection Selection = FGeometryScriptMeshSelection(),
 		UGeometryScriptDebug* Debug = nullptr);
 
 	/**
