Geometry Processing.   Add the ability to fit b-splines to provided data points.

diff --git a/Engine/Plugins/Runtime/GeometryProcessing/Source/GeometryAlgorithms/Private/Curve/BSplineCurveFit.cpp b/Engine/Plugins/Runtime/GeometryProcessing/Source/GeometryAlgorithms/Private/Curve/BSplineCurveFit.cpp
new file mode 100644
--- /dev/null
+++ b/Engine/Plugins/Runtime/GeometryProcessing/Source/GeometryAlgorithms/Private/Curve/BSplineCurveFit.cpp
@@ -0,0 +1,149 @@
+// Copyright Epic Games, Inc. All Rights Reserved.
+
+#include "Curve/BSplineCurveFit.h"
+
+#include "ThirdParty/GTEngine/Mathematics/GteBSplineCurveFit.h"
+#include "ThirdParty/GTEngine/LowLevel/GteLogger.h"
+
+namespace UE::Geometry
+{
+
+namespace SplineCurveFitUtils
+{
+
+bool HasValidParameters(const int NumSamples, const int Degree, const int NumControlPoints)
+{
+	if (Degree < 1)
+	{
+		return false;
+	}
+
+	if (Degree >= NumControlPoints)
+	{
+		return false;
+	}
+
+	if (NumControlPoints > NumSamples)
+	{
+		return false;
+	}
+
+	return true;
+}
+	
+class ScopedListener : public gte::Logger::Listener
+{
+public:
+	typedef  gte::Logger::Listener  MyBase;
+	ScopedListener( int flags = MyBase::LISTEN_FOR_ALL) 
+	: MyBase(flags)
+	, bHasReport(false)
+	{
+		gte::Logger::Subscribe(this);
+	}
+
+	virtual ~ScopedListener()
+	{
+		gte::Logger::Unsubscribe(this);
+	}
+
+	bool HasReport() const { return bHasReport;}
+private:
+	virtual void Report(std::string const& message) override
+	{
+		bHasReport= true;
+	}
+	
+	bool bHasReport;
+};
+
+} }
+
+bool UE::Geometry::BSplineCurveFit(const TArray<FVector2f>& DataPoints, const int32 SplineDegree, const int32 NumControlPoints, TArray<FVector2f>& ControlPointsOut)
+{
+	typedef FVector2f VectorType;
+	bool bSuccess = false;
+
+	using FReal = float;
+	const int32 Dim = 2;
+	const int32 NumSamples = DataPoints.Num();
+
+	if (!SplineCurveFitUtils::HasValidParameters(NumSamples, SplineDegree, NumControlPoints))
+	{
+		return bSuccess;
+	}
+
+	TArray<float> Data;
+	Data.SetNumUninitialized(DataPoints.Num()* Dim);
+	for (int i = 0, I = DataPoints.Num(); i < I; ++i)
+	{
+		Data[Dim * i] = DataPoints[i].X;
+		Data[Dim * i + 1] = DataPoints[i].Y;
+	}
+
+	// listener will catch errors with the curve fit.
+	SplineCurveFitUtils::ScopedListener Listener;
+
+	// fit the curve
+	gte::BSplineCurveFit<float>  CurveFit(Dim, NumSamples, Data.GetData(), SplineDegree, NumControlPoints);
+
+
+	if (!Listener.HasReport())
+	{ 
+		bSuccess = true;
+		const float* ControlData = CurveFit.GetControlData();
+
+		ControlPointsOut.SetNumUninitialized(NumControlPoints);
+		for (int i = 0, I = NumControlPoints; i < I; ++i)
+		{
+			ControlPointsOut[i] = VectorType(ControlData[Dim * i] , ControlData[Dim * i + 1]);
+		}	
+	}
+
+	return bSuccess;
+}
+
+bool  UE::Geometry::BSplineCurveFit(const TArray<FVector3f>& DataPoints, const int32 SplineDegree, const int32 NumControlPoints, TArray<FVector3f>& ControlPointsOut)
+{
+	typedef FVector3f VectorType;
+
+	bool bSuccess = false;
+
+	using FReal = float;
+	const int32 Dim = 3;
+	const int32 NumSamples = DataPoints.Num();
+
+	if (!SplineCurveFitUtils::HasValidParameters(NumSamples, SplineDegree, NumControlPoints))
+	{
+		return bSuccess;
+	}
+
+	TArray<float> Data;
+	Data.SetNumUninitialized(DataPoints.Num() * Dim);
+	for (int i = 0, I = DataPoints.Num(); i < I; ++i)
+	{
+		Data[Dim * i] = DataPoints[i].X;
+		Data[Dim * i + 1] = DataPoints[i].Y;
+		Data[Dim * i + 2] = DataPoints[i].Z;
+	}
+
+	// listener will catch errors with the curve fit.
+	SplineCurveFitUtils::ScopedListener Listener;
+
+	gte::BSplineCurveFit<float>  CurveFit(Dim, NumSamples, Data.GetData(), SplineDegree, NumControlPoints);
+
+
+	if (!Listener.HasReport())
+	{ 
+		bSuccess = true;
+
+		const float* ControlData = CurveFit.GetControlData();
+
+		ControlPointsOut.SetNumUninitialized(NumControlPoints);
+		for (int i = 0, I = NumControlPoints; i < I; ++i)
+		{
+			ControlPointsOut[i] = VectorType(ControlData[Dim * i], ControlData[Dim * i + 1], ControlData[Dim * i + 2]);
+		}
+	}
+	return bSuccess;
+}
diff --git a/Engine/Plugins/Runtime/GeometryProcessing/Source/GeometryAlgorithms/Public/Curve/BSplineCurveFit.h b/Engine/Plugins/Runtime/GeometryProcessing/Source/GeometryAlgorithms/Public/Curve/BSplineCurveFit.h
new file mode 100644
--- /dev/null
+++ b/Engine/Plugins/Runtime/GeometryProcessing/Source/GeometryAlgorithms/Public/Curve/BSplineCurveFit.h
@@ -0,0 +1,48 @@
+// Copyright Epic Games, Inc. All Rights Reserved.
+
+#pragma once
+
+#include "CoreMinimal.h"
+
+namespace UE::Geometry
+{
+using namespace UE::Math;
+
+/**
+*  Compute the control points that result from a best fit of the provided 2d data points when using b-splines of the specified degree.
+* 
+*  This fit assumes the sample points are evenly spaced in 't' over the interval [0, 1].
+*        and the knot vector has unique knots at  j / (NumControlPoints - SplineDegree) for j in [0, (NumControlPoints - Degree)].
+*        the knots at 0 and 1 each have multiplicity SplineDegree + 1, while the internal knots just have multiplicity 1.
+*
+*  @param DataPoints       - an array of 2d points the b-spline will attempt to fit
+*  @param SplineDegree     - the degree of the underling b-spline.  Must be greater that 0.
+*  @param NumControlPoints - requested number of control points solved for.  
+*                            Note: requires SplineDegree < NumControlPoints <= DataPoints.Num()
+*  @param ControlPointsOut - on successful return, an array of 2d control points of the specified length ( on failure the array is unchanged ).
+*
+*  @return true on success. 
+* 
+*/
+bool GEOMETRYALGORITHMS_API BSplineCurveFit(const TArray<FVector2f>& DataPoints,  const int32 SplineDegree, const int32 NumControlPoints, TArray<FVector2f>& ControlPointsOut);
+
+
+/**
+*  Compute the control points that result from a best fit of the provided 3d data points when using b-splines of the specified degree.
+* 
+*  This fit assumes the sample points are evenly spaced in 't' over the interval [0, 1].
+*        and the knot vector has unique knots at  j / (NumControlPoints - SplineDegree) for j in [0, (NumControlPoints - Degree)].
+*        the knots at 0 and 1 each have multiplicity SplineDegree + 1, while the internal knots just have multiplicity 1.
+*
+*  @param DataPoints       - an array of 3d points the b-spline will attempt to fit
+*  @param SplineDegree     - the degree of the underling b-spline.  Must be greater that 0.
+*  @param NumControlPoints - requested number of control points solved for.
+*                            Note: requires SplineDegree < NumControlPoints <= DataPoints.Num()
+*  @param ControlPointsOut - on successful return, an array of 3d control points of the specified length  ( on failure the array is unchanged ).
+*
+*  @return true on success.
+* 
+*/
+bool GEOMETRYALGORITHMS_API BSplineCurveFit(const TArray<FVector3f>& DataPoints, const int32 SplineDegree, const int32 NumControlPoints, TArray<FVector3f>& ControlPointsOut);
+
+} // end namespace
\ No newline at end of file
