make selection system not crash if it fails to find a valid group edge ID

diff --git a/Engine/Plugins/Runtime/GeometryProcessing/Source/DynamicMesh/Private/Selections/GeometrySelectionUtil.cpp b/Engine/Plugins/Runtime/GeometryProcessing/Source/DynamicMesh/Private/Selections/GeometrySelectionUtil.cpp
--- a/Engine/Plugins/Runtime/GeometryProcessing/Source/DynamicMesh/Private/Selections/GeometrySelectionUtil.cpp
+++ b/Engine/Plugins/Runtime/GeometryProcessing/Source/DynamicMesh/Private/Selections/GeometrySelectionUtil.cpp
@@ -1355,8 +1355,10 @@ if ( ensure( MeshSelection.TopologyType == EGeometryTopologyType::Polygroup ) ==
 			if (Mesh.IsEdge(SeedEdgeID))
 			{
 				const int32 GroupEdgeID = GroupTopology->FindGroupEdgeID(SeedEdgeID);
-
-				ProcessGroupEdgeID(GroupEdgeID);
+				if (GroupEdgeID != INDEX_NONE)
+				{
+					ProcessGroupEdgeID(GroupEdgeID);
+				}
 			}
 		}
 	}
