make FGroupTopology have a convention that group IDs are not negative, rather than just selectively applying Max(0, GroupID) during topology init, to fix crashes on using a group topology that does have negative group IDs (e.g., when the underlying mesh does not have groups set)

diff --git a/Engine/Plugins/Runtime/GeometryProcessing/Source/DynamicMesh/Private/GroupTopology.cpp b/Engine/Plugins/Runtime/GeometryProcessing/Source/DynamicMesh/Private/GroupTopology.cpp
--- a/Engine/Plugins/Runtime/GeometryProcessing/Source/DynamicMesh/Private/GroupTopology.cpp
+++ b/Engine/Plugins/Runtime/GeometryProcessing/Source/DynamicMesh/Private/GroupTopology.cpp
@@ -67,7 +67,7 @@ bool FGroupTopology::RebuildTopology()
 	GroupFaceCounts.Init(0, MaxGroupID);
 	for (int tid : Mesh->TriangleIndicesItr())
 	{
-		int GroupID = FMath::Max(0, GetGroupID(tid));
+		int GroupID = GetGroupID(tid);
 		if (GroupIDToGroupIndexMap[GroupID] == -1)
 		{
 			FGroup NewGroup;
@@ -85,7 +85,7 @@ bool FGroupTopology::RebuildTopology()
 	// sort faces into groups
 	for (int tid : Mesh->TriangleIndicesItr())
 	{
-		int GroupID = FMath::Max(0, GetGroupID(tid));
+		int GroupID = GetGroupID(tid);
 		Groups[GroupIDToGroupIndexMap[GroupID]].Triangles.Add(tid);
 	}
 
diff --git a/Engine/Plugins/Runtime/GeometryProcessing/Source/DynamicMesh/Public/GroupTopology.h b/Engine/Plugins/Runtime/GeometryProcessing/Source/DynamicMesh/Public/GroupTopology.h
--- a/Engine/Plugins/Runtime/GeometryProcessing/Source/DynamicMesh/Public/GroupTopology.h
+++ b/Engine/Plugins/Runtime/GeometryProcessing/Source/DynamicMesh/Public/GroupTopology.h
@@ -205,11 +205,12 @@ public:
 	/** 
 	 * Adjacency of Per-Triangle integers are what define the triangle groups.
 	 * Override this function to provide an alternate group definition.
+	 * Note: By convention, group IDs should not be negative
 	 * @return group id integer for given TriangleID 
 	 */
 	virtual int GetGroupID(int TriangleID) const
 	{
-		return (GroupLayer != nullptr) ? GroupLayer->GetValue(TriangleID) : Mesh->GetTriangleGroup(TriangleID);
+		return FMath::Max(0, (GroupLayer != nullptr) ? GroupLayer->GetValue(TriangleID) : Mesh->GetTriangleGroup(TriangleID));
 	}
 
 
