fix bug with oriented box fitting where the matrix frame->quaternion conversion would give a bad result because the matrix had a flipped axis

diff --git a/Engine/Plugins/Runtime/GeometryProcessing/Source/DynamicMesh/Private/ShapeApproximation/ShapeDetection3.cpp b/Engine/Plugins/Runtime/GeometryProcessing/Source/DynamicMesh/Private/ShapeApproximation/ShapeDetection3.cpp
--- a/Engine/Plugins/Runtime/GeometryProcessing/Source/DynamicMesh/Private/ShapeApproximation/ShapeDetection3.cpp
+++ b/Engine/Plugins/Runtime/GeometryProcessing/Source/DynamicMesh/Private/ShapeApproximation/ShapeDetection3.cpp
@@ -232,6 +232,11 @@ bool UE::Geometry::IsBoxMesh(const FDynamicMesh3& Mesh, FOrientedBox3d& BoxOut,
 		FVector3d X = Normals[UniqueAxes[0]];
 		FVector3d Y = Normals[UniqueAxes[1]];
 		FVector3d Z = Normals[UniqueAxes[2]];
+		// make sure the frame vectors are oriented as expected for a rotation matrix
+		if ((X.Cross(Y)).Dot(Z) < 0)
+		{
+			Z = -Z;
+		}
 		// compute AABB in the frame of the box
 		FQuaterniond Rotation(FMatrix3d(X, Y, Z, false));
 		FMatrix3d UnorientRotation = Rotation.Inverse().ToRotationMatrix();
