MeshModelingTools: Add a UV Transfer tool to modeling mode, which transfers UVs from a mesh simplified using an "existing vertices" option to the original, higher resolution mesh.

diff --git a/Engine/Plugins/Editor/ModelingToolsEditorMode/Source/ModelingToolsEditorMode/Private/ModelingToolsEditorMode.cpp b/Engine/Plugins/Editor/ModelingToolsEditorMode/Source/ModelingToolsEditorMode/Private/ModelingToolsEditorMode.cpp
--- a/Engine/Plugins/Editor/ModelingToolsEditorMode/Source/ModelingToolsEditorMode/Private/ModelingToolsEditorMode.cpp
+++ b/Engine/Plugins/Editor/ModelingToolsEditorMode/Source/ModelingToolsEditorMode/Private/ModelingToolsEditorMode.cpp
@@ -71,6 +71,7 @@
 #include "TransformMeshesTool.h"
 #include "UVProjectionTool.h"
 #include "UVLayoutTool.h"
+#include "UVTransferTool.h"
 #include "EditMeshMaterialsTool.h"
 #include "AddPivotActorTool.h"
 #include "EditPivotTool.h"
@@ -723,6 +724,9 @@ void UModelingToolsEditorMode::Enter()
 	auto UVLayoutToolBuilder = NewObject<UUVLayoutToolBuilder>();
 	RegisterTool(ToolManagerCommands.BeginUVLayoutTool, TEXT("BeginUVLayoutTool"), UVLayoutToolBuilder);
 
+	auto UVTransferToolBuilder = NewObject<UUVTransferToolBuilder>();
+	RegisterTool(ToolManagerCommands.BeginUVTransferTool, TEXT("BeginUVTransferTool"), UVTransferToolBuilder);
+
 #if WITH_PROXYLOD
 	auto MergeMeshesToolBuilder = NewObject<UMergeMeshesToolBuilder>();
 	RegisterTool(ToolManagerCommands.BeginVoxelMergeTool, TEXT("BeginVoxelMergeTool"), MergeMeshesToolBuilder);
diff --git a/Engine/Plugins/Editor/ModelingToolsEditorMode/Source/ModelingToolsEditorMode/Private/ModelingToolsEditorModeStyle.cpp b/Engine/Plugins/Editor/ModelingToolsEditorMode/Source/ModelingToolsEditorMode/Private/ModelingToolsEditorModeStyle.cpp
--- a/Engine/Plugins/Editor/ModelingToolsEditorMode/Source/ModelingToolsEditorMode/Private/ModelingToolsEditorModeStyle.cpp
+++ b/Engine/Plugins/Editor/ModelingToolsEditorMode/Source/ModelingToolsEditorMode/Private/ModelingToolsEditorModeStyle.cpp
@@ -308,7 +308,9 @@ void FModelingToolsEditorModeStyle::Initialize()
 		StyleSet->Set("ModelingToolsManagerCommands.BeginTransformUVIslandsTool",            new IMAGE_PLUGIN_BRUSH("Icons/TransformUVs_40x",     Icon20x20));         
 		StyleSet->Set("ModelingToolsManagerCommands.BeginTransformUVIslandsTool.Small",      new IMAGE_PLUGIN_BRUSH("Icons/TransformUVs_40x",     Icon20x20));         
 		StyleSet->Set("ModelingToolsManagerCommands.BeginUVLayoutTool",                      new IMAGE_PLUGIN_BRUSH("Icons/UVLayout_40x",         Icon20x20));    
-		StyleSet->Set("ModelingToolsManagerCommands.BeginUVLayoutTool.Small",                new IMAGE_PLUGIN_BRUSH("Icons/UVLayout_40x",         Icon20x20));    
+		StyleSet->Set("ModelingToolsManagerCommands.BeginUVLayoutTool.Small",                new IMAGE_PLUGIN_BRUSH("Icons/UVLayout_40x",         Icon20x20));  
+		StyleSet->Set("ModelingToolsManagerCommands.BeginUVTransferTool",                    new FSlateImageBrush(StyleSet->RootToCoreContentDir(TEXT("../Editor/Slate/Icons/GeneralTools/Next_40x.png")), Icon20x20));
+		StyleSet->Set("ModelingToolsManagerCommands.BeginUVTransferTool.Small",              new FSlateImageBrush(StyleSet->RootToCoreContentDir(TEXT("../Editor/Slate/Icons/GeneralTools/Next_40x.png")), Icon20x20));
 		StyleSet->Set("ModelingToolsManagerCommands.LaunchUVEditor",                         new IMAGE_BRUSH_SVG("Icons/UVEditor", Icon20x20));
 		StyleSet->Set("ModelingToolsManagerCommands.LaunchUVEditor.Small",                   new IMAGE_BRUSH_SVG("Icons/UVEditor", Icon20x20));
 
diff --git a/Engine/Plugins/Editor/ModelingToolsEditorMode/Source/ModelingToolsEditorMode/Private/ModelingToolsEditorModeToolkit.cpp b/Engine/Plugins/Editor/ModelingToolsEditorMode/Source/ModelingToolsEditorMode/Private/ModelingToolsEditorModeToolkit.cpp
--- a/Engine/Plugins/Editor/ModelingToolsEditorMode/Source/ModelingToolsEditorMode/Private/ModelingToolsEditorModeToolkit.cpp
+++ b/Engine/Plugins/Editor/ModelingToolsEditorMode/Source/ModelingToolsEditorMode/Private/ModelingToolsEditorModeToolkit.cpp
@@ -408,7 +408,8 @@ void FModelingToolsEditorModeToolkit::RegisterPalettes()
 		Commands.BeginUVProjectionTool,
 		Commands.BeginUVSeamEditTool,
 		Commands.BeginTransformUVIslandsTool,
-		Commands.BeginUVLayoutTool
+		Commands.BeginUVLayoutTool,
+		Commands.BeginUVTransferTool
 	});
 	if (IModularFeatures::Get().IsModularFeatureAvailable(IUVEditorModularFeature::GetModularFeatureName()))
 	{
@@ -1968,6 +1969,7 @@ void FModelingToolsEditorModeToolkit::BuildToolPalette(FName PaletteIndex, class
 		ToolbarBuilder.AddToolBarButton(Commands.BeginUVSeamEditTool);
 		ToolbarBuilder.AddToolBarButton(Commands.BeginTransformUVIslandsTool);
 		ToolbarBuilder.AddToolBarButton(Commands.BeginUVLayoutTool);
+		ToolbarBuilder.AddToolBarButton(Commands.BeginUVTransferTool);
 
 		// Handle the inclusion of the optional UVEditor button if the UVEditor plugin has been found
 		if (IModularFeatures::Get().IsModularFeatureAvailable(IUVEditorModularFeature::GetModularFeatureName()))
diff --git a/Engine/Plugins/Editor/ModelingToolsEditorMode/Source/ModelingToolsEditorMode/Private/ModelingToolsManagerActions.cpp b/Engine/Plugins/Editor/ModelingToolsEditorMode/Source/ModelingToolsEditorMode/Private/ModelingToolsManagerActions.cpp
--- a/Engine/Plugins/Editor/ModelingToolsEditorMode/Source/ModelingToolsEditorMode/Private/ModelingToolsManagerActions.cpp
+++ b/Engine/Plugins/Editor/ModelingToolsEditorMode/Source/ModelingToolsEditorMode/Private/ModelingToolsManagerActions.cpp
@@ -257,6 +257,7 @@ void FModelingToolsManagerCommands::RegisterCommands()
 		REGISTER_MODELING_TOOL_COMMAND(BeginUVSeamEditTool, "SeamEd", "Add UV seams");
 		REGISTER_MODELING_TOOL_COMMAND(BeginTransformUVIslandsTool, "XFormUV", "Transform UV islands in UV space");
 		REGISTER_MODELING_TOOL_COMMAND(BeginUVLayoutTool, "Layout", "Transform and Repack existing UVs");
+		REGISTER_MODELING_TOOL_COMMAND(BeginUVTransferTool, "Transfer", "Transfer UVs from a simplified mesh to a higher resolution mesh");
 	}
 	else
 	{
@@ -265,6 +266,7 @@ void FModelingToolsManagerCommands::RegisterCommands()
 		REGISTER_MODELING_TOOL_COMMAND(BeginUVSeamEditTool, "Edit UV Seams", "Add UV seams");
 		REGISTER_MODELING_TOOL_COMMAND(BeginTransformUVIslandsTool, "Transform UVs", "Transform UV islands in UV space");
 		REGISTER_MODELING_TOOL_COMMAND(BeginUVLayoutTool, "Layout UVs", "Transform and Repack existing UVs");
+		REGISTER_MODELING_TOOL_COMMAND(BeginUVTransferTool, "Transfer UVs", "Transfer UVs from a simplified mesh to a higher resolution mesh");
 	}
 
 	// Baking
diff --git a/Engine/Plugins/Editor/ModelingToolsEditorMode/Source/ModelingToolsEditorMode/Public/ModelingToolsManagerActions.h b/Engine/Plugins/Editor/ModelingToolsEditorMode/Source/ModelingToolsEditorMode/Public/ModelingToolsManagerActions.h
--- a/Engine/Plugins/Editor/ModelingToolsEditorMode/Source/ModelingToolsEditorMode/Public/ModelingToolsManagerActions.h
+++ b/Engine/Plugins/Editor/ModelingToolsEditorMode/Source/ModelingToolsEditorMode/Public/ModelingToolsManagerActions.h
@@ -92,6 +92,7 @@ public:
 	TSharedPtr<FUICommandInfo> BeginRemoveOccludedTrianglesTool;
 	TSharedPtr<FUICommandInfo> BeginUVProjectionTool;
 	TSharedPtr<FUICommandInfo> BeginUVLayoutTool;
+	TSharedPtr<FUICommandInfo> BeginUVTransferTool;
 	TSharedPtr<FUICommandInfo> BeginPlaneCutTool;
 	TSharedPtr<FUICommandInfo> BeginMirrorTool;
 	TSharedPtr<FUICommandInfo> BeginHoleFillTool;
diff --git a/Engine/Plugins/Experimental/MeshModelingToolsetExp/Source/MeshModelingToolsExp/Private/UVTransferTool.cpp b/Engine/Plugins/Experimental/MeshModelingToolsetExp/Source/MeshModelingToolsExp/Private/UVTransferTool.cpp
new file mode 100644
--- /dev/null
+++ b/Engine/Plugins/Experimental/MeshModelingToolsetExp/Source/MeshModelingToolsExp/Private/UVTransferTool.cpp
@@ -0,0 +1,426 @@
+// Copyright Epic Games, Inc. All Rights Reserved.
+
+#include "UVTransferTool.h"
+
+#include "Drawing/MeshElementsVisualizer.h"
+#include "InteractiveToolManager.h"
+#include "MeshOpPreviewHelpers.h"
+#include "ModelingOperators.h" // FDynamicMeshOperator
+#include "ModelingToolTargetUtil.h"
+#include "Parameterization/UVTransfer.h"
+#include "Properties/MeshMaterialProperties.h"
+#include "Properties/MeshUVChannelProperties.h"
+#include "Selections/GeometrySelectionUtil.h" // EnumerateSelectionTriangles
+#include "TargetInterfaces/DynamicMeshCommitter.h"
+#include "TargetInterfaces/DynamicMeshProvider.h"
+#include "ToolContextInterfaces.h" // FToolBuilderState
+#include "ToolSetupUtil.h"
+#include "ToolTargetManager.h"
+
+#define LOCTEXT_NAMESPACE "UUVTransferTool"
+
+namespace UVTransferToolLocals
+{
+	using namespace UE::Geometry;
+
+	FString CacheIdentifier = TEXT("UVTransferTool");
+
+	class FTransferUVsOp : public FDynamicMeshOperator
+	{
+	public:
+		virtual ~FTransferUVsOp() {}
+
+		// inputs
+		TSharedPtr<const FDynamicMesh3> SourceMesh;
+		TSharedPtr<const FDynamicMesh3> DestinationMesh;
+		TOptional<TSet<int32>> SourceSelectionTids;
+		TOptional<TSet<int32>> DestinationSelectionTids;
+		int UVLayerIndex = 0;
+		bool bTransferSeamsOnly = false;
+		bool bClearExistingSeams = true;
+
+		double VertexSearchDistance = KINDA_SMALL_NUMBER;
+		double VertexSearchCellSize = VertexSearchDistance * 3.0;
+		double PathSimilarityWeight = 200.0;
+
+		void SetTransform(const FTransformSRT3d& Transform)
+		{
+			ResultTransform = Transform;
+		}
+
+		virtual void CalculateResult(FProgressCancel* Progress) override
+		{
+			if (Progress && Progress->Cancelled()) { return; }
+
+			if (!DestinationMesh || !SourceMesh
+				|| UVLayerIndex < 0 || UVLayerIndex > 7
+				|| !SourceMesh->HasAttributes()
+				|| SourceMesh->Attributes()->NumUVLayers() <= UVLayerIndex)
+			{
+				ensure(false);
+				return;
+			}
+			
+			ResultMesh->Copy(*DestinationMesh);
+
+			if (Progress && Progress->Cancelled()) { return; }
+
+			UE::Geometry::FDynamicMeshUVTransfer Transferer(SourceMesh.Get(), ResultMesh.Get(), UVLayerIndex);
+
+			Transferer.VertexSearchDistance = VertexSearchDistance;
+			Transferer.VertexSearchCellSize = VertexSearchCellSize;
+			Transferer.PathSimilarityWeight = PathSimilarityWeight;
+			Transferer.bClearExistingSeamsInDestination = bClearExistingSeams;
+
+			// TODO SELECTIONS: Once we support multi-mesh selections, we would uncomment these
+			//Transferer.SourceSelectionTids = SourceSelectionTids.GetPtrOrNull();
+			//Transferer.DestinationSelectionTids = DestinationSelectionTids.GetPtrOrNull();
+
+			if (bTransferSeamsOnly)
+			{
+				bTransferSucceeded = Transferer.TransferSeams(Progress);
+			}
+			else
+			{
+				bTransferSucceeded = Transferer.TransferSeamsAndUVs(Progress);
+			}
+			
+		}
+
+		bool bTransferSucceeded = false;
+	private:
+
+	};
+}
+
+void UUVTransferTool::OnShutdown(EToolShutdownType ShutdownType)
+{
+	using namespace UVTransferToolLocals;
+
+	Settings->SaveProperties(this);
+	DestinationMaterialSettings->SaveProperties(this, CacheIdentifier);
+	UVChannelProperties->SaveProperties(this);
+
+	SourceSeamVisualizer->Disconnect();
+	DestinationSeamVisualizer->Disconnect();
+
+	for (TObjectPtr<UToolTarget> Target : Targets)
+	{
+		UE::ToolTarget::ShowSourceObject(Target);
+	}
+
+	FDynamicMeshOpResult Result = DestinationPreview->Shutdown();
+	if (ShutdownType == EToolShutdownType::Accept)
+	{
+		GenerateAsset(Result);
+	}
+
+	SourcePreview->Disconnect();
+
+	Settings = nullptr;
+	UVChannelProperties = nullptr;
+	DestinationMaterialSettings = nullptr;
+	DestinationPreview = nullptr;
+	SourcePreview = nullptr;
+	SourceSeamVisualizer = nullptr;
+	DestinationSeamVisualizer = nullptr;
+	
+	for (int i = 0; i < 2; ++i)
+	{
+		Meshes[i].Reset();
+		SelectionTidSets[i].Reset();
+	}
+	
+}
+
+void UUVTransferTool::Setup()
+{
+	using namespace UVTransferToolLocals;
+
+	UInteractiveTool::Setup();
+
+	// Set up our property sets
+	Settings = NewObject<UUVTransferToolProperties>(this);
+	Settings->RestoreProperties(this);
+	AddToolPropertySource(Settings);
+
+	DestinationMaterialSettings = NewObject<UExistingMeshMaterialProperties>(this);
+	DestinationMaterialSettings->RestoreProperties(this, CacheIdentifier);
+	AddToolPropertySource(DestinationMaterialSettings);
+
+	UVChannelProperties = NewObject<UMeshUVChannelProperties>(this);
+	UVChannelProperties->RestoreProperties(this);
+	AddToolPropertySource(UVChannelProperties);
+	UVChannelProperties->ValidateSelection(true);
+
+	// Hide inputs
+	for (TObjectPtr<UToolTarget> Target : Targets)
+	{
+		UE::ToolTarget::HideSourceObject(Target);
+	}
+
+	// Extract the input meshes and selections
+	for (int i = 0; i < 2; ++i)
+	{
+		Meshes[i] = MakeShared<UE::Geometry::FDynamicMesh3>(UE::ToolTarget::GetDynamicMeshCopy(Targets[i]));
+		if (HasGeometrySelection(i))
+		{
+			SelectionTidSets[i].Emplace();
+			EnumerateSelectionTriangles(GetGeometrySelection(0), *Meshes[i], [this, i](int32 Tid)
+			{
+				SelectionTidSets[i]->Add(Tid);
+			});
+		}
+	}
+
+	// Set up previews
+	DestinationPreview = NewObject<UMeshOpPreviewWithBackgroundCompute>(this);
+	DestinationPreview->Setup(GetTargetWorld(), this);
+	DestinationPreview->OnOpCompleted.AddWeakLambda(this, [this](const UE::Geometry::FDynamicMeshOperator* UncastOp)
+	{
+		const FTransferUVsOp* Op = static_cast<const FTransferUVsOp*>(UncastOp);
+		if (!Op->bTransferSucceeded)
+		{
+			GetToolManager()->DisplayMessage(LOCTEXT("TransferUnsuccessful", 
+				"Transfer was not fully successful, possibly because correspondence couldn't be found. The "
+				"source mesh is expected to be a version of the destination mesh simplified via \"Existing Positions\" "
+				"option so that any UV layout done on the simplified mesh can be mapped back to the (original) destination "
+				"mesh via vertex positions."),
+				EToolMessageLevel::UserWarning);
+		}
+		else
+		{
+			GetToolManager()->DisplayMessage(FText(), EToolMessageLevel::UserWarning);
+		}
+	});
+	DestinationPreview->OnMeshUpdated.AddWeakLambda(this, [this](UMeshOpPreviewWithBackgroundCompute* Compute)
+	{
+		if (DestinationSeamVisualizer)
+		{
+			DestinationSeamVisualizer->NotifyMeshChanged();
+		}
+	});
+	
+	SourcePreview = NewObject<UPreviewMesh>(this);
+	SourcePreview->CreateInWorld(GetTargetWorld(), FTransform::Identity);
+	
+	ReinitializePreviews();
+
+	UVChannelProperties->WatchProperty(UVChannelProperties->UVChannel, [this](const FString& NewValue)
+	{
+		DestinationMaterialSettings->UpdateUVChannels(
+			UVChannelProperties->UVChannelNamesList.IndexOfByKey(UVChannelProperties->UVChannel),
+			UVChannelProperties->UVChannelNamesList);
+		InvalidatePreview();
+	});
+
+	SourceSeamVisualizer = NewObject<UMeshElementsVisualizer>(this);
+	SourceSeamVisualizer->CreateInWorld(GetWorld(), SourcePreview->GetTransform());
+	if (ensure(SourceSeamVisualizer->Settings))
+	{
+		SourceSeamVisualizer->Settings->ShowAllElements(false);
+	}
+	SourceSeamVisualizer->SetMeshAccessFunction([this](UMeshElementsVisualizer::ProcessDynamicMeshFunc ProcessFunc) {
+		SourcePreview->ProcessMesh(ProcessFunc);
+	});
+
+	DestinationSeamVisualizer = NewObject<UMeshElementsVisualizer>(this);
+	DestinationSeamVisualizer->CreateInWorld(GetWorld(), DestinationPreview->PreviewMesh->GetTransform());
+	if (ensure(DestinationSeamVisualizer->Settings))
+	{
+		DestinationSeamVisualizer->Settings->ShowAllElements(false);
+	}
+	DestinationSeamVisualizer->SetMeshAccessFunction([this](UMeshElementsVisualizer::ProcessDynamicMeshFunc ProcessFunc) {
+		DestinationPreview->ProcessCurrentMesh(ProcessFunc);
+	});
+	UpdateVisualizations();
+	
+	Settings->WatchProperty(Settings->bReverseDirection, [this](bool) 
+	{ 
+		ReinitializePreviews();
+		InvalidatePreview();
+	});
+	// TODO SELECTIONS: Once we support multi-mesh selections, we would add this bool to the settings
+	//Settings->WatchProperty(Settings->bRespectSelection, [this](bool) { InvalidatePreview(); });
+
+	Settings->WatchProperty(Settings->VertexSearchDistance, [this](double) { InvalidatePreview(); });
+	Settings->WatchProperty(Settings->PathSimilarityWeight, [this](double) { InvalidatePreview(); });
+	Settings->WatchProperty(Settings->bTransferSeamsOnly, [this](bool) { InvalidatePreview(); });
+	Settings->WatchProperty(Settings->bClearExistingSeams, [this](bool) { InvalidatePreview(); });
+	
+	Settings->WatchProperty(Settings->bShowWireframes, [this](bool) { UpdateVisualizations(); });
+	Settings->WatchProperty(Settings->bShowSeams, [this](bool) { UpdateVisualizations(); });
+
+	Settings->SilentUpdateWatched();
+	UVChannelProperties->SilentUpdateWatched();
+	DestinationMaterialSettings->SilentUpdateWatched();
+
+	SetToolDisplayName(LOCTEXT("ToolName", "UV Transfer"));
+	GetToolManager()->DisplayMessage(LOCTEXT("OnStart", "Transfer UVs from a low-res mesh to a high-res mesh."),
+		EToolMessageLevel::UserNotification);
+
+	InvalidatePreview();
+}
+
+void UUVTransferTool::ReinitializePreviews()
+{
+	int SourceIndex = Settings->bReverseDirection ? 1 : 0;
+	int DestinationIndex = Settings->bReverseDirection ? 0 : 1;
+
+	ToolSetupUtil::ApplyRenderingConfigurationToPreview(DestinationPreview->PreviewMesh, Targets[DestinationIndex]);
+	DestinationPreview->ConfigureMaterials(UE::ToolTarget::GetMaterialSet(Targets[DestinationIndex]).Materials,
+		ToolSetupUtil::GetDefaultWorkingMaterial(GetToolManager())
+	);
+	DestinationPreview->OverrideMaterial = DestinationMaterialSettings->GetActiveOverrideMaterial();
+	DestinationPreview->PreviewMesh->UpdatePreview(Meshes[DestinationIndex].Get());
+	DestinationPreview->PreviewMesh->SetTransform((FTransform)UE::ToolTarget::GetLocalToWorldTransform(Targets[DestinationIndex]));
+
+	ToolSetupUtil::ApplyRenderingConfigurationToPreview(SourcePreview, Targets[SourceIndex]);
+	SourcePreview->SetMaterials(UE::ToolTarget::GetMaterialSet(Targets[SourceIndex]).Materials);
+	SourcePreview->UpdatePreview(Meshes[SourceIndex].Get());
+	SourcePreview->SetTransform((FTransform)UE::ToolTarget::GetLocalToWorldTransform(Targets[SourceIndex]));
+
+	UVChannelProperties->Initialize(Meshes[SourceIndex].Get(), false);
+	UVChannelProperties->ValidateSelection(true);
+	DestinationMaterialSettings->UpdateUVChannels(
+		UVChannelProperties->GetSelectedChannelIndex(),
+		UVChannelProperties->UVChannelNamesList);
+
+	DestinationMaterialSettings->UpdateMaterials();
+	DestinationPreview->OverrideMaterial = DestinationMaterialSettings->GetActiveOverrideMaterial();
+
+	if (SourceSeamVisualizer)
+	{
+		SourceSeamVisualizer->SetTransform(SourcePreview->GetTransform());
+		SourceSeamVisualizer->NotifyMeshChanged();
+	}
+	if (DestinationSeamVisualizer)
+	{
+		DestinationSeamVisualizer->SetTransform(DestinationPreview->PreviewMesh->GetTransform());
+		DestinationSeamVisualizer->NotifyMeshChanged();
+	}
+}
+
+void UUVTransferTool::UpdateVisualizations()
+{
+	// The visualizers expect that we add their settings objects to our property sets, which
+	//  we don't want to do, so we need to do their updates ourselves.
+	if (ensure(SourceSeamVisualizer->Settings))
+	{
+		SourceSeamVisualizer->Settings->bShowWireframe = Settings->bShowWireframes;
+		SourceSeamVisualizer->Settings->bShowUVSeams = Settings->bShowSeams;
+		SourceSeamVisualizer->Settings->CheckAndUpdateWatched();
+	}
+	if (ensure(DestinationSeamVisualizer->Settings))
+	{
+		DestinationSeamVisualizer->Settings->bShowWireframe = Settings->bShowWireframes;
+		DestinationSeamVisualizer->Settings->bShowUVSeams = Settings->bShowSeams;
+		DestinationSeamVisualizer->Settings->CheckAndUpdateWatched();
+	}
+}
+
+void UUVTransferTool::OnTick(float DeltaTime)
+{
+	if (DestinationPreview)
+	{
+		DestinationPreview->Tick(DeltaTime);
+	}
+	if (DestinationSeamVisualizer)
+	{
+		DestinationSeamVisualizer->OnTick(DeltaTime);
+	}
+	if (SourceSeamVisualizer)
+	{
+		SourceSeamVisualizer->OnTick(DeltaTime);
+	}
+}
+
+void UUVTransferTool::Render(IToolsContextRenderAPI* RenderAPI)
+{
+}
+
+bool UUVTransferTool::CanAccept() const
+{
+	return DestinationPreview && DestinationPreview->HaveValidResult();
+}
+
+void UUVTransferTool::OnPropertyModified(UObject* PropertySet, FProperty* Property)
+{
+	if (PropertySet == DestinationMaterialSettings)
+	{
+		DestinationMaterialSettings->UpdateMaterials();
+		DestinationPreview->OverrideMaterial = DestinationMaterialSettings->GetActiveOverrideMaterial();
+	}
+}
+
+TUniquePtr<UE::Geometry::FDynamicMeshOperator> UUVTransferTool::MakeNewOperator()
+{
+	using namespace UVTransferToolLocals;
+	TUniquePtr<FTransferUVsOp> Op = MakeUnique<FTransferUVsOp>();
+
+	int SourceIndex = Settings->bReverseDirection ? 1 : 0;
+	int DestinationIndex = Settings->bReverseDirection ? 0 : 1;
+
+	Op->SourceMesh = Meshes[SourceIndex];
+	Op->DestinationMesh = Meshes[DestinationIndex];
+	Op->SetTransform(UE::ToolTarget::GetLocalToWorldTransform(Targets[DestinationIndex]));
+	
+	// TODO SELECTIONS: Once we support multi-mesh selections, we would uncomment these
+	//if (Settings->bRespectSelection)
+	//{
+	//	Op->SourceSelectionTids = SelectionTidSets[SourceIndex];
+	//	Op->DestinationSelectionTids = SelectionTidSets[DestinationIndex];
+	//}
+
+	Op->UVLayerIndex = UVChannelProperties->GetSelectedChannelIndex(true);
+	Op->bTransferSeamsOnly = Settings->bTransferSeamsOnly;
+	Op->bClearExistingSeams = Settings->bClearExistingSeams;
+
+	Op->VertexSearchDistance = Settings->VertexSearchDistance;
+	Op->VertexSearchCellSize = Settings->VertexSearchDistance;
+	Op->PathSimilarityWeight = Settings->PathSimilarityWeight;
+
+	return Op;
+}
+
+void UUVTransferTool::InvalidatePreview()
+{
+	if (DestinationPreview)
+	{
+		DestinationPreview->InvalidateResult();
+	}
+}
+
+void UUVTransferTool::GenerateAsset(const FDynamicMeshOpResult& Result)
+{
+	GetToolManager()->BeginUndoTransaction(LOCTEXT("UVLayoutToolTransactionName", "UV Layout Tool"));
+	UE::ToolTarget::CommitDynamicMeshUVUpdate(
+		Targets[Settings->bReverseDirection ? 0 : 1], 
+		Result.Mesh.Get());
+	GetToolManager()->EndUndoTransaction();
+}
+
+
+// Builder:
+
+UMultiTargetWithSelectionTool* UUVTransferToolBuilder::CreateNewTool(const FToolBuilderState& SceneState) const
+{
+	return NewObject<UUVTransferTool>(SceneState.ToolManager);
+}
+
+bool UUVTransferToolBuilder::CanBuildTool(const FToolBuilderState& SceneState) const
+{
+	const int32 NumTargets = SceneState.TargetManager->CountSelectedAndTargetable(SceneState, GetTargetRequirements());
+	return NumTargets == 2;
+}
+
+const FToolTargetTypeRequirements& UUVTransferToolBuilder::GetTargetRequirements() const
+{
+	static FToolTargetTypeRequirements TypeRequirements({
+		UDynamicMeshProvider::StaticClass(),
+		UDynamicMeshCommitter::StaticClass()
+		});
+	return TypeRequirements;
+}
+
+#undef LOCTEXT_NAMESPACE
diff --git a/Engine/Plugins/Experimental/MeshModelingToolsetExp/Source/MeshModelingToolsExp/Public/UVTransferTool.h b/Engine/Plugins/Experimental/MeshModelingToolsetExp/Source/MeshModelingToolsExp/Public/UVTransferTool.h
new file mode 100644
--- /dev/null
+++ b/Engine/Plugins/Experimental/MeshModelingToolsetExp/Source/MeshModelingToolsExp/Public/UVTransferTool.h
@@ -0,0 +1,154 @@
+// Copyright Epic Games, Inc. All Rights Reserved.
+
+#pragma once
+
+#include "BaseTools/MultiTargetWithSelectionTool.h" 
+#include "GeometryBase.h"
+#include "InteractiveToolBuilder.h"
+#include "Misc/Optional.h"
+#include "ModelingOperators.h" // IDynamicMeshOperatorFactory
+
+#include "UVTransferTool.generated.h"
+
+struct FDynamicMeshOpResult;
+class UMeshUVChannelProperties;
+class UExistingMeshMaterialProperties;
+class UMeshElementsVisualizer;
+class UMeshOpPreviewWithBackgroundCompute;
+class UPreviewMesh;
+class UUVLayoutPreview;
+
+PREDECLARE_GEOMETRY(class FDynamicMesh3);
+
+
+UCLASS()
+class MESHMODELINGTOOLSEXP_API UUVTransferToolBuilder : public UMultiTargetWithSelectionToolBuilder
+{
+	GENERATED_BODY()
+
+public:
+	virtual UMultiTargetWithSelectionTool* CreateNewTool(const FToolBuilderState& SceneState) const override;
+	virtual bool CanBuildTool(const FToolBuilderState& SceneState) const override;
+	virtual bool RequiresInputSelection() const override { return false; }
+
+protected:
+	virtual const FToolTargetTypeRequirements& GetTargetRequirements() const override;
+};
+
+UCLASS()
+class MESHMODELINGTOOLSEXP_API UUVTransferToolProperties : public UInteractiveToolPropertySet
+{
+	GENERATED_BODY()
+
+public:
+
+	// If false, the first selected mesh's UV's are applied to the second selected mesh. If true, the
+	//  reverse direction is used.
+	UPROPERTY(EditAnywhere, Category = Options)
+	bool bReverseDirection = false;
+
+	/**
+	 * If true, we only transfer the seams without trying to transfer actual UV element values.
+	 */
+	UPROPERTY(EditAnywhere, Category = Options)
+	bool bTransferSeamsOnly = false;
+
+	/**
+	 * If true, clears existing seams on the destination mesh before carrying over new ones.
+	 */
+	UPROPERTY(EditAnywhere, Category = Options)
+	bool bClearExistingSeams = true;
+
+	/**
+	 * Setting this above 0 will include a measure of path similarity to seam transfer, so that among
+	 *  similarly short paths, we pick one that lies closer to the edge. Useful in cases where the path
+	 *  is on the wrong diagonal to the triangulation, because it prefers a closely zigzagging path over
+	 *  a wider "up and over" path that has similar length. If set to 0, only path length is used.
+	 */
+	UPROPERTY(EditAnywhere, Category = Options, meta = (
+		ClampMin = 0, UIMax = 1000))
+	double PathSimilarityWeight = 200;
+
+	UPROPERTY(EditAnywhere, Category = Visualization)
+	bool bShowWireframes = false;
+
+	UPROPERTY(EditAnywhere, Category = Visualization);
+	bool bShowSeams = true;
+	/**
+	 * How far to look for a corresponding vertex on the destination. The destination is expected to
+	 *  be a simplified version of source using existing vertices, so this should not need to be set high.
+	 */
+	UPROPERTY(EditAnywhere, Category = Options, AdvancedDisplay, meta = (
+		ClampMin = 0.0001, ClampMax = 10000, UIMax = 1.0))
+	double VertexSearchDistance = 0.0001;
+};
+
+
+/**
+ * Tool that transfers UV data from a lower res mesh to a higher one. The lower resolution mesh is typically
+ *  obtained by simplifying the destination mesh with a "use existing vertices" setting so that an easy
+ *  correspondence between mesh vertices can be found.
+ */
+UCLASS()
+class MESHMODELINGTOOLSEXP_API UUVTransferTool : public UMultiTargetWithSelectionTool
+	, public UE::Geometry::IDynamicMeshOperatorFactory
+	, public IInteractiveToolManageGeometrySelectionAPI
+{
+	GENERATED_BODY()
+
+public:
+
+	// UMultiSelectionMeshEditingTool
+	virtual void OnShutdown(EToolShutdownType ShutdownType) override;
+
+	// UInteractiveTool
+	virtual void Setup() override;
+	virtual void OnTick(float DeltaTime) override;
+	virtual void Render(IToolsContextRenderAPI* RenderAPI) override;
+	virtual bool HasCancel() const override { return true; }
+	virtual bool HasAccept() const override { return true; }
+	virtual bool CanAccept() const override;
+	virtual void OnPropertyModified(UObject* PropertySet, FProperty* Property) override;
+
+	// IDynamicMeshOperatorFactory
+	virtual TUniquePtr<UE::Geometry::FDynamicMeshOperator> MakeNewOperator() override;
+
+	// IInteractiveToolManageGeometrySelectionAPI -- this tool won't update external geometry selection or change selection-relevant mesh IDs
+	virtual bool IsInputSelectionValidOnOutput() override
+	{
+		return true;
+	}
+
+private:
+
+	UPROPERTY()
+	TObjectPtr<UUVTransferToolProperties> Settings = nullptr;
+
+	UPROPERTY()
+	TObjectPtr<UMeshUVChannelProperties> UVChannelProperties = nullptr;
+
+	UPROPERTY()
+	TObjectPtr<UExistingMeshMaterialProperties> DestinationMaterialSettings = nullptr;
+
+	UPROPERTY()
+	TObjectPtr<UMeshOpPreviewWithBackgroundCompute> DestinationPreview = nullptr;
+
+	UPROPERTY()
+	TObjectPtr<UPreviewMesh> SourcePreview = nullptr;
+
+	UPROPERTY()
+	TObjectPtr<UMeshElementsVisualizer> SourceSeamVisualizer = nullptr;
+
+	UPROPERTY()
+	TObjectPtr<UMeshElementsVisualizer> DestinationSeamVisualizer = nullptr;
+
+	// Normally, Mesh1 corresponds to source and Mesh2 is destination, but this is reversed if bReverseDirection is true.
+	TSharedPtr<UE::Geometry::FDynamicMesh3> Meshes[2];
+	TOptional<TSet<int32>> SelectionTidSets[2];
+	
+	void ReinitializePreviews();
+	void UpdateVisualizations();
+	
+	void InvalidatePreview();
+	void GenerateAsset(const FDynamicMeshOpResult& Result);
+};
diff --git a/Engine/Plugins/Runtime/GeometryProcessing/Source/DynamicMesh/Private/Parameterization/DynamicMeshUVEditor.cpp b/Engine/Plugins/Runtime/GeometryProcessing/Source/DynamicMesh/Private/Parameterization/DynamicMeshUVEditor.cpp
--- a/Engine/Plugins/Runtime/GeometryProcessing/Source/DynamicMesh/Private/Parameterization/DynamicMeshUVEditor.cpp
+++ b/Engine/Plugins/Runtime/GeometryProcessing/Source/DynamicMesh/Private/Parameterization/DynamicMeshUVEditor.cpp
@@ -552,21 +552,42 @@ bool FDynamicMeshUVEditor::SetTriangleUVsFromExpMap(
 
 bool FDynamicMeshUVEditor::SetTriangleUVsFromFreeBoundaryConformal(const TArray<int32>& Triangles, FUVEditResult* Result)
 {
-	return SetTriangleUVsFromFreeBoundaryConformal(Triangles, false, Result);
+	return SetTriangleUVsFromFreeBoundaryConformal(Triangles, false,  Result);
 }
 
 bool FDynamicMeshUVEditor::SetTriangleUVsFromFreeBoundaryConformal(const TArray<int32>& Triangles, bool bUseExistingUVTopology, FUVEditResult* Result) 
 {
-	return SetTriangleUVsFromConformal(Triangles, bUseExistingUVTopology, false, false, Result);
+	FSetUVsFromConformalOptions Options;
+	Options.bUseExistingUVTopology = bUseExistingUVTopology;
+	Options.bUseSpectral = false;
+	Options.bPreserveIrregularity = false;
+	return SetTriangleUVsFromConformal(Triangles, Options, Result);
+}
+
+bool UE::Geometry::FDynamicMeshUVEditor::SetTriangleUVsFromFreeBoundaryConformal(const TArray<int32>& Triangles, 
+	const TSet<int32>& PinnedElementIDs, FUVEditResult* Result)
+{
+	FSetUVsFromConformalOptions Options;
+	Options.bUseExistingUVTopology = true;
+	Options.bUseSpectral = false;
+	Options.bPreserveIrregularity = false;
+	Options.PinnedElementIDs = &PinnedElementIDs;
+	return SetTriangleUVsFromConformal(Triangles, Options, Result);
 }
 
 bool FDynamicMeshUVEditor::SetTriangleUVsFromFreeBoundarySpectralConformal(const TArray<int32>& Triangles, bool bUseExistingUVTopology, bool bPreserveIrregularity, FUVEditResult* Result) 
 {
-	return SetTriangleUVsFromConformal(Triangles, bUseExistingUVTopology, true, bPreserveIrregularity, Result);
+	FSetUVsFromConformalOptions Options;
+	Options.bUseExistingUVTopology = bUseExistingUVTopology;
+	Options.bUseSpectral = true;
+	Options.bPreserveIrregularity = bPreserveIrregularity;
+	return SetTriangleUVsFromConformal(Triangles, Options, Result);
 }
 
-bool FDynamicMeshUVEditor::SetTriangleUVsFromConformal(const TArray<int32>& Triangles, bool bUseExistingUVTopology, bool bUseSpectral, bool bPreserveIrregularity, FUVEditResult* Result)
+bool FDynamicMeshUVEditor::SetTriangleUVsFromConformal(const TArray<int32>& Triangles, const FSetUVsFromConformalOptions& Options, FUVEditResult* Result)
 {
+	bool bUseExistingUVTopology = Options.bUseExistingUVTopology;
+
 	if (ensure(UVOverlay) == false) return false;
 	if (Triangles.Num() == 0) return false;
 
@@ -624,9 +645,12 @@ bool FDynamicMeshUVEditor::SetTriangleUVsFromConformal(const TArray<int32>& Tria
 	const TArray<int32>& ConstrainLoop = Loops[LongestLoopIndex].Vertices;
 	int32 LoopNum = ConstrainLoop.Num();
 
-	if (bUseSpectral) 
+	// Potentially used in the non-spectral case
+	TOptional<TPair<int32, FVector2f>> SinglePinnedElement;
+
+	if (Options.bUseSpectral) 
 	{
-		Solver = UE::MeshDeformation::ConstructSpectralConformalParamSolver(Submesh, bPreserveIrregularity);
+		Solver = UE::MeshDeformation::ConstructSpectralConformalParamSolver(Submesh, Options.bPreserveIrregularity);
 		
 		for (int32 Idx = 0; Idx < LoopNum; ++Idx)
 		{
@@ -639,32 +663,81 @@ bool FDynamicMeshUVEditor::SetTriangleUVsFromConformal(const TArray<int32>& Tria
 	{
 		Solver = UE::MeshDeformation::ConstructNaturalConformalParamSolver(Submesh);
 
-		// Find a pair of vertices to constrain. The standard procedure is to find the two furthest-apart vertices 
-		// on the largest boundary loop. 
-		FIndex2i MaxDistPair = FIndex2i::Invalid();
-		double MaxDistSqr = 0;
-		for (int32 Idx = 0; Idx < LoopNum; ++Idx)
+		// There are three options for constraints.
+		// 1. No pinned elements: the standard thing is to constrain the two furthest boundary vertices (this is supposed
+		//   to be a geodesic distance, but for now we just do euclidean).
+		// 2. One pinned element: we can do the same thing as 1, except translate afterward to put the pinned element 
+		//   in the desired coordinate.
+		// 3. More than one pinned element: constrain all the pinned elements.
+
+		int32 PinnedElementCount = 0;
+
+		if (Options.PinnedElementIDs)
 		{
-			for (int32 NextIdx = Idx + 1; NextIdx < LoopNum; ++NextIdx)
+			for (int32 ElementID : *Options.PinnedElementIDs)
 			{
-				const double DistSqr = DistanceSquared(Submesh.GetVertex(ConstrainLoop[Idx]), 
-													   Submesh.GetVertex(ConstrainLoop[NextIdx]));
-				if (DistSqr > MaxDistSqr)
+				if (!UVOverlay->IsElement(ElementID))
+				{
+					continue;
+				}
+				int32 BaseVert = bUseExistingUVTopology ? ElementID : UVOverlay->GetParentVertex(ElementID);
+				int32* SubmeshVert = BaseToSubmeshV.Find(BaseVert);
+				if (!SubmeshVert)
 				{
-					MaxDistSqr = DistSqr;
-					MaxDistPair = FIndex2i(ConstrainLoop[Idx], ConstrainLoop[NextIdx]);
+					continue;
+				}
+
+				Solver->AddConstraint(*SubmeshVert, 1.0, FVector2d(UVOverlay->GetElement(ElementID)), false);
+
+				++PinnedElementCount;
+
+				if (!SinglePinnedElement.IsSet())
+				{
+					// We'll clear this later if we have more than one element
+					SinglePinnedElement.Emplace(*SubmeshVert, UVOverlay->GetElement(ElementID));
 				}
 			}
 		}
 
-		if (ensure(MaxDistPair != FIndex2i::Invalid()) == false)
+		if (PinnedElementCount > 1)
 		{
-			return false;
+			// We don't want to trigger our whole island translation code further below
+			SinglePinnedElement.Reset();
 		}
 
-		// pin those vertices
-		Solver->AddConstraint(MaxDistPair.A, 1.0, FVector2d(0.0, 0.5), false);
-		Solver->AddConstraint(MaxDistPair.B, 1.0, FVector2d(1.0, 0.5), false);
+		// Pick our constraints if we have fewer than 2.
+		if (PinnedElementCount < 2)
+		{
+			// In case we had one pinned element, we'll constrain it with our own translation afterward
+			Solver->ClearConstraints(); 
+			
+			// Find a pair of vertices to constrain. The standard procedure is to find the two furthest-apart vertices 
+			// on the largest boundary loop. 
+			FIndex2i MaxDistPair = FIndex2i::Invalid();
+			double MaxDistSqr = 0;
+			for (int32 Idx = 0; Idx < LoopNum; ++Idx)
+			{
+				for (int32 NextIdx = Idx + 1; NextIdx < LoopNum; ++NextIdx)
+				{
+					const double DistSqr = DistanceSquared(Submesh.GetVertex(ConstrainLoop[Idx]), 
+														   Submesh.GetVertex(ConstrainLoop[NextIdx]));
+					if (DistSqr > MaxDistSqr)
+					{
+						MaxDistSqr = DistSqr;
+						MaxDistPair = FIndex2i(ConstrainLoop[Idx], ConstrainLoop[NextIdx]);
+					}
+				}
+			}
+
+			if (ensure(MaxDistPair != FIndex2i::Invalid()) == false)
+			{
+				return false;
+			}
+
+			// pin those vertices
+			Solver->AddConstraint(MaxDistPair.A, 1.0, FVector2d(0.0, 0.5), false);
+			Solver->AddConstraint(MaxDistPair.B, 1.0, FVector2d(1.0, 0.5), false);
+		}
 	}
 
 	// solve for UVs
@@ -674,6 +747,19 @@ bool FDynamicMeshUVEditor::SetTriangleUVsFromConformal(const TArray<int32>& Tria
 		return false;
 	}
 
+	// Handle the single-constrained-element case for the natural conformal solver
+	if (SinglePinnedElement.IsSet())
+	{
+		FVector2d DeltaToApply = FVector2d(SinglePinnedElement->Value) - UVBuffer[SinglePinnedElement->Key];
+		if (!DeltaToApply.IsZero())
+		{
+			int32 NumSubVerts = SubmeshToBaseV.Num();
+			for (int32 k = 0; k < NumSubVerts; ++k)
+			{
+				UVBuffer[k] += DeltaToApply;
+			}
+		}
+	}
 
 	int32 NumFailed = 0;
 	if (bUseExistingUVTopology)
diff --git a/Engine/Plugins/Runtime/GeometryProcessing/Source/DynamicMesh/Private/Parameterization/UVTransfer.cpp b/Engine/Plugins/Runtime/GeometryProcessing/Source/DynamicMesh/Private/Parameterization/UVTransfer.cpp
new file mode 100644
--- /dev/null
+++ b/Engine/Plugins/Runtime/GeometryProcessing/Source/DynamicMesh/Private/Parameterization/UVTransfer.cpp
@@ -0,0 +1,546 @@
+// Copyright Epic Games, Inc. All Rights Reserved.
+
+#include "Parameterization/UVTransfer.h"
+
+#include "Distance/DistLine3Segment3.h"
+#include "DynamicMesh/DynamicMesh3.h"
+#include "EdgeSpan.h"
+#include "Parameterization/DynamicMeshUVEditor.h"
+#include "Parameterization/MeshDijkstra.h"
+#include "Selections/MeshConnectedComponents.h"
+#include "Util/ProgressCancel.h"
+
+UE::Geometry::FDynamicMeshUVTransfer::FDynamicMeshUVTransfer(
+	const FDynamicMesh3* SourceMeshIn, FDynamicMesh3* DestinationMeshIn, int32 UVLayerIndexIn)
+: SourceMesh(SourceMeshIn)
+, DestinationMesh(DestinationMeshIn)
+, UVLayerIndex(FMath::Clamp(UVLayerIndexIn, 0, 7))
+{}
+
+bool UE::Geometry::FDynamicMeshUVTransfer::TransferSeams(FProgressCancel* Progress)
+{
+	if (!DestinationMesh || !SourceMesh
+		|| UVLayerIndex < 0 || UVLayerIndex > 7
+		|| !SourceMesh->HasAttributes()
+		|| SourceMesh->Attributes()->NumUVLayers() <= UVLayerIndex)
+	{
+		return ensure(false);
+	}
+
+	SourceOverlay = SourceMesh->Attributes()->GetUVLayer(UVLayerIndex);
+	if (!ensure(SourceOverlay))
+	{
+		return false;
+	}
+
+	DestOverlay = DestinationMesh->Attributes()->GetUVLayer(UVLayerIndex);
+	if (!ensure(DestOverlay))
+	{
+		return false;
+	}
+
+	if (Progress && Progress->Cancelled()) { return false; }
+
+	InitializeHashGrid();
+
+	if (Progress && Progress->Cancelled()) { return false; }
+
+	if (bClearExistingSeamsInDestination)
+	{
+		// Start with a UV topology that matches the mesh topology, i.e. remove any existing seams in the
+		//  region of interest.
+		ResetDestinationUVTopology(Progress);
+	}
+
+	if (Progress && Progress->Cancelled()) { return false; }
+
+	return PerformSeamTransfer(Progress);
+}
+
+bool  UE::Geometry::FDynamicMeshUVTransfer::TransferSeamsAndUVs(FProgressCancel* Progress)
+{
+	if (Progress && Progress->Cancelled()) { return false; }
+
+	bool bSeamSuccess = TransferSeams(Progress);
+
+	if (Progress && Progress->Cancelled()) { return false; }
+
+	// Even if seam transfer wasn't fully successful, we'll try to transfer what we can.
+	return PerformElementsTransfer(Progress) && bSeamSuccess;
+}
+
+void UE::Geometry::FDynamicMeshUVTransfer::InitializeHashGrid()
+{
+	HashGrid = MakeUnique<TPointHashGrid3<int32, double>>(
+		FMath::Max(KINDA_SMALL_NUMBER, VertexSearchCellSize), IndexConstants::InvalidID);
+	for (int32 Vid : DestinationMesh->VertexIndicesItr())
+	{
+		HashGrid->InsertPointUnsafe(Vid, DestinationMesh->GetVertex(Vid));
+	}
+}
+
+int32 UE::Geometry::FDynamicMeshUVTransfer::GetCorrespondingDestVid(int32 SourceVid)
+{
+	int32* FoundDestVid = SourceVidToDestinationVid.Find(SourceVid);
+	if (FoundDestVid)
+	{
+		return *FoundDestVid;
+	}
+
+	FVector3d VertPosition = SourceMesh->GetVertex(SourceVid);
+	int32 DestVid = HashGrid->FindNearestInRadius(VertPosition, VertexSearchDistance, [this, &VertPosition](int32 DestVid)
+	{
+		return FVector3d::DistSquared(VertPosition, DestinationMesh->GetVertex(DestVid));
+	}).Key;
+
+	if (DestVid != IndexConstants::InvalidID)
+	{
+		SourceVidToDestinationVid.Add(SourceVid, DestVid);
+	}
+
+	return DestVid;
+};
+
+void UE::Geometry::FDynamicMeshUVTransfer::ResetDestinationUVTopology(FProgressCancel* Progress)
+{
+	// To remove any seams in the region of interest, we just make sure that every vid is
+	//  only matched to a single element.
+	TMap<int32, int32> MeshVidToElementID;
+
+	auto ProcessTriangle = [this, &MeshVidToElementID](int32 Tid)
+	{
+		FIndex3i MeshTriangle = DestinationMesh->GetTriangle(Tid);
+		FIndex3i CurrentUVTriangle = DestOverlay->GetTriangle(Tid);
+		FIndex3i NewUVTriangle;
+		for (int32 SubIdx = 0; SubIdx < 3; ++SubIdx)
+		{
+			int32 Vid = MeshTriangle[SubIdx];
+			int32* ExistingMatchedElement = MeshVidToElementID.Find(Vid);
+			if (ExistingMatchedElement)
+			{
+				NewUVTriangle[SubIdx] = *ExistingMatchedElement;
+			}
+			else
+			{
+				int32 ElementID = CurrentUVTriangle[SubIdx] >= 0 ? CurrentUVTriangle[SubIdx]
+					// The value we choose if we didn't have the UVs set doesn't matter, since we'll
+					//  be changing it anyway.
+					: DestOverlay->AppendElement(FVector2f::Zero());
+				MeshVidToElementID.Add(Vid, ElementID);
+				NewUVTriangle[SubIdx] = ElementID;
+			}
+		}
+		DestOverlay->SetTriangle(Tid, NewUVTriangle);
+	};
+
+	// We'll check Progress->Cancelled every 2^10=1024 tris.
+	uint32 Count = 0;
+	uint32 IntervalCheckMask = (1 << 10) - 1;
+
+	if (DestinationSelectionTids)
+	{
+		for (int32 Tid : *DestinationSelectionTids)
+		{
+			if (Progress && ((++Count & IntervalCheckMask) == 0) && Progress->Cancelled()) { return; }
+			ProcessTriangle(Tid);
+		}
+	}
+	else
+	{
+		for (int32 Tid : DestinationMesh->TriangleIndicesItr())
+		{
+			if (Progress && ((++Count & IntervalCheckMask) == 0) && Progress->Cancelled()) { return; }
+			ProcessTriangle(Tid);
+		}
+	}
+}//end ResetDestinationUVTopology
+
+bool UE::Geometry::FDynamicMeshUVTransfer::PerformSeamTransfer(FProgressCancel* Progress)
+{
+	bool bAllSuccessful = true;
+
+	if (!HashGrid.IsValid())
+	{
+		InitializeHashGrid();
+	}
+
+	TSet<int32> DestEidsToMakeSeams;
+	SourceBoundaryElements.Reset();
+
+	double MeshMaxDim = SourceMesh->GetBounds(true).MaxDim();
+
+	auto ProcessSourceEdge = [this, &DestEidsToMakeSeams, MeshMaxDim](int32 SourceEid) -> bool
+	{
+		if (!SourceOverlay->IsSeamEdge(SourceEid))
+		{
+			return true;
+		}
+
+		FIndex2i SourceVids = SourceMesh->GetEdgeV(SourceEid);
+		FIndex2i DestVids;
+		DestVids.A = GetCorrespondingDestVid(SourceVids.A);
+		if (DestVids.A == IndexConstants::InvalidID)
+		{
+			return false;
+		}
+		DestVids.B = GetCorrespondingDestVid(SourceVids.B);
+		if (DestVids.B == IndexConstants::InvalidID)
+		{
+			return false;
+		}
+
+		if (SourceMesh->IsBoundaryEdge(SourceEid))
+		{
+			// We don't try to transfer boundary edges because we don't want to accidentally add a seam close
+			//  to (but not exactly on) the destination mesh boundary just because the simplified mesh boundary
+			//  passed here. However, we do need to mark the source boundary elements at some point because that
+			//  knowledge is useful for tranferring any seams that touch the boundary.
+			int32 Tid = SourceMesh->GetEdgeT(SourceEid).A;
+			FIndex3i TriVids = SourceMesh->GetTriangle(Tid);
+			FIndex3i TriElements = SourceOverlay->GetTriangle(Tid);
+			for (int i = 0; i < 2; ++i)
+			{
+				int SubIdx = TriVids.IndexOf(SourceVids[i]);
+				if (ensure(SubIdx >= 0) && SourceOverlay->IsElement(TriElements[SubIdx]))
+				{
+					SourceBoundaryElements.Add(TriElements[SubIdx]);
+				}
+			}
+			return true;
+		}
+
+		FVector SourceEdgeStart = SourceMesh->GetVertex(SourceVids.A);
+		FVector SourceEdgeVector = SourceMesh->GetVertex(SourceVids.B) - SourceEdgeStart;
+		double SourceEdgeLength = 0;
+		FVector SourceEdgeDirection = FVector::Zero();
+		SourceEdgeVector.ToDirectionAndLength(SourceEdgeDirection, SourceEdgeLength);
+
+		UE::Geometry::TMeshDijkstra<FDynamicMesh3> PathFinder(DestinationMesh);
+		TMeshDijkstra<FDynamicMesh3>::FSeedPoint SeedPoint = { DestVids.A, DestVids.A, 0 };
+		
+		double MaxDistance = TNumericLimits<double>::Max();
+		if (PathLengthToleranceMultiplier > 1)
+		{
+			MaxDistance = PathLengthToleranceMultiplier * SourceEdgeLength;
+			if (PathSimilarityWeight > 0)
+			{
+				MaxDistance *= (1 + PathSimilarityWeight);
+			}
+			// Arbitrary clamping to prevent failures to find vertices when source edge is tiny
+			MaxDistance = FMath::Max(MinimalPathSearchDistance, MaxDistance);
+		}
+		
+		if (PathSimilarityWeight > 0 && !SourceEdgeDirection.IsZero())
+		{
+			PathFinder.bEnableDistanceWeighting = true;
+
+			// We need some similarity metric. The current approach is to integrate the squared distance of points along
+			//  the new edge to the line of the source edge.
+			PathFinder.GetWeightedDistanceFunc = [this, SourceEdgeStart, SourceEdgeDirection, MeshMaxDim]
+				(int32 FromVid, int32 ToVid, int32 SeedVid, double EuclideanDistance) 
+			{
+				double SimilarityMetric = SquaredDistanceFromLineIntegratedAlongSegment(
+					FLine3d(SourceEdgeStart, SourceEdgeDirection),
+					FSegment3d(DestinationMesh->GetVertex(FromVid), DestinationMesh->GetVertex(ToVid)));
+
+				// The problem is that we want to combine this integral with distance, and it's hard to know the relative
+				//  importance of the two. Moreover, that balance changes based on mesh scaling, since the integral is
+				//  cubic relative to scale, and distance is linear. To try to balance the two, we divide the integral
+				//  by squared maximal dimension to try to make the integral linear with respect to scale.
+				SimilarityMetric /= (MeshMaxDim*MeshMaxDim);
+				
+				return EuclideanDistance + PathSimilarityWeight * SimilarityMetric;
+			};
+		}
+		
+		if (!PathFinder.ComputeToTargetPoint({ SeedPoint }, DestVids.B, MaxDistance))
+		{
+			return false;
+		}
+
+		TArray<int32> DestVidPath;
+		PathFinder.FindPathToNearestSeed(DestVids.B, DestVidPath); // Note that the output path comes out backwards here
+
+		TArray<int32> DestEidPath;
+		FEdgeSpan::VertexSpanToEdgeSpan(DestinationMesh, DestVidPath, DestEidPath);
+		if (!ensure(DestEidPath.Num() > 0))
+		{
+			return false;
+		}
+
+		if (!DestinationSelectionTids)
+		{
+			DestEidsToMakeSeams.Append(DestEidPath);
+		}
+		else
+		{
+			for (int32 Eid : DestEidPath)
+			{
+				FIndex2i EdgeTids = DestinationMesh->GetEdgeT(Eid);
+				if (DestinationSelectionTids->Contains(EdgeTids.A)
+					|| (EdgeTids.B != IndexConstants::InvalidID && DestinationSelectionTids->Contains(EdgeTids.B)))
+				{
+					DestEidsToMakeSeams.Add(Eid);
+				}
+			}
+		}
+
+		SourceEidToDestinationEndpointEidsVids.Add(SourceEid, TPair<FIndex2i, FIndex2i>(
+			FIndex2i(DestEidPath.Last(), DestEidPath[0]),
+			FIndex2i(DestVidPath.Last(), DestVidPath[0])));
+
+		return true;
+	};//end ProcessSourceEdge 
+
+	// We'll check Progress->Cancelled every 256 source edges.
+	uint32 Count = 0;
+	uint32 IntervalCheckMask = 0xFF;
+
+	if (SourceSelectionTids)
+	{
+		// Iterate through all three edges of each triangle, and use a set to avoid redoing an edge
+		TSet<int32> ProcessedEids;
+		for (int32 Tid : *SourceSelectionTids)
+		{
+			FIndex3i TriEids = SourceMesh->GetTriEdges(Tid);
+			for (int i = 0; i < 3; ++i)
+			{
+				if (Progress && ((++Count & IntervalCheckMask) == 0) && Progress->Cancelled()) { return false; }
+
+				int32 Eid = TriEids[i];
+				bool bAlreadyProcessed = false;
+				ProcessedEids.Add(Eid, &bAlreadyProcessed);
+				if (!bAlreadyProcessed)
+				{
+					bAllSuccessful = ProcessSourceEdge(Eid) && bAllSuccessful;
+				}
+			}
+		}
+	}
+	else
+	{
+		for (int32 Eid : SourceMesh->EdgeIndicesItr())
+		{
+			if (Progress && ((++Count & IntervalCheckMask) == 0) && Progress->Cancelled()) { return false; }
+
+			bAllSuccessful = ProcessSourceEdge(Eid) && bAllSuccessful;
+		}
+	}
+
+	FDynamicMeshUVEditor UVEditor(DestinationMesh, UVLayerIndex, true);
+	bAllSuccessful = UVEditor.CreateSeamsAtEdges(DestEidsToMakeSeams) && bAllSuccessful;
+
+	return bAllSuccessful;
+}
+
+bool UE::Geometry::FDynamicMeshUVTransfer::PerformElementsTransfer(FProgressCancel* Progress)
+{
+	// We'll check Progress->Cancelled every 512 elements.
+	uint32 Count = 0;
+	uint32 IntervalCheckMask = (1 << 9) - 1;
+
+	bool bAllSuccessful = true;
+	TSet<int32> ProcessedSourceElementIDs;
+
+	// Transferring elements on seams is tricky. We have information about the destination edges
+	//  that correspond to the starts and ends of the source edges, and this information can be
+	//  used to match source seam elements to destination elements. Unfortunately, it is possible
+	//  for multiple source edges to map to the same destination edges if the shortest destination
+	//  paths partially overlap. This typically looks like a "pinching" or collapse of a corner 
+	//  of an island, and it causes confusion about which source elements should map to the 
+	//  destination elements at those locations, because some source elements are collapsed away
+	//  here.
+	// The way we resolve this relies on the following: each source element that lies on a seam
+	//  has two adjoining seam edges that can try to bind it to the destination element that is on
+	//  that side of their corresponding destination edge. That element will be the same if the
+	//  two edges didn't collapse, and different if it did. So, by counting the number of times we try
+	//  to make the same binding of elements, we can keep the the right ones.
+	// Seam ends are a special case where the same edge tries to do the same mapping twice, so they
+	//  don't need special handling. Elements at the boundary where an internal seam meets the boundary,
+	//  however, require special handling because we don't happen to iterate through source mesh boundary
+	//  (because we don't try to transfer that). The way we handle that is as follows: if the boundary
+	//  element doesn't show up in the mappings at all, then its source triangle must have been collapsed
+	//  by path transfer, or else it's not at a meeting point of interior seam, in which case it doesn't
+	//  need special handling. If it does show up, then an adjoining seam try to set up that mapping, and
+	//  since we can assume the mesh boundary to still exist, we should accept that mapping. This requires
+	//  us to know when source elements are on the boundary.
+
+	// Key is (DestinationElementID, SourceElementID), value is true if mapping was repeated more than once. 
+	TMap<FIndex2i, bool> Mappings;
+	for (TPair<int32, TPair<FIndex2i, FIndex2i>> EidAndDestEidsVids : SourceEidToDestinationEndpointEidsVids)
+	{
+		if (Progress && ((++Count & IntervalCheckMask) == 0) && Progress->Cancelled()) { return false; }
+
+		const FDynamicMesh3::FEdge Edge = SourceMesh->GetEdge(EidAndDestEidsVids.Key);
+		const FIndex2i& SourceTids = Edge.Tri;
+		const FIndex2i& SourceVids = Edge.Vert;
+		const FIndex2i& DestEids = EidAndDestEidsVids.Value.Key;
+		const FIndex2i& DestVids = EidAndDestEidsVids.Value.Value;
+
+		// Get an "orientation indicator" for source Tid A, which will be either 1 or 2 depending on 
+		//  the ordering of vert A and B in the triangle. This will let us pick the matching tids in
+		//  the destination.
+		FIndex3i SourceTriA = SourceMesh->GetTriangle(SourceTids.A);
+		int SourceTriAOrientationIndicator = (SourceTriA.IndexOf(SourceVids.A) + 3 - SourceTriA.IndexOf(SourceVids.B)) % 3;
+
+		for (int EndpointIndex = 0; EndpointIndex < 2; ++EndpointIndex) // for each endpoint of edge
+		{
+			int32 SourceVid = SourceVids[EndpointIndex];
+			int32 DestVid = DestVids[EndpointIndex];
+			int32 DestOtherVid = DestinationMesh->GetEdgeV(DestEids[EndpointIndex]).OtherElement(DestVid);
+			FIndex2i VertsForOrientation = EndpointIndex == 0 ? FIndex2i(DestVid, DestOtherVid) : FIndex2i(DestOtherVid, DestVid);
+			
+			FIndex2i DestTids = DestinationMesh->GetEdgeT(DestEids[EndpointIndex]);
+			FIndex3i DestTriA = DestinationMesh->GetTriangle(DestTids.A);
+			if (SourceTriAOrientationIndicator != (DestTriA.IndexOf(VertsForOrientation.A) + 3 - DestTriA.IndexOf(VertsForOrientation.B)) % 3)
+			{
+				Swap(DestTids.A, DestTids.B);
+			}
+
+			for (int SideIndex = 0; SideIndex < 2; ++SideIndex)
+			{
+				int32 SourceTid = SourceTids[SideIndex];
+				if (SourceTid == IndexConstants::InvalidID)
+				{
+					continue;
+				}
+
+				FIndex3i SourceTriangle = SourceMesh->GetTriangle(SourceTid);
+				int SourceSubIdx = SourceTriangle.IndexOf(SourceVid);
+				FIndex3i SourceElements = SourceOverlay->GetTriangle(SourceTid);
+				int32 SourceElementID = SourceElements[SourceSubIdx];
+
+				int32 DestTid = DestTids[SideIndex];
+				if (DestTid == IndexConstants::InvalidID)
+				{
+					continue;
+				}
+				FIndex3i DestTriangle = DestinationMesh->GetTriangle(DestTid);
+				int DestSubIdx = DestTriangle.IndexOf(DestVid);
+				FIndex3i DestElements = DestOverlay->GetTriangle(DestTid);
+				int32 DestElementID = DestElements[DestSubIdx];
+
+				if (bool* ExistingValue = Mappings.Find(FIndex2i(DestElementID, SourceElementID)))
+				{
+					*ExistingValue = true;
+				}
+				else
+				{
+					Mappings.Add(FIndex2i(DestElementID, SourceElementID), false);
+				}
+
+				// Note that we don't early out on having processed before because we specifically
+				//  need to touch the same element multiple times to count it. The set is just to
+				//  avoid doing this work for the other element transfer loop further below.
+				ProcessedSourceElementIDs.Add(SourceElementID);
+			}
+		}
+	}
+
+	TSet<int32> PinnedElementIDs;
+
+	for (TPair<FIndex2i, bool> Mapping : Mappings)
+	{
+		// See the big comment block above for our approach to which mappings we keep
+		if (Mapping.Value || SourceBoundaryElements.Contains(Mapping.Key.B))
+		{
+			DestOverlay->SetElement(Mapping.Key.A, SourceOverlay->GetElement(Mapping.Key.B));
+			PinnedElementIDs.Add(Mapping.Key.A);
+		}
+	}
+
+	// Now carry over the values for elements that were not on seams. 
+
+	// Returns false if we fail to find a correspondence or fail to get value
+	auto ProcessElement = [this, &ProcessedSourceElementIDs, &PinnedElementIDs](int32 SourceElementID) -> bool
+	{
+		if (SourceElementID == IndexConstants::InvalidID)
+		{
+			return true;
+		}
+
+		bool bAlreadyProcessed = false;
+		ProcessedSourceElementIDs.Add(SourceElementID, &bAlreadyProcessed);
+		if (bAlreadyProcessed)
+		{
+			return true;
+		}
+
+		int32 ParentSourceVid = SourceOverlay->GetParentVertex(SourceElementID);
+		if (!ensure(ParentSourceVid != IndexConstants::InvalidID))
+		{
+			return false;
+		}
+		int32 DestVid = GetCorrespondingDestVid(ParentSourceVid);
+		if (DestVid == IndexConstants::InvalidID)
+		{
+			return false;
+		}
+
+		// Unfortunately it is possible for a vertex that was not on a seam before to become a seam if
+		//  a nearby path runs through it. With sufficient pain, it might be possible to figure out which
+		//  side of the seam to place the source element value, but that seems not worth it. 
+		TArray<int32> DestElements;
+		DestOverlay->GetVertexElements(DestVid, DestElements);
+		if (DestElements.Num() > 1)
+		{
+			return true;
+		}
+
+		DestOverlay->SetElement(DestElements[0], SourceOverlay->GetElement(SourceElementID));
+		PinnedElementIDs.Add(DestElements[0]);
+		return true;
+	};
+
+	if (SourceSelectionTids)
+	{
+		for (int32 Tid : *SourceSelectionTids)
+		{
+			FIndex3i TriangleElements = SourceOverlay->GetTriangle(Tid);
+			for (int i = 0; i < 3; ++i)
+			{
+				if (Progress && ((++Count & IntervalCheckMask) == 0) && Progress->Cancelled()) { return false; }
+
+				bAllSuccessful = ProcessElement(TriangleElements[i]) && bAllSuccessful;
+			}
+		}
+	}
+	else
+	{
+		for (int32 ElementID : SourceOverlay->ElementIndicesItr())
+		{
+			if (Progress && ((++Count & IntervalCheckMask) == 0) && Progress->Cancelled()) { return false; }
+
+			bAllSuccessful = ProcessElement(ElementID) && bAllSuccessful;
+		}
+	}
+
+	FMeshConnectedComponents ConnectedComponents(DestinationMesh);
+	if (DestinationSelectionTids)
+	{
+		ConnectedComponents.FindConnectedTriangles(DestinationSelectionTids->Array(), [this](int32 Tid1, int32 Tid2)
+		{
+			return DestOverlay->AreTrianglesConnected(Tid1, Tid2);
+		});
+	}
+	else
+	{
+		ConnectedComponents.FindConnectedTriangles([this](int32 Tid1, int32 Tid2)
+		{
+			return DestOverlay->AreTrianglesConnected(Tid1, Tid2);
+		});
+	}
+
+	FDynamicMeshUVEditor UVEditor(DestinationMesh, UVLayerIndex, true);
+	int32 NumComponents = ConnectedComponents.Num();
+	for (int32 ComponentIndex = 0; ComponentIndex < NumComponents; ++ComponentIndex)
+	{
+		if (Progress && Progress->Cancelled()) { return false; }
+
+		bAllSuccessful = UVEditor.SetTriangleUVsFromFreeBoundaryConformal(
+			ConnectedComponents[ComponentIndex].Indices, PinnedElementIDs)
+			&& bAllSuccessful;
+	}
+
+	return bAllSuccessful;
+}
diff --git a/Engine/Plugins/Runtime/GeometryProcessing/Source/DynamicMesh/Public/Parameterization/DynamicMeshUVEditor.h b/Engine/Plugins/Runtime/GeometryProcessing/Source/DynamicMesh/Public/Parameterization/DynamicMeshUVEditor.h
--- a/Engine/Plugins/Runtime/GeometryProcessing/Source/DynamicMesh/Public/Parameterization/DynamicMeshUVEditor.h
+++ b/Engine/Plugins/Runtime/GeometryProcessing/Source/DynamicMesh/Public/Parameterization/DynamicMeshUVEditor.h
@@ -179,6 +179,16 @@ public:
 	 */
 	bool SetTriangleUVsFromFreeBoundaryConformal(const TArray<int32>& Triangles, bool bUseExistingUVTopology, FUVEditResult* Result = nullptr);
 
+	/**
+	 * Create new UV island for given Triangles, and set UVs for that island using Discrete Natural Conformal Map (equivalent to Least-Squares Conformal Map).
+	 * Using this overload implies that we're re-solving for existing UV set rather than constructing per-vertex UVs from triangle set,
+	 * since pinned elements rely on existing values.
+	 * @param Triangles list of triangles
+	 * @param PinnedElementIDs Elements for which to use unmodified UV values.
+	 * @warning computes a single parameterization, so input triangle set must be connected, however this is not verified internally
+	 */
+	bool SetTriangleUVsFromFreeBoundaryConformal(const TArray<int32>& Triangles, const TSet<int32>& PinnedElementIDs, FUVEditResult* Result = nullptr);
+
 	/**
 	 * Create new UV island for given Triangles, and set UVs for that island using Spectral Conformal Map.
 	 * @param Triangles list of triangles
@@ -288,10 +298,19 @@ public:
 
 private: 
 
+	struct FSetUVsFromConformalOptions
+	{
+		bool bUseExistingUVTopology = false;
+		bool bUseSpectral = false; 
+		bool bPreserveIrregularity = false;
+		// currently only used if bUseSpectral is false
+		const TSet<int32>* PinnedElementIDs = nullptr;
+	};
 	/**
 	 * Helper method to solve either the Least Squares Conformal Map or the Spectral Conformal Map problems.
 	 */
-	bool SetTriangleUVsFromConformal(const TArray<int32>& Triangles, bool bUseExistingUVTopology, bool bUseSpectral, bool bPreserveIrregularity, FUVEditResult* Result);
+	bool SetTriangleUVsFromConformal(const TArray<int32>& Triangles, const FSetUVsFromConformalOptions& Options, FUVEditResult* Result);
+	
 };
 
 
diff --git a/Engine/Plugins/Runtime/GeometryProcessing/Source/DynamicMesh/Public/Parameterization/UVTransfer.h b/Engine/Plugins/Runtime/GeometryProcessing/Source/DynamicMesh/Public/Parameterization/UVTransfer.h
new file mode 100644
--- /dev/null
+++ b/Engine/Plugins/Runtime/GeometryProcessing/Source/DynamicMesh/Public/Parameterization/UVTransfer.h
@@ -0,0 +1,88 @@
+// Copyright Epic Games, Inc. All Rights Reserved.
+
+#pragma once
+
+#include "Spatial/PointHashGrid3.h"
+#include "Containers/Set.h"
+#include "Containers/Map.h"
+#include "DynamicMesh/DynamicMeshAttributeSet.h" // FDynamicMeshUVOverlay
+#include "Templates/UniquePtr.h"
+
+class FProgressCancel;
+
+namespace UE::Geometry
+{
+class FDynamicMesh3;
+
+/**
+ * Transfers UVs from a low-resolution mesh to a high-resolution mesh. The source mesh is expected
+ *  to be a simplified version of the destination mesh, simplified using existing vertices only. Thus,
+ *  correspondences are found by position.
+ */
+class FDynamicMeshUVTransfer
+{
+public:
+	DYNAMICMESH_API FDynamicMeshUVTransfer(const FDynamicMesh3* SourceMeshIn, FDynamicMesh3* DestinationMeshIn,
+		int32 UVLayerIndexIn);
+
+	const FDynamicMesh3* SourceMesh = nullptr;
+	FDynamicMesh3* DestinationMesh = nullptr;
+	int UVLayerIndex = 0;
+
+	/** How far to search for a matching vertex on the destination mesh. */
+	double VertexSearchDistance = KINDA_SMALL_NUMBER;
+	/** Cell size used in hash grid when finding correspndence (only affects performance) */
+	double VertexSearchCellSize = VertexSearchDistance * 3.0;
+
+	/** 
+	 * Tuning parameter to make found destination paths follow source edges closer. The higher
+	 *  the value, the more the path similarity metric gets weighed relative to simple euclidean
+	 *  length of the path, which helps select a zigzagging path that follows the edge closer over
+	 *  an "up and over" path that might have slightly shorter length (particularly relevant on
+	 *  against-the-grain diagonals in a gridded region, where path lengths are manhattan distance-like).
+	 */
+	double PathSimilarityWeight = 200.0;
+
+	/**
+	 * If true, existing seams in the destination are removed before adding new ones.
+	 */
+	bool bClearExistingSeamsInDestination = true;
+
+	DYNAMICMESH_API bool TransferSeams(FProgressCancel* Progress);
+	DYNAMICMESH_API bool TransferSeamsAndUVs(FProgressCancel* Progress);
+private:
+	// TODO: Expose and test selection-constrained transfers when support for selecting elements on
+	//  multiple meshes is available. Selection should be settable on either mesh (it is not necessary
+	//  to set both).
+	TSet<int32>* SourceSelectionTids = nullptr;
+	TSet<int32>* DestinationSelectionTids = nullptr;
+
+	void InitializeHashGrid();
+	// Should only be called after InitializeHashGrid
+	int32 GetCorrespondingDestVid(int32 SourceVid);
+
+	void ResetDestinationUVTopology(FProgressCancel* Progress);
+	bool PerformSeamTransfer(FProgressCancel* Progress);
+	// Expects PerformSeamTransfer to have been called so that seams are transferred, and
+	//  SourceEidToDestinationEndpointEidsVids and SourceBoundaryElements are initialized
+	bool PerformElementsTransfer(FProgressCancel* Progress);
+
+	const FDynamicMeshUVOverlay* SourceOverlay = nullptr;
+	FDynamicMeshUVOverlay* DestOverlay = nullptr;
+
+	// Used for getting a correspondence between low res mesh and high res mesh
+	TUniquePtr<TPointHashGrid3<int32, double>> HashGrid;
+	TMap<int32, int32> SourceVidToDestinationVid;
+
+	TMap<int32, TPair<FIndex2i, FIndex2i>> SourceEidToDestinationEndpointEidsVids;
+	TSet<int32> SourceBoundaryElements;
+
+	// This gets multiplied by (SourceEdgeLength + PathSimilarityWeight * PathSimilarityMetric) to limit the
+	//  max search distance when looking for corresponding paths. This is speculative, to avoid looking for inexistent
+	//  paths on very dense meshes on incorrect inputs.
+	double PathLengthToleranceMultiplier = 4.0;
+	// Minimal max distance to use in looking for corresponding paths, in case the source edge is near-degenerate.
+	double MinimalPathSearchDistance = 20.0;
+};
+
+}
\ No newline at end of file
diff --git a/Engine/Plugins/Runtime/MeshModelingToolset/Source/ModelingComponents/Private/Drawing/MeshElementsVisualizer.cpp b/Engine/Plugins/Runtime/MeshModelingToolset/Source/ModelingComponents/Private/Drawing/MeshElementsVisualizer.cpp
--- a/Engine/Plugins/Runtime/MeshModelingToolset/Source/ModelingComponents/Private/Drawing/MeshElementsVisualizer.cpp
+++ b/Engine/Plugins/Runtime/MeshModelingToolset/Source/ModelingComponents/Private/Drawing/MeshElementsVisualizer.cpp
@@ -219,3 +219,13 @@ void UMeshElementsVisualizer::NotifyMeshChanged()
 {
 	WireframeComponent->UpdateWireframe();
 }
+
+void UMeshElementsVisualizerProperties::ShowAllElements(bool bShow)
+{
+	bShowWireframe = bShow;
+	bShowBorders = bShow;
+	bShowUVSeams = bShow;
+	bShowNormalSeams = bShow;
+	bShowTangentSeams = bShow;
+	bShowColorSeams = bShow;
+}
diff --git a/Engine/Plugins/Runtime/MeshModelingToolset/Source/ModelingComponents/Public/Drawing/MeshElementsVisualizer.h b/Engine/Plugins/Runtime/MeshModelingToolset/Source/ModelingComponents/Public/Drawing/MeshElementsVisualizer.h
--- a/Engine/Plugins/Runtime/MeshModelingToolset/Source/ModelingComponents/Public/Drawing/MeshElementsVisualizer.h
+++ b/Engine/Plugins/Runtime/MeshModelingToolset/Source/ModelingComponents/Public/Drawing/MeshElementsVisualizer.h
@@ -47,6 +47,9 @@ public:
 	UPROPERTY(EditAnywhere, Category = MeshElementVisualization, meta = (EditCondition = "bVisible"))
 	bool bShowColorSeams = true;
 
+	/** Helper function to set all the bShow... booleans to the same value. */
+	void ShowAllElements(bool bShow);
+
 	/** Multiplier on edge thicknesses */
 	UPROPERTY(EditAnywhere, Category = MeshElementVisualization, AdvancedDisplay, meta = (UIMin = 0.1, UIMax = 10.0, EditCondition = "bVisible"))
 	float ThicknessScale = 1.0;
