ModelingTools: Remove deprecated baking code

diff --git a/Engine/Plugins/Experimental/MeshModelingToolsetExp/Source/MeshModelingToolsExp/Private/BakeMultiMeshAttributeMapsTool.cpp b/Engine/Plugins/Experimental/MeshModelingToolsetExp/Source/MeshModelingToolsExp/Private/BakeMultiMeshAttributeMapsTool.cpp
--- a/Engine/Plugins/Experimental/MeshModelingToolsetExp/Source/MeshModelingToolsExp/Private/BakeMultiMeshAttributeMapsTool.cpp
+++ b/Engine/Plugins/Experimental/MeshModelingToolsetExp/Source/MeshModelingToolsExp/Private/BakeMultiMeshAttributeMapsTool.cpp
@@ -122,11 +122,6 @@ public:
 		DetailTextureMaps[Mesh] = Map;
 	}
 
-	virtual void SetNormalMap(const void* Mesh, const FBakeDetailTexture& Map) override
-	{
-		DetailNormalMaps[Mesh] = FBakeDetailNormalTexture(Map.Key, Map.Value, EBakeDetailNormalSpace::Tangent);
-	}
-
 	virtual void SetNormalTextureMap(const void* Mesh, const FBakeDetailNormalTexture& Map) override
 	{
 		DetailNormalMaps[Mesh] = Map;
@@ -136,11 +131,6 @@ public:
 	{
 		return DetailTextureMaps.Find(Mesh);
 	}
-	
-	virtual const FBakeDetailTexture* GetNormalMap(const void* Mesh) const override
-	{
-		return nullptr;
-	}
 
 	virtual const FBakeDetailNormalTexture* GetNormalTextureMap(const void* Mesh) const override
 	{
diff --git a/Engine/Plugins/Runtime/GeometryProcessing/Source/DynamicMesh/Public/Sampling/MeshBakerCommon.h b/Engine/Plugins/Runtime/GeometryProcessing/Source/DynamicMesh/Public/Sampling/MeshBakerCommon.h
--- a/Engine/Plugins/Runtime/GeometryProcessing/Source/DynamicMesh/Public/Sampling/MeshBakerCommon.h
+++ b/Engine/Plugins/Runtime/GeometryProcessing/Source/DynamicMesh/Public/Sampling/MeshBakerCommon.h
@@ -123,20 +123,12 @@ public:
 	/** Associate a texture map and UV layer index for a given mesh in the detail set */
 	virtual void SetTextureMap(const void* Mesh, const FBakeDetailTexture& Map) = 0;
 
-	/** Associate a normal map and UV layer index for a given mesh in the detail set */
-	UE_DEPRECATED(5.1, "Use SetNormalTextureMap instead. This implementation assumes tangent normals.")
-	virtual void SetNormalMap(const void* Mesh, const FBakeDetailTexture& Map) = 0;
-
 	/** Associate a normal map and UV layer index for a given mesh in the detail set */
 	virtual void SetNormalTextureMap(const void* Mesh, const FBakeDetailNormalTexture& Map) = 0;
 
 	/** Retrieve a texture map and UV layer index from a given mesh in the detail set */
 	virtual const FBakeDetailTexture* GetTextureMap(const void* Mesh) const = 0;
 
-	/** Retrieve a normal map and UV layer index from a given mesh in the detail set */
-	UE_DEPRECATED(5.1, "Use GetNormalTextureMap instead. This implementation assumes tangent normals.")
-	virtual const FBakeDetailTexture* GetNormalMap(const void* Mesh) const = 0;
-
 	virtual const FBakeDetailNormalTexture* GetNormalTextureMap(const void* Mesh) const = 0;
 
 	/** @return true if identity correspondence is supported */
@@ -394,11 +386,6 @@ public:
 		DetailTextureMap = Map;
 	}
 
-	virtual void SetNormalMap(const void* Mesh, const FBakeDetailTexture& Map) override
-	{
-		DetailNormalTextureMap = FBakeDetailNormalTexture(Map.Key, Map.Value, EBakeDetailNormalSpace::Tangent);
-	}
-
 	virtual void SetNormalTextureMap(const void* Mesh, const FBakeDetailNormalTexture& Map) override
 	{
 		DetailNormalTextureMap = Map;
@@ -408,11 +395,6 @@ public:
 	{
 		return &DetailTextureMap;
 	}
-	
-	virtual const FBakeDetailTexture* GetNormalMap(const void* Mesh) const override
-	{
-		return nullptr;
-	}
 
 	virtual const FBakeDetailNormalTexture* GetNormalTextureMap(const void* Mesh) const override
 	{
@@ -622,9 +604,6 @@ protected:
 	const FDynamicMeshAABBTree3* DetailSpatial = nullptr;
 	const FMeshTangentsd* DetailTangents = nullptr;
 	FBakeDetailTexture DetailTextureMap = FBakeDetailTexture(nullptr, 0);
-
-	UE_DEPRECATED(5.1, "Use DetailNormalTextureMap instead.")
-	FBakeDetailTexture DetailNormalMap = FBakeDetailTexture(nullptr, 0);
 	FBakeDetailNormalTexture DetailNormalTextureMap = FBakeDetailNormalTexture(nullptr, 0, EBakeDetailNormalSpace::Tangent);
 };		
 	
diff --git a/Engine/Plugins/Runtime/GeometryProcessing/Source/DynamicMesh/Public/Sampling/MeshNormalMapEvaluator.h b/Engine/Plugins/Runtime/GeometryProcessing/Source/DynamicMesh/Public/Sampling/MeshNormalMapEvaluator.h
--- a/Engine/Plugins/Runtime/GeometryProcessing/Source/DynamicMesh/Public/Sampling/MeshNormalMapEvaluator.h
+++ b/Engine/Plugins/Runtime/GeometryProcessing/Source/DynamicMesh/Public/Sampling/MeshNormalMapEvaluator.h
@@ -44,10 +44,6 @@ public:
 protected:
 	// Cached data
 	const IMeshBakerDetailSampler* DetailSampler = nullptr;
-	using FDetailNormalTexture UE_DEPRECATED(5.1, "Use FNormalTexture instead.") = IMeshBakerDetailSampler::FBakeDetailTexture;
-	using FDetailNormalTextureMap UE_DEPRECATED(5.1, "Use FNormalTextureMap instead.") = TMap<const void*, IMeshBakerDetailSampler::FBakeDetailTexture>;
-	UE_DEPRECATED(5.1, "Use DetailNormalMaps instead.")
-	TMap<const void*, IMeshBakerDetailSampler::FBakeDetailTexture> DetailNormalTextures;
 
 	using FNormalTexture = IMeshBakerDetailSampler::FBakeDetailNormalTexture;
 	using FNormalTextureMap = TMap<const void*, FNormalTexture>;
diff --git a/Engine/Plugins/Runtime/GeometryProcessing/Source/DynamicMesh/Public/Sampling/MeshPropertyMapEvaluator.h b/Engine/Plugins/Runtime/GeometryProcessing/Source/DynamicMesh/Public/Sampling/MeshPropertyMapEvaluator.h
--- a/Engine/Plugins/Runtime/GeometryProcessing/Source/DynamicMesh/Public/Sampling/MeshPropertyMapEvaluator.h
+++ b/Engine/Plugins/Runtime/GeometryProcessing/Source/DynamicMesh/Public/Sampling/MeshPropertyMapEvaluator.h
@@ -58,11 +58,6 @@ public:
 protected:
 	// Cached data
 	const IMeshBakerDetailSampler* DetailSampler = nullptr;
-	
-	using FDetailNormalTexture UE_DEPRECATED(5.1, "Use FNormalTexture instead.") = IMeshBakerDetailSampler::FBakeDetailTexture;
-	using FDetailNormalTextureMap UE_DEPRECATED(5.1, "Use FNormalTextureMap instead.") = TMap<const void*, IMeshBakerDetailSampler::FBakeDetailTexture>;
-	UE_DEPRECATED(5.1, "Use DetailNormalMaps instead.")
-	TMap<const void*, IMeshBakerDetailSampler::FBakeDetailTexture> DetailNormalTextures;
 
 	using FNormalTexture = IMeshBakerDetailSampler::FBakeDetailNormalTexture;
 	using FNormalTextureMap = TMap<const void*, FNormalTexture>;
