ModelingTools: Removed unused code from FMeshUVShellMapEvaluator.

diff --git a/Engine/Plugins/Runtime/GeometryProcessing/Source/DynamicMesh/Private/Sampling/MeshUVShellMapEvaluator.cpp b/Engine/Plugins/Runtime/GeometryProcessing/Source/DynamicMesh/Private/Sampling/MeshUVShellMapEvaluator.cpp
--- a/Engine/Plugins/Runtime/GeometryProcessing/Source/DynamicMesh/Private/Sampling/MeshUVShellMapEvaluator.cpp
+++ b/Engine/Plugins/Runtime/GeometryProcessing/Source/DynamicMesh/Private/Sampling/MeshUVShellMapEvaluator.cpp
@@ -61,9 +61,6 @@ FVector4f FMeshUVShellMapEvaluator::SampleFunction(const FCorrespondenceSample&
 	if (DetailSampler->IsTriangle(DetailMesh, DetailTriID))
 	{
 		// Compute barycentric/areal coords in UV space.
-		FVector2f UV = FVector2f::Zero();
-		DetailSampler->TriBaryInterpolateUV(DetailMesh, DetailTriID, DetailBaryCoords, UVLayer, UV);
-
 		FVector2f UV0, UV1, UV2;
 		DetailSampler->GetTriUVs(DetailMesh, DetailTriID, UVLayer, UV0, UV1, UV2);
 
