UVEditor: Add "make island" action and a debug "Unset UVs" action.

diff --git a/Engine/Plugins/Editor/UVEditor/Source/UVEditor/Private/Actions/UnsetUVsAction.cpp b/Engine/Plugins/Editor/UVEditor/Source/UVEditor/Private/Actions/UnsetUVsAction.cpp
new file mode 100644
--- /dev/null
+++ b/Engine/Plugins/Editor/UVEditor/Source/UVEditor/Private/Actions/UnsetUVsAction.cpp
@@ -0,0 +1,95 @@
+// Copyright Epic Games, Inc. All Rights Reserved.
+
+#include "Actions/UnsetUVsAction.h"
+
+#include "DynamicMesh/DynamicMesh3.h"
+#include "DynamicMesh/DynamicMeshAttributeSet.h" // FDynamicMeshUVOverlay
+#include "ToolTargets/UVEditorToolMeshInput.h"
+
+#include "Parameterization/DynamicMeshUVEditor.h"
+#include "Selection/UVToolSelection.h"
+#include "Selection/UVToolSelectionAPI.h"
+#include "Parameterization/UVUnwrapMeshUtil.h"
+
+#define LOCTEXT_NAMESPACE "UUnsetUVsAction"
+
+namespace UnsetUVsActionLocals
+{
+	using namespace UE::Geometry;
+
+	const FText TransactionName(LOCTEXT("TransactionName", "UnsetUVs"));
+}
+
+bool UUnsetUVsAction::CanExecuteAction() const
+{
+	return SelectionAPI->HaveSelections()
+		&& SelectionAPI->GetSelectionsType() == UE::Geometry::FUVToolSelection::EType::Triangle;
+}
+bool UUnsetUVsAction::ExecuteAction()
+{
+	using namespace UnsetUVsActionLocals;
+	using namespace UE::Geometry;
+
+	EmitChangeAPI->BeginUndoTransaction(TransactionName);
+	
+	bool bSuccess = true;
+
+	// We make a copy of selection and clear it right away so that we don't have an inconsistent selection
+	//  state while updating the unwrap.
+	const TArray<FUVToolSelection> Selections = SelectionAPI->GetSelections();
+	// Don't broadcast, but do emit the change, so that it is in the correct place in the transaction.
+	SelectionAPI->ClearSelections(false, true);
+
+	// This is a copy so that we can add to it and call the setter
+	TArray<FUVToolSelection> UnsetSelections = SelectionAPI->GetUnsetElementAppliedMeshSelections();
+
+	for (const FUVToolSelection& Selection : Selections)
+	{
+		if (!ensure(Selection.Target.IsValid() && Selection.Target->AppliedCanonical 
+			&& !Selection.SelectedIDs.IsEmpty() && Selection.Type == FUVToolSelection::EType::Triangle))
+		{
+			continue;
+		}
+
+		FDynamicMesh3& Mesh = *Selection.Target->AppliedCanonical;
+		FDynamicMeshUVOverlay* UVOverlay = Mesh.Attributes()->GetUVLayer(Selection.Target->UVLayerIndex);
+		if (!ensure(UVOverlay))
+		{
+			continue;
+		}
+
+		for (int32 Tid : Selection.SelectedIDs)
+		{
+			UVOverlay->UnsetTriangle(Tid);
+		}
+
+		FDynamicMeshChangeTracker ChangeTracker(Selection.Target->UnwrapCanonical.Get());
+		ChangeTracker.BeginChange();
+		ChangeTracker.SaveTriangles(Selection.SelectedIDs, true);
+
+		TArray<int32> ChangedTidsArray = Selection.SelectedIDs.Array();
+		Selection.Target->UpdateAllFromAppliedCanonical(UUVEditorToolMeshInput::NONE_CHANGED_ARG, &ChangedTidsArray, &ChangedTidsArray);
+
+		checkSlow(UVUnwrapMeshUtil::DoesUnwrapMatchOverlay(*UVOverlay, *Selection.Target->UnwrapCanonical, Selection.Target->UVToVertPosition, 0.01));
+
+		EmitChangeAPI->EmitToolIndependentUnwrapCanonicalChange(Selection.Target.Get(),
+			ChangeTracker.EndChange(), TransactionName);
+
+		if (FUVToolSelection* ExistingUnsetSelection = UnsetSelections.FindByPredicate(
+			[Selection](const FUVToolSelection& ExistingSelection) { return ExistingSelection.Target == Selection.Target; }))
+		{
+			ExistingUnsetSelection->SelectedIDs.Append(Selection.SelectedIDs);
+		}
+		else
+		{
+			UnsetSelections.Emplace(Selection);
+		}
+	}//end for each target that has selections
+
+	// Emit the selection change from empty to new unset
+	SelectionAPI->SetUnsetElementAppliedMeshSelections(UnsetSelections, true, true);
+	EmitChangeAPI->EndUndoTransaction();
+	return bSuccess;
+}
+
+#undef LOCTEXT_NAMESPACE
\ No newline at end of file
diff --git a/Engine/Plugins/Editor/UVEditor/Source/UVEditor/Private/Actions/UnsetUVsAction.h b/Engine/Plugins/Editor/UVEditor/Source/UVEditor/Private/Actions/UnsetUVsAction.h
new file mode 100644
--- /dev/null
+++ b/Engine/Plugins/Editor/UVEditor/Source/UVEditor/Private/Actions/UnsetUVsAction.h
@@ -0,0 +1,20 @@
+// Copyright Epic Games, Inc. All Rights Reserved.
+
+#pragma once
+
+#include "Actions/UVToolAction.h"
+
+#include "UnsetUVsAction.generated.h"
+
+//~ Action that unsets UVs, intended for testing other tools. For now, this is
+//~  in the editor module (instead of UVEditorTools) so that we can keep it
+//~  unexposed.
+UCLASS()
+class UUnsetUVsAction : public UUVToolAction
+{
+	GENERATED_BODY()
+
+public:
+	virtual bool CanExecuteAction() const override;
+	virtual bool ExecuteAction() override;
+};
\ No newline at end of file
diff --git a/Engine/Plugins/Editor/UVEditor/Source/UVEditor/Private/UVEditorMode.cpp b/Engine/Plugins/Editor/UVEditor/Source/UVEditor/Private/UVEditorMode.cpp
--- a/Engine/Plugins/Editor/UVEditor/Source/UVEditor/Private/UVEditorMode.cpp
+++ b/Engine/Plugins/Editor/UVEditor/Source/UVEditor/Private/UVEditorMode.cpp
@@ -3,6 +3,8 @@
 #include "UVEditorMode.h"
 
 #include "Algo/AnyOf.h"
+#include "Actions/UnsetUVsAction.h"
+#include "Actions/UVMakeIslandAction.h"
 #include "Actions/UVSplitAction.h"
 #include "Actions/UVSeamSewAction.h"
 #include "Actions/UVToolAction.h"
@@ -674,6 +676,8 @@ void UUVEditorMode::RegisterActions()
 	};
 	PrepAction(CommandInfos.SewAction, NewObject<UUVSeamSewAction>());
 	PrepAction(CommandInfos.SplitAction, NewObject<UUVSplitAction>());
+	PrepAction(CommandInfos.MakeIslandAction, NewObject<UUVMakeIslandAction>());
+	PrepAction(CommandInfos.UnsetUVsAction, NewObject<UUnsetUVsAction>());
 }
 
 bool UUVEditorMode::ShouldToolStartBeAllowed(const FString& ToolIdentifier) const
diff --git a/Engine/Plugins/Editor/UVEditor/Source/UVEditor/Private/UVEditorModeToolkit.cpp b/Engine/Plugins/Editor/UVEditor/Source/UVEditor/Private/UVEditorModeToolkit.cpp
--- a/Engine/Plugins/Editor/UVEditor/Source/UVEditor/Private/UVEditorModeToolkit.cpp
+++ b/Engine/Plugins/Editor/UVEditor/Source/UVEditor/Private/UVEditorModeToolkit.cpp
@@ -477,6 +477,11 @@ void FUVEditorModeToolkit::BuildToolPalette(FName PaletteIndex, class FToolBarBu
 	{
 		ToolbarBuilder.AddToolBarButton(Commands.SewAction);
 		ToolbarBuilder.AddToolBarButton(Commands.SplitAction);
+		ToolbarBuilder.AddToolBarButton(Commands.MakeIslandAction);
+		if (CVarAddDebugTools.GetValueOnGameThread())
+		{
+			ToolbarBuilder.AddToolBarButton(Commands.UnsetUVsAction);
+		}
 
 		ToolbarBuilder.AddToolBarButton(Commands.BeginLayoutTool);
 		ToolbarBuilder.AddToolBarButton(Commands.BeginTransformTool);
diff --git a/Engine/Plugins/Editor/UVEditor/Source/UVEditorTools/Private/Actions/UVMakeIslandAction.cpp b/Engine/Plugins/Editor/UVEditor/Source/UVEditorTools/Private/Actions/UVMakeIslandAction.cpp
new file mode 100644
--- /dev/null
+++ b/Engine/Plugins/Editor/UVEditor/Source/UVEditorTools/Private/Actions/UVMakeIslandAction.cpp
@@ -0,0 +1,112 @@
+// Copyright Epic Games, Inc. All Rights Reserved.
+
+#include "Actions/UVMakeIslandAction.h"
+
+#include "Containers/Map.h"
+#include "Containers/Set.h"
+#include "DynamicMesh/DynamicMesh3.h"
+#include "DynamicMesh/DynamicMeshAttributeSet.h" // FDynamicMeshUVOverlay
+#include "Parameterization/DynamicMeshUVEditor.h"
+#include "Parameterization/UVUnwrapMeshUtil.h"
+#include "Selection/UVToolSelection.h"
+#include "Selection/UVToolSelectionAPI.h"
+#include "ToolTargets/UVEditorToolMeshInput.h"
+#include "UVEditorUXSettings.h"
+
+#define LOCTEXT_NAMESPACE "UUVMakeIslandAction"
+
+namespace UVMakeIslandActionLocals
+{
+	using namespace UE::Geometry;
+
+	const FText TransactionName(LOCTEXT("TransactionName", "Make Island"));
+}
+
+bool UUVMakeIslandAction::CanExecuteAction() const
+{
+	return (SelectionAPI->HaveSelections()
+		&& SelectionAPI->GetSelectionsType() == UE::Geometry::FUVToolSelection::EType::Triangle)
+		|| SelectionAPI->HaveUnsetElementAppliedMeshSelections();
+}
+bool UUVMakeIslandAction::ExecuteAction()
+{
+	using namespace UVMakeIslandActionLocals;
+	using namespace UE::Geometry;
+
+	EmitChangeAPI->BeginUndoTransaction(TransactionName);
+	
+	bool bSuccess = true;
+
+	TArray<FUVToolSelection> Selections = SelectionAPI->GetSelections();
+	
+	// Lump the unset element selections into the regular selections, because the "make island" operation will set them
+	for (const FUVToolSelection& UnsetSelection : SelectionAPI->GetUnsetElementAppliedMeshSelections())
+	{
+		if (!ensure(UnsetSelection.Target.IsValid() && UnsetSelection.Target->AppliedCanonical
+			&& !UnsetSelection.SelectedIDs.IsEmpty() && UnsetSelection.Type == FUVToolSelection::EType::Triangle))
+		{
+			continue;
+		}
+
+		if (FUVToolSelection* ExistingSelection = Selections.FindByPredicate(
+			[&UnsetSelection](const FUVToolSelection& Selection) { return Selection.Target == UnsetSelection.Target; }))
+		{
+			if (ensure(ExistingSelection->Type == FUVToolSelection::EType::Triangle))
+			{
+				ExistingSelection->SelectedIDs.Append(UnsetSelection.SelectedIDs);
+			}
+		}
+		else
+		{
+			Selections.Add(UnsetSelection);
+		}
+	}
+	// Clear the unset selections now so that its in the proper place in the undo stack
+	bool bHadUnsetSelections = SelectionAPI->HaveUnsetElementAppliedMeshSelections();
+	SelectionAPI->ClearUnsetElementAppliedMeshSelections(/*bBroadcast*/ false, /*bEmit*/ true);
+
+	// Now process all the regular selections
+	for (FUVToolSelection& Selection : Selections)
+	{
+		if (!ensure(Selection.Target.IsValid() && Selection.Target->AppliedCanonical 
+			&& !Selection.SelectedIDs.IsEmpty() && Selection.Type == FUVToolSelection::EType::Triangle))
+		{
+			continue;
+		}
+
+		FDynamicMesh3& Mesh = *Selection.Target->AppliedCanonical;
+		FDynamicMeshUVOverlay* UVOverlay = Mesh.Attributes()->GetUVLayer(Selection.Target->UVLayerIndex);
+		if (!ensure(UVOverlay))
+		{
+			continue;
+		}
+
+		TSet<int32> ChangedTids;
+		UE::Geometry::FUVEditResult Result;
+		FDynamicMeshUVEditor UVEditor(&Mesh, UVOverlay);
+		bSuccess = UVEditor.MakeIsland(Selection.SelectedIDs, &Result, &ChangedTids) && bSuccess;
+		ensure(bSuccess);
+
+		FDynamicMeshChangeTracker ChangeTracker(Selection.Target->UnwrapCanonical.Get());
+		ChangeTracker.BeginChange();
+		ChangeTracker.SaveTriangles(ChangedTids, true);
+
+		TArray<int32> ChangedTidsArray = ChangedTids.Array();
+		Selection.Target->UpdateAllFromAppliedCanonical(&Result.NewUVElements, &ChangedTidsArray, &ChangedTidsArray);
+
+		checkSlow(UVUnwrapMeshUtil::DoesUnwrapMatchOverlay(*UVOverlay, *Selection.Target->UnwrapCanonical, Selection.Target->UVToVertPosition, 0.01));
+
+		EmitChangeAPI->EmitToolIndependentUnwrapCanonicalChange(Selection.Target.Get(),
+			ChangeTracker.EndChange(), TransactionName);
+	}//end for each target that has selections
+
+	if (bHadUnsetSelections)
+	{
+		// If we had unset elements selected, they are now selected elements.
+		SelectionAPI->SetSelections(Selections, true, true);
+	}
+	EmitChangeAPI->EndUndoTransaction();
+	return bSuccess;
+}
+
+#undef LOCTEXT_NAMESPACE
\ No newline at end of file
diff --git a/Engine/Plugins/Editor/UVEditor/Source/UVEditorTools/Private/UVEditorStyle.cpp b/Engine/Plugins/Editor/UVEditor/Source/UVEditorTools/Private/UVEditorStyle.cpp
--- a/Engine/Plugins/Editor/UVEditor/Source/UVEditorTools/Private/UVEditorStyle.cpp
+++ b/Engine/Plugins/Editor/UVEditor/Source/UVEditorTools/Private/UVEditorStyle.cpp
@@ -33,10 +33,13 @@ FUVEditorStyle::FUVEditorStyle()
 	Set("UVEditor.BeginAlignTool",              new IMAGE_BRUSH_SVG("AlignLeft", ToolbarIconSize));
 	Set("UVEditor.BeginDistributeTool",         new IMAGE_BRUSH_SVG("DistributeHorizontally", ToolbarIconSize));
 	Set("UVEditor.BeginTexelDensityTool",       new IMAGE_BRUSH_SVG("TexelDensity", ToolbarIconSize));
+	Set("UVEditor.BeginBrushSelectTool",        new IMAGE_BRUSH("MeshSelect_40x", ToolbarIconSize));
+	Set("UVEditor.BeginBrushSelectTool.Small",  new IMAGE_BRUSH("MeshSelect_40x", ToolbarIconSize));
 
 	// Select tool actions
 	Set("UVEditor.SewAction", new IMAGE_BRUSH_SVG("UVSew", ToolbarIconSize));
 	Set("UVEditor.SplitAction", new IMAGE_BRUSH_SVG("UVCut", ToolbarIconSize));
+	Set("UVEditor.MakeIslandAction", new IMAGE_BRUSH_SVG("SelectionIslands", ToolbarIconSize));
 	Set("UVEditor.IslandConformalUnwrapAction", new IMAGE_BRUSH_SVG("UVUnwrap", ToolbarIconSize));
 
 	// Top toolbar icons
diff --git a/Engine/Plugins/Editor/UVEditor/Source/UVEditorTools/Public/Actions/UVMakeIslandAction.h b/Engine/Plugins/Editor/UVEditor/Source/UVEditorTools/Public/Actions/UVMakeIslandAction.h
new file mode 100644
--- /dev/null
+++ b/Engine/Plugins/Editor/UVEditor/Source/UVEditorTools/Public/Actions/UVMakeIslandAction.h
@@ -0,0 +1,26 @@
+// Copyright Epic Games, Inc. All Rights Reserved.
+
+#pragma once
+
+#include "Actions/UVToolAction.h"
+
+#include "UVMakeIslandAction.generated.h"
+
+/**
+ * Action that takes the currently selected triangles and makes a separate UV island out of them,
+ *  i.e. any interior seams are removed, and seams are added around the boundary of the selection.
+ *  If the selection is not connected in the mesh, islands will be created for each connected
+ *  component of selected triangles. If some of the selected triangles have unset UVs, they will
+ *  be initialized to 0 for the purposes to creating an island. The output is the triangles of
+ *  the created islands.
+ * An unwrap operation on the island is typically a good thing to follow this action.
+ */
+UCLASS()
+class UVEDITORTOOLS_API UUVMakeIslandAction : public UUVToolAction
+{
+	GENERATED_BODY()
+
+public:
+	virtual bool CanExecuteAction() const override;
+	virtual bool ExecuteAction() override;
+};
\ No newline at end of file
diff --git a/Engine/Plugins/Runtime/GeometryProcessing/Source/DynamicMesh/Private/Parameterization/DynamicMeshUVEditor.cpp b/Engine/Plugins/Runtime/GeometryProcessing/Source/DynamicMesh/Private/Parameterization/DynamicMeshUVEditor.cpp
--- a/Engine/Plugins/Runtime/GeometryProcessing/Source/DynamicMesh/Private/Parameterization/DynamicMeshUVEditor.cpp
+++ b/Engine/Plugins/Runtime/GeometryProcessing/Source/DynamicMesh/Private/Parameterization/DynamicMeshUVEditor.cpp
@@ -933,6 +933,52 @@ static int32 FindUVElementForVertex(const FDynamicMesh3* Mesh, const FDynamicMes
 }
 
 
+void UE::Geometry::FDynamicMeshUVEditor::MakeSureUVsAreSet(const TSet<int32>& Triangles, 
+	FUVEditResult* Result, TSet<int32>* ChangedTrianglesOut)
+{
+	if (!ensure(Mesh && UVOverlay))
+	{
+		return;
+	}
+
+	TMap<int32, int32> VidToElement;
+	for (int32 Tid : Triangles)
+	{
+		if (UVOverlay->IsSetTriangle(Tid))
+		{
+			continue;
+		}
+
+		FIndex3i ElementsToSet;
+		FIndex3i TriVids = Mesh->GetTriangle(Tid);
+		for (int i = 0; i < 3; ++i)
+		{
+			if (int32* ExistingElement = VidToElement.Find(TriVids[i]))
+			{
+				ElementsToSet[i] = *ExistingElement;
+			}
+			else
+			{
+				int32 Element = UVOverlay->AppendElement(FVector2f::Zero());
+				ElementsToSet[i] = Element;
+				VidToElement.Add(TriVids[i], Element);
+
+				if (Result)
+				{
+					Result->NewUVElements.Add(Element);
+				}
+			}
+		}
+		UVOverlay->SetTriangle(Tid, ElementsToSet);
+
+		if (ChangedTrianglesOut)
+		{
+			ChangedTrianglesOut->Add(Tid);
+		}
+	}
+}
+
+
 bool FDynamicMeshUVEditor::RemoveSeamsAtEdges(const TSet<int32>& EidsToRemoveAsSeams)
 {
 	for (int32 Eid : EidsToRemoveAsSeams)
@@ -1124,6 +1170,109 @@ bool FDynamicMeshUVEditor::CreateSeamsAtEdges(const TSet<int32>& EidsToMakeIntoS
 
 
 
+bool UE::Geometry::FDynamicMeshUVEditor::MakeIsland(const TSet<int32>& TidsToMakeIntoIsland, FUVEditResult* Result, TSet<int32>* ChangedTrianglesOut)
+{
+	using namespace FDynamicMeshUVEditorLocals;
+
+	if (!ensure(UVOverlay && Mesh))
+	{
+		return false;
+	}
+	
+	// We may add new elements during either initialization or seam insertion. However, upon welding,
+	//  we might end up destroying them. So we'll accumulate them and then filter them out at the end.
+	ON_SCOPE_EXIT
+	{
+		if (Result)
+		{
+			Result->NewUVElements.RemoveAllSwap([this](int32 Element) { return !UVOverlay->IsElement(Element); });
+		}
+	};
+
+	// First make sure that all the relevant triangles have UVs set.
+	MakeSureUVsAreSet(TidsToMakeIntoIsland, Result, ChangedTrianglesOut);
+
+	// Gather the edges we need to edit
+	TSet<int32> EidsToMakeSeams;
+	TSet<int32> EidsToJoin;
+	TSet<int32> TouchedVids;
+
+	TSet<int32> ProcessedEids;
+	for (int32 Tid : TidsToMakeIntoIsland)
+	{
+		FIndex3i TriEids = Mesh->GetTriEdges(Tid);
+		for (int i = 0; i < 3; ++i)
+		{
+			int32 Eid = TriEids[i];
+			bool bAlreadyProcessed = false;
+			ProcessedEids.Add(Eid, &bAlreadyProcessed);
+			if (bAlreadyProcessed)
+			{
+				continue;
+			}
+
+			FDynamicMesh3::FEdge Edge = Mesh->GetEdge(Eid);
+			if (Edge.Tri.B == IndexConstants::InvalidID)
+			{
+				// Don't need to do anything for edges that are on the mesh boundary
+				continue;
+			}
+
+			bool bIsCurrentlySeam = UVOverlay->IsSeamEdge(Eid);
+			bool bShouldBeSeam = !TidsToMakeIntoIsland.Contains(Edge.Tri.A == Tid ? Edge.Tri.B : Edge.Tri.A);
+			
+			if (bIsCurrentlySeam != bShouldBeSeam)
+			{
+				TouchedVids.Add(Edge.Vert.A);
+				TouchedVids.Add(Edge.Vert.B);
+
+				if (bShouldBeSeam)
+				{
+					EidsToMakeSeams.Add(Eid);
+				}
+				else
+				{
+					EidsToJoin.Add(Eid);
+				}
+			}
+		}
+	}//end gathering edges to edit
+
+	if (EidsToJoin.IsEmpty() && EidsToMakeSeams.IsEmpty())
+	{
+		// There must not have been anything to change
+		return true;
+	}
+
+	// We need to do seam insertion first so that we don't move neighboring triangles unnecessarily while welding
+	//  seams inside the island. This is minorly inconvenient since we'll end up having to filter newly created
+	//  elements after the subsequent join operation, but we have to do that for any newly initialized UVs anyway.
+	UE::Geometry::FUVEditResult AddSeamResult;
+	bool bSuccess = CreateSeamsAtEdges(EidsToMakeSeams, &AddSeamResult);
+	if (Result)
+	{
+		Result->NewUVElements.Append(AddSeamResult.NewUVElements);
+		// These get filtered on exit.
+	}
+
+	bSuccess = RemoveSeamsAtEdges(EidsToJoin) && bSuccess;
+
+	if (ChangedTrianglesOut)
+	{
+		// Some of these didn't actually get changed (if they kept their original element), but this is the easiest
+		//  way to make sure we mark anything whose connectivity might have changed. The ideal thing would have been
+		//  to make CreateSeamsAtEdges and RemoveSeamsAtEdges output changed tids instead.
+		for (int32 Vid : TouchedVids)
+		{
+			TArray<int32> Tids;
+			Mesh->GetVtxTriangles(Vid, Tids);
+			ChangedTrianglesOut->Append(Tids);
+		}
+	}
+	return bSuccess;
+}
+
+
 void FDynamicMeshUVEditor::SetTriangleUVsFromBoxProjection(
 	const TArray<int32>& Triangles, 
 	TFunctionRef<FVector3d(const FVector3d&)> PointTransform, 
diff --git a/Engine/Plugins/Runtime/GeometryProcessing/Source/DynamicMesh/Private/Parameterization/UVUnwrapMeshUtil.cpp b/Engine/Plugins/Runtime/GeometryProcessing/Source/DynamicMesh/Private/Parameterization/UVUnwrapMeshUtil.cpp
--- a/Engine/Plugins/Runtime/GeometryProcessing/Source/DynamicMesh/Private/Parameterization/UVUnwrapMeshUtil.cpp
+++ b/Engine/Plugins/Runtime/GeometryProcessing/Source/DynamicMesh/Private/Parameterization/UVUnwrapMeshUtil.cpp
@@ -74,7 +74,11 @@ namespace UVUnwrapMeshUtilLocals
 			}
 
 			FIndex3i NewTriangle = UVOverlayIn.GetTriangle(Tid);
-			UnwrapMeshOut.InsertTriangle(Tid, NewTriangle, 0, true);
+			UE::Geometry::EMeshResult Result = UnwrapMeshOut.InsertTriangle(Tid, NewTriangle, 0, true);
+			
+			// One way you could hit this ensure is if the corresponding vertices didn't exist in the unwrap. That
+			// means that somewhere, newly added elements were excluded from the list of changed elements.
+			ensure(Result == UE::Geometry::EMeshResult::Ok);
 		}
 		UnwrapMeshOut.EndUnsafeTrianglesInsert();
 
@@ -553,12 +557,12 @@ bool UVUnwrapMeshUtil::DoesUnwrapMatchOverlay(const FDynamicMeshUVOverlay& Overl
 #undef UVEDITOR_CHECK_AND_RETURN_ON_FAILURE
 
 // Explicit instantiations
-template bool UVUnwrapMeshUtil::DoesUnwrapMatchOverlay<EValidityCheckFailMode::ReturnOnly>
+template DYNAMICMESH_API bool UVUnwrapMeshUtil::DoesUnwrapMatchOverlay<EValidityCheckFailMode::ReturnOnly>
 (const FDynamicMeshUVOverlay& Overlay, const FDynamicMesh3& UnwrapMesh,
 	TFunctionRef<FVector3d(const FVector2f&)> UVToVertPosition, double Tolerance);
-template bool UVUnwrapMeshUtil::DoesUnwrapMatchOverlay<EValidityCheckFailMode::Check>
+template DYNAMICMESH_API bool UVUnwrapMeshUtil::DoesUnwrapMatchOverlay<EValidityCheckFailMode::Check>
 (const FDynamicMeshUVOverlay& Overlay, const FDynamicMesh3& UnwrapMesh,
 	TFunctionRef<FVector3d(const FVector2f&)> UVToVertPosition, double Tolerance);
-template bool UVUnwrapMeshUtil::DoesUnwrapMatchOverlay<EValidityCheckFailMode::Ensure>
+template DYNAMICMESH_API bool UVUnwrapMeshUtil::DoesUnwrapMatchOverlay<EValidityCheckFailMode::Ensure>
 (const FDynamicMeshUVOverlay& Overlay, const FDynamicMesh3& UnwrapMesh,
 	TFunctionRef<FVector3d(const FVector2f&)> UVToVertPosition, double Tolerance);
diff --git a/Engine/Plugins/Runtime/GeometryProcessing/Source/DynamicMesh/Public/Parameterization/DynamicMeshUVEditor.h b/Engine/Plugins/Runtime/GeometryProcessing/Source/DynamicMesh/Public/Parameterization/DynamicMeshUVEditor.h
--- a/Engine/Plugins/Runtime/GeometryProcessing/Source/DynamicMesh/Public/Parameterization/DynamicMeshUVEditor.h
+++ b/Engine/Plugins/Runtime/GeometryProcessing/Source/DynamicMesh/Public/Parameterization/DynamicMeshUVEditor.h
@@ -198,6 +198,12 @@ public:
 	 */
 	bool SetTriangleUVsFromFreeBoundarySpectralConformal(const TArray<int32>& Triangles, bool bUseExistingUVTopology, bool bPreserveIrregularity, FUVEditResult* Result = nullptr);
 
+	/**
+	 * Initializes any uninitialized triangles in the set to per-vertex (0,0) UVs. Doesn't clear or modify any existing 
+	 *  UVs, and only operates on the given set.
+	 */
+	void MakeSureUVsAreSet(const TSet<int32>& Triangles, FUVEditResult* Result = nullptr, TSet<int32>* ChangedTrianglesOut = nullptr);
+
 	/** 
 	 * Merge existing UV topology with a set of edges, removing seams at edges if they exist within the UV topology.
 	 * 
@@ -219,6 +225,21 @@ public:
 	 */
 	bool CreateSeamsAtEdges(const TSet<int32>& EidsToMakeIntoSeams, FUVEditResult* Result = nullptr);
 
+	/**
+	 * Takes the currently selected triangles and makes a separate UV island out of them, i.e. any interior seams
+	 *  are removed, and seams are added around the boundary of the selection. If the selection is not connected
+	 *  in the mesh, islands will be created for each connected component of selected triangles. If some of the 
+	 *  selected triangles have unset UVs, they will be initialized to per-vertex zero UVs for the purposes of
+	 *  creating an island.
+	 * 
+	 * @param TidsToMakeIntoIsland Triangles in the mesh
+	 * @param Result if non-null, list of new UV elements created for new seams, or to initialize UVs.
+	 * @param ChangedTrianglesOut Optional output of triangles that had their UVs changed. Currently this is a conservative
+	 *   superset, so some might not have actually changed. Not cleared before use.
+	 * @return true on success
+	 */
+	bool MakeIsland(const TSet<int32>& TidsToMakeIntoIsland, FUVEditResult* Result = nullptr, TSet<int32>* ChangedTrianglesOut = nullptr);
+
 	/**
 	 * Set UVs by box projection. Triangles will be grouped to "best" box face
 	 * PointTransform is applied to points before projectiong onto ProjectionFrame X/Y axes
