Add option to set material ID on hole fill triangles for plane cuts, and correspondingly for the plane cut and plane slice geometry script methods

diff --git a/Engine/Plugins/Runtime/GeometryProcessing/Source/DynamicMesh/Private/Operations/MeshPlaneCut.cpp b/Engine/Plugins/Runtime/GeometryProcessing/Source/DynamicMesh/Private/Operations/MeshPlaneCut.cpp
--- a/Engine/Plugins/Runtime/GeometryProcessing/Source/DynamicMesh/Private/Operations/MeshPlaneCut.cpp
+++ b/Engine/Plugins/Runtime/GeometryProcessing/Source/DynamicMesh/Private/Operations/MeshPlaneCut.cpp
@@ -777,7 +777,7 @@ bool FMeshPlaneCut::MinimalHoleFill(int ConstantGroupID)
 
 
 
-bool FMeshPlaneCut::HoleFill(TFunction<TArray<FIndex3i>(const FGeneralPolygon2d&)> PlanarTriangulationFunc, bool bFillSpans, int ConstantGroupID)
+bool FMeshPlaneCut::HoleFill(TFunction<TArray<FIndex3i>(const FGeneralPolygon2d&)> PlanarTriangulationFunc, bool bFillSpans, int ConstantGroupID, int MaterialID)
 {
 	bool bAllOk = true;
 
@@ -802,7 +802,6 @@ bool FMeshPlaneCut::HoleFill(TFunction<TArray<FIndex3i>(const FGeneralPolygon2d&
 		int GID = ConstantGroupID >= 0 ? ConstantGroupID : Mesh->AllocateTriangleGroup();
 		bool bFullyFilledHole = Filler.Fill(GID);
 
-		HoleFillTriangles.Add(Filler.NewTriangles);
 		if (Mesh->HasAttributes())
 		{
 			FDynamicMeshEditor Editor(Mesh);
@@ -813,8 +812,19 @@ bool FMeshPlaneCut::HoleFill(TFunction<TArray<FIndex3i>(const FGeneralPolygon2d&
 			{
 				Editor.SetTriangleUVsFromProjection(Filler.NewTriangles, ProjectionFrame, UVScaleFactor, FVector2f::Zero(), true, UVLayerIdx);
 			}
+
+			if (MaterialID > -1 && Mesh->Attributes()->HasMaterialID())
+			{
+				FDynamicMeshMaterialAttribute* MaterialAttrib = Mesh->Attributes()->GetMaterialID();
+				for (int32 TID : Filler.NewTriangles)
+				{
+					MaterialAttrib->SetValue(TID, MaterialID);
+				}
+			}
 		}
 
+		HoleFillTriangles.Add(MoveTemp(Filler.NewTriangles));
+
 		bAllOk = bAllOk && bFullyFilledHole;
 	}
 
diff --git a/Engine/Plugins/Runtime/GeometryProcessing/Source/DynamicMesh/Public/Operations/MeshPlaneCut.h b/Engine/Plugins/Runtime/GeometryProcessing/Source/DynamicMesh/Public/Operations/MeshPlaneCut.h
--- a/Engine/Plugins/Runtime/GeometryProcessing/Source/DynamicMesh/Public/Operations/MeshPlaneCut.h
+++ b/Engine/Plugins/Runtime/GeometryProcessing/Source/DynamicMesh/Public/Operations/MeshPlaneCut.h
@@ -185,7 +185,7 @@ public:
 	/**
 	 *  Fill cut loops with FPlanarHoleFiller, using a caller-provided triangulation function
 	 */
-	virtual bool HoleFill(TFunction<TArray<FIndex3i>(const FGeneralPolygon2d&)> PlanarTriangulationFunc, bool bFillSpans, int ConstantGroupID = -1);
+	virtual bool HoleFill(TFunction<TArray<FIndex3i>(const FGeneralPolygon2d&)> PlanarTriangulationFunc, bool bFillSpans, int ConstantGroupID = -1, int MaterialID = -1);
 
 	
 	virtual void TransferTriangleLabelsToHoleFillTriangles(TDynamicMeshScalarTriangleAttribute<int>* TriLabels);
diff --git a/Engine/Plugins/Runtime/GeometryScripting/Source/GeometryScriptingCore/Private/MeshBooleanFunctions.cpp b/Engine/Plugins/Runtime/GeometryScripting/Source/GeometryScriptingCore/Private/MeshBooleanFunctions.cpp
--- a/Engine/Plugins/Runtime/GeometryScripting/Source/GeometryScriptingCore/Private/MeshBooleanFunctions.cpp
+++ b/Engine/Plugins/Runtime/GeometryScripting/Source/GeometryScriptingCore/Private/MeshBooleanFunctions.cpp
@@ -207,7 +207,7 @@ UDynamicMesh* UGeometryScriptLibrary_MeshBooleanFunctions::ApplyMeshPlaneCut(
 
 		if (Options.bFillHoles)
 		{
-			Cut.HoleFill(ConstrainedDelaunayTriangulate<double>, Options.bFillSpans);
+			Cut.HoleFill(ConstrainedDelaunayTriangulate<double>, Options.bFillSpans, -1, Options.HoleFillMaterialID);
 		}
 
 	}, EDynamicMeshChangeType::GeneralEdit, EDynamicMeshAttributeChangeFlags::Unknown, false);
@@ -274,7 +274,7 @@ UDynamicMesh* UGeometryScriptLibrary_MeshBooleanFunctions::ApplyMeshPlaneSlice(
 
 		if (Options.bFillHoles)
 		{
-			Cut.HoleFill(ConstrainedDelaunayTriangulate<double>, Options.bFillSpans);
+			Cut.HoleFill(ConstrainedDelaunayTriangulate<double>, Options.bFillSpans, -1, Options.HoleFillMaterialID);
 		}
 
 		EditMesh.Attributes()->RemoveAttribute(ObjectIndexAttribute);
diff --git a/Engine/Plugins/Runtime/GeometryScripting/Source/GeometryScriptingCore/Public/GeometryScript/MeshBooleanFunctions.h b/Engine/Plugins/Runtime/GeometryScripting/Source/GeometryScriptingCore/Public/GeometryScript/MeshBooleanFunctions.h
--- a/Engine/Plugins/Runtime/GeometryScripting/Source/GeometryScriptingCore/Public/GeometryScript/MeshBooleanFunctions.h
+++ b/Engine/Plugins/Runtime/GeometryScripting/Source/GeometryScriptingCore/Public/GeometryScript/MeshBooleanFunctions.h
@@ -70,6 +70,10 @@ public:
 	UPROPERTY(EditAnywhere, BlueprintReadWrite, Category = Options)
 	bool bFillHoles = true;
 
+	// If non-negative and bFillHoles is true, specify the material ID to set on hole fill triangles
+	UPROPERTY(EditAnywhere, BlueprintReadWrite, Category = Options)
+	int32 HoleFillMaterialID = INDEX_NONE;
+
 	UPROPERTY(EditAnywhere, BlueprintReadWrite, Category = Options)
 	bool bFillSpans = true;
 
@@ -89,6 +93,10 @@ public:
 	UPROPERTY(EditAnywhere, BlueprintReadWrite, Category = Options)
 	bool bFillHoles = true;
 
+	// If non-negative and bFillHoles is true, specify the material ID to set on hole fill triangles
+	UPROPERTY(EditAnywhere, BlueprintReadWrite, Category = Options)
+	int32 HoleFillMaterialID = INDEX_NONE;
+
 	UPROPERTY(EditAnywhere, BlueprintReadWrite, Category = Options)
 	bool bFillSpans = true;
 
