Modeling Tools: Make mesh element selection conversions stable

diff --git a/Engine/Plugins/Runtime/GeometryProcessing/Source/DynamicMesh/Private/Selections/GeometrySelectionUtil.cpp b/Engine/Plugins/Runtime/GeometryProcessing/Source/DynamicMesh/Private/Selections/GeometrySelectionUtil.cpp
--- a/Engine/Plugins/Runtime/GeometryProcessing/Source/DynamicMesh/Private/Selections/GeometrySelectionUtil.cpp
+++ b/Engine/Plugins/Runtime/GeometryProcessing/Source/DynamicMesh/Private/Selections/GeometrySelectionUtil.cpp
@@ -1563,12 +1563,21 @@ bool UE::Geometry::InitializeSelectionFromTriangles(
 	return false;
 }
 
-
 bool UE::Geometry::ConvertSelection(
 	const UE::Geometry::FDynamicMesh3& Mesh,
 	const FGroupTopology* GroupTopology,
 	const FGeometrySelection& FromSelectionIn,
 	FGeometrySelection& ToSelectionOut)
+{
+	return ConvertSelection(Mesh, GroupTopology, FromSelectionIn, ToSelectionOut, EEnumerateSelectionConversionParams::ContainSelection);
+}
+
+bool UE::Geometry::ConvertSelection(
+	const UE::Geometry::FDynamicMesh3& Mesh,
+	const FGroupTopology* GroupTopology,
+	const FGeometrySelection& FromSelectionIn,
+	FGeometrySelection& ToSelectionOut,
+	const EEnumerateSelectionConversionParams ConversionParams)
 {
 	using namespace GeometrySelectionUtilLocals;
 
@@ -1576,7 +1585,8 @@ bool UE::Geometry::ConvertSelection(
 		const FDynamicMesh3& Mesh,
 		const FGroupTopology* GroupTopology,
 		const FGeometrySelection& FromSelectionIn,
-		FGeometrySelection& ToSelectionOut) -> bool
+		FGeometrySelection& ToSelectionOut,
+		const EEnumerateSelectionConversionParams ConversionParams) -> bool
 	{
 		checkSlow(FromSelectionIn.IsSameType(ToSelectionOut));
 
@@ -1589,24 +1599,103 @@ bool UE::Geometry::ConvertSelection(
 		const FDynamicMesh3& Mesh,
 		const FGroupTopology* GroupTopology,
 		const FGeometrySelection& FromSelectionIn,
-		FGeometrySelection& ToSelectionOut) -> bool
+		FGeometrySelection& ToSelectionOut,
+		const EEnumerateSelectionConversionParams ConversionParams) -> bool
 	{
 		checkSlow(ToSelectionOut.TopologyType == EGeometryTopologyType::Triangle);
 		checkSlow(ToSelectionOut.ElementType == EGeometryElementType::Face);
 
 		const FPolygroupSet GroupSet = FPolygroupSet(&Mesh);
-		return EnumerateSelectionTriangles(FromSelectionIn, Mesh,
-		[&ToSelectionOut](const int32 TID)
+
+		// TODO: these if statements could be consolidated to minimize minor repeat code; however its currently set up this way for readability and
+		//		 scalability for if/when more EEnumerateSelectionConversionParams are added
+
+		if (ConversionParams == EEnumerateSelectionConversionParams::ExpandSelection)
 		{
-			ToSelectionOut.Selection.Add(FGeoSelectionID::MeshTriangle(TID).Encoded());
-		}, &GroupSet);
+			EnumerateSelectionTriangles(
+				FromSelectionIn,
+				Mesh,
+				[&ToSelectionOut](const int32 TID) { ToSelectionOut.Selection.Add(FGeoSelectionID::MeshTriangle(TID).Encoded()); },
+				&GroupSet);
+		}
+
+		else if (ConversionParams == EEnumerateSelectionConversionParams::ContainSelection)
+		{
+			// where FromSelection type is either PolyVerts, PolyEdges, PolyFaces, TriFaces -> ensuring stable selection is not applicable and can do same as in ExpandSelection
+			if (FromSelectionIn.TopologyType == EGeometryTopologyType::Polygroup
+				|| ( FromSelectionIn.ElementType == EGeometryElementType::Face && FromSelectionIn.TopologyType == EGeometryTopologyType::Triangle))
+			{
+				EnumerateSelectionTriangles(
+				FromSelectionIn,
+				Mesh,
+				[&ToSelectionOut](const int32 TID) { ToSelectionOut.Selection.Add(FGeoSelectionID::MeshTriangle(TID).Encoded()); },
+				&GroupSet);
+			}
+			// where FromSelection type is either TriEdge or TriVert
+			else
+			{
+				TArray<int32> AllTriIDsConnectedToSelection;
+				
+				// retrieve all the triangles connected to the selected verts/edges but they will not all be included in final selection
+				EnumerateSelectionTriangles(
+					FromSelectionIn,
+					Mesh,
+					[&AllTriIDsConnectedToSelection](const int32 TID) { return AllTriIDsConnectedToSelection.Add(TID); },
+					&GroupSet);
+
+
+				if (FromSelectionIn.ElementType == EGeometryElementType::Edge)
+				{
+					// retrieve all the edges currently in the selection
+					TSet<int32> SelectedEdges;
+					for (const uint64 EncodedID : FromSelectionIn.Selection)
+					{
+						const FMeshTriEdgeID TriEdgeID( FGeoSelectionID(EncodedID).GeometryID );
+						const int32 EdgeID = Mesh.IsTriangle(TriEdgeID.TriangleID) ? Mesh.GetTriEdge(TriEdgeID.TriangleID, TriEdgeID.TriEdgeIndex) : IndexConstants::InvalidID;
+						SelectedEdges.Add(EdgeID);
+					}
+
+					// for each triangle that's connected to the selection, determine if it is formed by 3 selected edges or not
+					for (const int TriangleID : AllTriIDsConnectedToSelection)
+					{
+						const FIndex3i TriEdges = Mesh.GetTriEdges(TriangleID);
+						if (SelectedEdges.Contains(TriEdges.A) && SelectedEdges.Contains(TriEdges.B) && SelectedEdges.Contains(TriEdges.C))
+						{
+							ToSelectionOut.Selection.Add(FGeoSelectionID::MeshTriangle(TriangleID).Encoded());
+						}
+					}
+				}
+
+				if (FromSelectionIn.ElementType == EGeometryElementType::Vertex)
+				{
+					// retrieve all the vertices currently in the selection
+					TSet<int32> SelectedVerts;
+					for (const uint64 VertexID : FromSelectionIn.Selection)
+					{
+						SelectedVerts.Add((int32)VertexID);
+					}
+					
+					// for each triangle that's connected to the selection, determine if it is formed by 3 selected vertices or not
+					for (const int TriangleID : AllTriIDsConnectedToSelection)
+					{
+						const FIndex3i TriVerts = Mesh.GetTriangle(TriangleID);
+						if (SelectedVerts.Contains(TriVerts.A) && SelectedVerts.Contains(TriVerts.B) && SelectedVerts.Contains(TriVerts.C))
+						{
+							ToSelectionOut.Selection.Add(FGeoSelectionID::MeshTriangle(TriangleID).Encoded());
+						}
+					}
+				}
+			}
+		}
+		return true;
 	};
 
 	const auto ToTriEdge = [](
 		const FDynamicMesh3& Mesh,
 		const FGroupTopology* GroupTopology,
 		const FGeometrySelection& FromSelectionIn,
-		FGeometrySelection& ToSelectionOut) -> bool
+		FGeometrySelection& ToSelectionOut,
+		const EEnumerateSelectionConversionParams ConversionParams) -> bool
 	{
 		checkSlow(ToSelectionOut.TopologyType == EGeometryTopologyType::Triangle);
 		checkSlow(ToSelectionOut.ElementType == EGeometryElementType::Edge);
@@ -1620,13 +1709,60 @@ bool UE::Geometry::ConvertSelection(
 				});
 		};
 
-		if (FromSelectionIn.TopologyType == EGeometryTopologyType::Triangle)
+		// TODO: these if statements could be consolidated to minimize repeat code; however its currently set up this way for readability and
+		//		 scalability for if/when more EEnumerateSelectionConversionParams are added
+		
+		if (ConversionParams == EEnumerateSelectionConversionParams::ExpandSelection)
 		{
-			return EnumerateTriangleSelectionEdges(FromSelectionIn, Mesh, SelectEdgesFunc);
+			if (FromSelectionIn.TopologyType == EGeometryTopologyType::Triangle)
+			{
+				EnumerateTriangleSelectionEdges(FromSelectionIn, Mesh, SelectEdgesFunc);
+			}
+			else if (FromSelectionIn.TopologyType == EGeometryTopologyType::Polygroup)
+			{
+				EnumeratePolygroupSelectionEdges(FromSelectionIn, Mesh, *GroupTopology, SelectEdgesFunc);
+			}
 		}
-		else if (FromSelectionIn.TopologyType == EGeometryTopologyType::Polygroup)
+
+		else if (ConversionParams == EEnumerateSelectionConversionParams::ContainSelection)
 		{
-			return EnumeratePolygroupSelectionEdges(FromSelectionIn, Mesh, *GroupTopology, SelectEdgesFunc);
+			if (FromSelectionIn.TopologyType == EGeometryTopologyType::Triangle)
+			{
+				// in the case of TriVert->TriEdge, ensure that only edges where both verts are in the initial selection are included
+				if (FromSelectionIn.ElementType == EGeometryElementType::Vertex)
+				{
+					TSet<uint64> SelectedVerts;
+					TArray<int32> AllEdgesConnectedToVerts;
+					for (const uint64 VertexID : FromSelectionIn.Selection)
+					{
+						SelectedVerts.Add(VertexID);
+						Mesh.EnumerateVertexEdges((int32)VertexID, [&SelectEdgesFunc, &AllEdgesConnectedToVerts](const int32 EdgeID)
+						{
+							AllEdgesConnectedToVerts.Add(EdgeID);
+						});
+					}
+
+					// ensure stable selection by only selecting the edges where BOTH its verts were in init selection
+					// however a selection which includes a single vert or any verts without any of their adjacent verts selected will be lost in conversion
+					for (const int32 EdgeID : AllEdgesConnectedToVerts)
+					{
+						const FIndex2i EdgeVerts = Mesh.GetEdgeV(EdgeID);
+						if (SelectedVerts.Contains(EdgeVerts.A) && SelectedVerts.Contains(EdgeVerts.B))
+						{
+							SelectEdgesFunc(EdgeID);
+						}
+					}
+				}
+				else // case of converting from TriEdge or TriFace -> ensuring stable selection is not applicable and can do same as in ExpandSelection
+				{
+					EnumerateTriangleSelectionEdges(FromSelectionIn, Mesh, SelectEdgesFunc);
+				}
+			}
+			// case of converting from PolyVert, PolyEdge, PolyFace -> ensuring stable selection is not applicable and can do same as in ExpandSelection
+			else if (FromSelectionIn.TopologyType == EGeometryTopologyType::Polygroup)
+			{
+				EnumeratePolygroupSelectionEdges(FromSelectionIn, Mesh, *GroupTopology, SelectEdgesFunc);
+			}
 		}
 		return true;
 	};
@@ -1635,7 +1771,8 @@ bool UE::Geometry::ConvertSelection(
 		const FDynamicMesh3& Mesh,
 		const FGroupTopology* GroupTopology,
 		const FGeometrySelection& FromSelectionIn,
-		FGeometrySelection& ToSelectionOut) -> bool
+		FGeometrySelection& ToSelectionOut,
+		const EEnumerateSelectionConversionParams ConversionParams) -> bool
 	{
 		checkSlow(ToSelectionOut.TopologyType == EGeometryTopologyType::Triangle);
 		checkSlow(ToSelectionOut.ElementType == EGeometryElementType::Vertex);
@@ -1665,7 +1802,8 @@ bool UE::Geometry::ConvertSelection(
 		const FDynamicMesh3& Mesh,
 		const FGroupTopology* GroupTopology,
 		const FGeometrySelection& FromSelectionIn,
-		FGeometrySelection& ToSelectionOut) -> bool
+		FGeometrySelection& ToSelectionOut,
+		const EEnumerateSelectionConversionParams ConversionParams) -> bool
 	{
 		checkSlow(ToSelectionOut.TopologyType == EGeometryTopologyType::Polygroup);
 		checkSlow(ToSelectionOut.ElementType == EGeometryElementType::Face);
@@ -1688,7 +1826,8 @@ bool UE::Geometry::ConvertSelection(
 		const FDynamicMesh3& Mesh,
 		const FGroupTopology* GroupTopology,
 		const FGeometrySelection& FromSelectionIn,
-		FGeometrySelection& ToSelectionOut)
+		FGeometrySelection& ToSelectionOut,
+		const EEnumerateSelectionConversionParams ConversionParams)
 	{
 		checkSlow(ToSelectionOut.TopologyType == EGeometryTopologyType::Polygroup);
 		checkSlow(ToSelectionOut.ElementType == EGeometryElementType::Edge);
@@ -1736,7 +1875,8 @@ bool UE::Geometry::ConvertSelection(
 		const FDynamicMesh3& Mesh,
 		const FGroupTopology* GroupTopology,
 		const FGeometrySelection& FromSelectionIn,
-		FGeometrySelection& ToSelectionOut) -> bool
+		FGeometrySelection& ToSelectionOut,
+		const EEnumerateSelectionConversionParams ConversionParams) -> bool
 	{
 		checkSlow(ToSelectionOut.TopologyType == EGeometryTopologyType::Polygroup);
 		checkSlow(ToSelectionOut.ElementType == EGeometryElementType::Vertex);
@@ -1780,10 +1920,11 @@ bool UE::Geometry::ConvertSelection(
 			const FDynamicMesh3& Mesh,
 			const FGroupTopology* GroupTopology,
 			const FGeometrySelection& FromSelectionIn,
-			FGeometrySelection& ToSelectionOut);
+			FGeometrySelection& ToSelectionOut,
+			const EEnumerateSelectionConversionParams ConversionParams);
 
 	constexpr ConvertSelectionFunc ConvertFuncs[6][6] = {
-		{FromTypeToSame,  ToTriEdge,      ToTriFace,      ToPolyVtx,     ToPolyEdge,     ToPolyFace },
+		{FromTypeToSame, ToTriEdge,      ToTriFace,      ToPolyVtx,     ToPolyEdge,     ToPolyFace },
 		{ToTriVtx,		  FromTypeToSame, ToTriFace,      ToPolyVtx,     ToPolyEdge,     ToPolyFace },
 		{ToTriVtx,		  ToTriEdge,      FromTypeToSame, ToPolyVtx,     ToPolyEdge,     ToPolyFace },
 		{ToTriVtx,		  ToTriEdge,      ToTriFace,      FromTypeToSame,ToPolyEdge,     ToPolyFace },
@@ -1794,7 +1935,7 @@ bool UE::Geometry::ConvertSelection(
 	const int FromIndex = GetSelectionTypeAsIndex(FromSelectionIn);
 	const int ToIndex =   GetSelectionTypeAsIndex(ToSelectionOut);
 
-	return ConvertFuncs[FromIndex][ToIndex](Mesh, GroupTopology, FromSelectionIn, ToSelectionOut);
+	return ConvertFuncs[FromIndex][ToIndex](Mesh, GroupTopology, FromSelectionIn, ToSelectionOut, ConversionParams);
 }
 
 bool UE::Geometry::ConvertTriangleSelectionToOverlaySelection(
@@ -1930,7 +2071,7 @@ bool UE::Geometry::ConvertPolygroupSelectionToIncidentOverlaySelection(
 		IncidentSelection->InitializeTypes(EGeometryElementType::Vertex, EGeometryTopologyType::Triangle);
 
 		// GroupTopology argument is ignored if MeshSelection has Triangle topology
-		bool bSuccess = ConvertSelection(Mesh, &GroupTopology, MeshSelection, *IncidentSelection);
+		bool bSuccess = ConvertSelection(Mesh, &GroupTopology, MeshSelection, *IncidentSelection, EEnumerateSelectionConversionParams::ContainSelection);
 		ensure(bSuccess == true);
 		ensure(!IncidentSelection->IsEmpty());
 
diff --git a/Engine/Plugins/Runtime/GeometryProcessing/Source/DynamicMesh/Public/Selections/GeometrySelectionUtil.h b/Engine/Plugins/Runtime/GeometryProcessing/Source/DynamicMesh/Public/Selections/GeometrySelectionUtil.h
--- a/Engine/Plugins/Runtime/GeometryProcessing/Source/DynamicMesh/Public/Selections/GeometrySelectionUtil.h
+++ b/Engine/Plugins/Runtime/GeometryProcessing/Source/DynamicMesh/Public/Selections/GeometrySelectionUtil.h
@@ -30,6 +30,14 @@ enum class EEnumerateSelectionMapping : uint8
 };
 ENUM_CLASS_FLAGS(EEnumerateSelectionMapping);
 	
+enum class EEnumerateSelectionConversionParams
+{
+	/** Converts to components whose constituent parts are all contained in the initial selection. */
+	ContainSelection,
+	/** Converts to components where any of its constituent parts are included in the initial selection. */
+	ExpandSelection
+};
+	
 /**
  * Test if SelectionA and SelectionB are the same selection.
  * This is currently relatively expensive on Polygroup selections due to how they are encoded
@@ -364,12 +372,20 @@ DYNAMICMESH_API bool InitializeSelectionFromTriangles(
  *
  * @return true if conversion is supported and was computed successfully, return false otherwise
  */
+UE_DEPRECATED(5.5, "This version of ConvertSelection is deprecated, please use the version with the EENumerateSelectionConversionParams parameter going forward.")	
 DYNAMICMESH_API bool ConvertSelection(
 	const UE::Geometry::FDynamicMesh3& Mesh,
 	const FGroupTopology* GroupTopology,
 	const FGeometrySelection& FromSelectionIn,
 	FGeometrySelection& ToSelectionOut);
 
+DYNAMICMESH_API bool ConvertSelection(
+	const UE::Geometry::FDynamicMesh3& Mesh,
+	const FGroupTopology* GroupTopology,
+	const FGeometrySelection& FromSelectionIn,
+	FGeometrySelection& ToSelectionOut,
+	const EEnumerateSelectionConversionParams ConversionParams);
+
 /**
  * Convert the given MeshSelection to a list of Triangles and Vertices into the Mesh,
  * which can be used to represent a selection of overlay elements. This is always possible
diff --git a/Engine/Plugins/Runtime/MeshModelingToolset/Source/ModelingComponents/Private/Selection/DynamicMeshSelector.cpp b/Engine/Plugins/Runtime/MeshModelingToolset/Source/ModelingComponents/Private/Selection/DynamicMeshSelector.cpp
--- a/Engine/Plugins/Runtime/MeshModelingToolset/Source/ModelingComponents/Private/Selection/DynamicMeshSelector.cpp
+++ b/Engine/Plugins/Runtime/MeshModelingToolset/Source/ModelingComponents/Private/Selection/DynamicMeshSelector.cpp
@@ -192,7 +192,7 @@ void FBaseDynamicMeshSelector::UpdateSelectionFromSelection(
 	{ 
 		FGeometrySelection TempSelection;
 		TempSelection.InitializeTypes(SelectionEditor.GetElementType(), SelectionEditor.GetTopologyType());
-		bConverted = UE::Geometry::ConvertSelection(Mesh, GetGroupTopology(), FromSelection, TempSelection);		// todo do not always need group topology here...
+		bConverted = UE::Geometry::ConvertSelection(Mesh, GetGroupTopology(), FromSelection, TempSelection, EEnumerateSelectionConversionParams::ContainSelection);		// todo do not always need group topology here...
 		if ( bConverted )
 		{
 			UE::Geometry::UpdateSelectionWithNewElements(&SelectionEditor, UpdateConfig.ChangeType,
