Modeling Tools: Fixes warnings with previously submitted charge regarding making Mesh Element Selection conversions stable

diff --git a/Engine/Plugins/Runtime/GeometryProcessing/Source/DynamicMesh/Private/Selections/GeometrySelectionUtil.cpp b/Engine/Plugins/Runtime/GeometryProcessing/Source/DynamicMesh/Private/Selections/GeometrySelectionUtil.cpp
--- a/Engine/Plugins/Runtime/GeometryProcessing/Source/DynamicMesh/Private/Selections/GeometrySelectionUtil.cpp
+++ b/Engine/Plugins/Runtime/GeometryProcessing/Source/DynamicMesh/Private/Selections/GeometrySelectionUtil.cpp
@@ -1656,7 +1656,7 @@ bool UE::Geometry::ConvertSelection(
 					}
 
 					// for each triangle that's connected to the selection, determine if it is formed by 3 selected edges or not
-					for (const int TriangleID : AllTriIDsConnectedToSelection)
+					for (const int TriangleID : AllTriIDsConnectedToSelection) //-V1078
 					{
 						const FIndex3i TriEdges = Mesh.GetTriEdges(TriangleID);
 						if (SelectedEdges.Contains(TriEdges.A) && SelectedEdges.Contains(TriEdges.B) && SelectedEdges.Contains(TriEdges.C))
@@ -1676,7 +1676,7 @@ bool UE::Geometry::ConvertSelection(
 					}
 					
 					// for each triangle that's connected to the selection, determine if it is formed by 3 selected vertices or not
-					for (const int TriangleID : AllTriIDsConnectedToSelection)
+					for (const int TriangleID : AllTriIDsConnectedToSelection) //-V1078
 					{
 						const FIndex3i TriVerts = Mesh.GetTriangle(TriangleID);
 						if (SelectedVerts.Contains(TriVerts.A) && SelectedVerts.Contains(TriVerts.B) && SelectedVerts.Contains(TriVerts.C))
@@ -1744,7 +1744,7 @@ bool UE::Geometry::ConvertSelection(
 
 					// ensure stable selection by only selecting the edges where BOTH its verts were in init selection
 					// however a selection which includes a single vert or any verts without any of their adjacent verts selected will be lost in conversion
-					for (const int32 EdgeID : AllEdgesConnectedToVerts)
+					for (const int32 EdgeID : AllEdgesConnectedToVerts) //-V1078
 					{
 						const FIndex2i EdgeVerts = Mesh.GetEdgeV(EdgeID);
 						if (SelectedVerts.Contains(EdgeVerts.A) && SelectedVerts.Contains(EdgeVerts.B))
