GeometryProcessing Sparse Matrix: fix a crash when solving large sparse systems. We were computing NRows * NCols which can overflow int32 on moderately large matrices (~50k by 50k).

diff --git a/Engine/Plugins/Runtime/GeometryProcessing/Source/DynamicMesh/Private/Solvers/Internal/SparseMatrix.cpp b/Engine/Plugins/Runtime/GeometryProcessing/Source/DynamicMesh/Private/Solvers/Internal/SparseMatrix.cpp
--- a/Engine/Plugins/Runtime/GeometryProcessing/Source/DynamicMesh/Private/Solvers/Internal/SparseMatrix.cpp
+++ b/Engine/Plugins/Runtime/GeometryProcessing/Source/DynamicMesh/Private/Solvers/Internal/SparseMatrix.cpp
@@ -71,7 +71,9 @@ bool SliceSparseMatrix(const FSparseMatrixD& InMatrix,
 
     // Take a guess at the number of the nonzero elements in the sliced matrix
     std::vector<Eigen::Triplet<FSparseMatrixD::Scalar>> Triplets;
-    const float NonZeroProbabilitytPerEntry = static_cast<float>(NumNonZeros) / static_cast<float>(NumRows * NumCols);
+    const float NonZeroProbabilitytPerEntry = static_cast<float>(NumNonZeros) / static_cast<float>(NumRows) / static_cast<float>(NumCols);	// ~= NumNonZeros / (NumRows*NumCols) but 
+																																			// won't hit int32 overflow on large systems
+
     const float SlicedNumEntries = static_cast<float>(NumRowsSliced * NumColsSliced);
     const int SlicedNonZeroEstimate = static_cast<int>(NonZeroProbabilitytPerEntry * SlicedNumEntries);
     Triplets.reserve(SlicedNonZeroEstimate);
