fix any UV bowties that would be introduced when performing box or cylinder UV projections

diff --git a/Engine/Plugins/Runtime/GeometryProcessing/Source/DynamicMesh/Private/Parameterization/DynamicMeshUVEditor.cpp b/Engine/Plugins/Runtime/GeometryProcessing/Source/DynamicMesh/Private/Parameterization/DynamicMeshUVEditor.cpp
--- a/Engine/Plugins/Runtime/GeometryProcessing/Source/DynamicMesh/Private/Parameterization/DynamicMeshUVEditor.cpp
+++ b/Engine/Plugins/Runtime/GeometryProcessing/Source/DynamicMesh/Private/Parameterization/DynamicMeshUVEditor.cpp
@@ -1409,6 +1409,9 @@ void FDynamicMeshUVEditor::SetTriangleUVsFromBoxProjection(
 		UVOverlay->SetTriangle(tid, ElemTri);
 	}
 
+	// Above process can introduce bowties, so we split any bowties on new element IDs
+	SplitBowtiesOnUVElements(NewUVIndices, true);
+
 	if (Result != nullptr)
 	{
 		Result->NewUVElements = MoveTemp(NewUVIndices);
@@ -1417,7 +1420,20 @@ void FDynamicMeshUVEditor::SetTriangleUVsFromBoxProjection(
 
 
 
+void FDynamicMeshUVEditor::SplitBowtiesOnUVElements(TArray<int32>& UVElementIDs, bool bAddNewElementsToInputArray)
+{
+	if (!ensure(UVOverlay)) return;
 
+	const int32 InitialNumElements = UVElementIDs.Num();
+	for (int32 Idx = 0; Idx < InitialNumElements; ++Idx)
+	{
+		int32 ParentVID = UVOverlay->GetParentVertex(UVElementIDs[Idx]);
+		if (UVOverlay->IsBowtieInOverlay(ParentVID))
+		{
+			UVOverlay->SplitBowtiesAtVertex(ParentVID, bAddNewElementsToInputArray ? &UVElementIDs : nullptr);
+		}
+	}
+}
 
 
 void FDynamicMeshUVEditor::SetTriangleUVsFromCylinderProjection(
@@ -1547,6 +1563,9 @@ void FDynamicMeshUVEditor::SetTriangleUVsFromCylinderProjection(
 		UVOverlay->SetTriangle(tid, ElemTri);
 	}
 
+	// Above process can introduce bowties, so we split any bowties on new element IDs
+	SplitBowtiesOnUVElements(NewUVIndices, true);
+
 	if (Result != nullptr)
 	{
 		Result->NewUVElements = MoveTemp(NewUVIndices);
diff --git a/Engine/Plugins/Runtime/GeometryProcessing/Source/DynamicMesh/Public/Parameterization/DynamicMeshUVEditor.h b/Engine/Plugins/Runtime/GeometryProcessing/Source/DynamicMesh/Public/Parameterization/DynamicMeshUVEditor.h
--- a/Engine/Plugins/Runtime/GeometryProcessing/Source/DynamicMesh/Public/Parameterization/DynamicMeshUVEditor.h
+++ b/Engine/Plugins/Runtime/GeometryProcessing/Source/DynamicMesh/Public/Parameterization/DynamicMeshUVEditor.h
@@ -331,6 +331,9 @@ private:
 	 * Helper method to solve either the Least Squares Conformal Map or the Spectral Conformal Map problems.
 	 */
 	bool SetTriangleUVsFromConformal(const TArray<int32>& Triangles, const FSetUVsFromConformalOptions& Options, FUVEditResult* Result);
+
+	// Helper to split bowties on element IDs
+	void SplitBowtiesOnUVElements(TArray<int32>& UVElementIDs, bool bAddNewElementsToInputArray = true);
 	
 };
 
