[DynamicMesh]
* Changed template instantiation of TMeshSimplification to have dllexport on functions instead of instantiated type. This is more compatible with other compilers than msvc

diff --git a/Engine/Plugins/Runtime/GeometryProcessing/Source/DynamicMesh/Private/MeshSimplification.cpp b/Engine/Plugins/Runtime/GeometryProcessing/Source/DynamicMesh/Private/MeshSimplification.cpp
--- a/Engine/Plugins/Runtime/GeometryProcessing/Source/DynamicMesh/Private/MeshSimplification.cpp
+++ b/Engine/Plugins/Runtime/GeometryProcessing/Source/DynamicMesh/Private/MeshSimplification.cpp
@@ -1756,9 +1756,9 @@ namespace Geometry
 // These are explicit instantiations of the templates that are exported from the shared lib.
 // Only these instantiations of the template can be used.
 // This is necessary because we have placed most of the templated functions in this .cpp file, instead of the header.
-template class DYNAMICMESH_API TMeshSimplification< FAttrBasedQuadricErrord >;
-template class DYNAMICMESH_API TMeshSimplification< FVolPresQuadricErrord >;
-template class DYNAMICMESH_API TMeshSimplification< FQuadricErrord >;
+template class TMeshSimplification< FAttrBasedQuadricErrord >;
+template class TMeshSimplification< FVolPresQuadricErrord >;
+template class TMeshSimplification< FQuadricErrord >;
 
 } // end namespace UE::Geometry
 } // end namespace UE
\ No newline at end of file
diff --git a/Engine/Plugins/Runtime/GeometryProcessing/Source/DynamicMesh/Public/MeshSimplification.h b/Engine/Plugins/Runtime/GeometryProcessing/Source/DynamicMesh/Public/MeshSimplification.h
--- a/Engine/Plugins/Runtime/GeometryProcessing/Source/DynamicMesh/Public/MeshSimplification.h
+++ b/Engine/Plugins/Runtime/GeometryProcessing/Source/DynamicMesh/Public/MeshSimplification.h
@@ -117,25 +117,25 @@ public:
 	 * Note that because triangles are removed in pairs, the resulting count may be TriangleCount-1
 	 * @param TriangleCount the target triangle count
 	 */
-	virtual void SimplifyToTriangleCount(int TriangleCount);
+	DYNAMICMESH_API virtual void SimplifyToTriangleCount(int TriangleCount);
 
 	/**
 	 * Simplify mesh until it has a specific vertex count
 	 * @param VertexCount the target vertex count
 	 */
-	virtual void SimplifyToVertexCount(int VertexCount);
+	DYNAMICMESH_API virtual void SimplifyToVertexCount(int VertexCount);
 
 	/**
 	 * Simplify mesh until no edges smaller than min length remain. This is not a great criteria.
 	 * @param MinEdgeLength collapse any edge longer than this
 	 */
-	virtual void SimplifyToEdgeLength(double MinEdgeLength);
+	DYNAMICMESH_API virtual void SimplifyToEdgeLength(double MinEdgeLength);
 
 	/**
 	 * Simplify mesh until the quadric error of an edge collapse exceeds the specified criteria.
 	 * @param MaxError collapse an edge if the corresponding quadric error exceeds this 
 	 */
-	virtual void SimplifyToMaxError(double MaxError);
+	DYNAMICMESH_API virtual void SimplifyToMaxError(double MaxError);
 
 	/**
 	 * Maximally collapse mesh in a way that does not change shape at all.
@@ -143,7 +143,7 @@ public:
 	 * @param AngleTolDeg two triangles are considered coplanar if their normals are within this angle tolerance
 	 * @param EdgeFilterPredicate only edges that pass this predicate will be considered for collapse. Default all true.
 	 */
-	virtual void SimplifyToMinimalPlanar(
+	DYNAMICMESH_API virtual void SimplifyToMinimalPlanar(
 		double CoplanarAngleTolDeg = 0.001,
 		TFunctionRef<bool(int32 EdgeID)> EdgeFilterPredicate = [](int32) { return true; }
 		);
@@ -157,7 +157,7 @@ public:
 	 * @param MeshIsClosedHint if you know the mesh is closed, this pass this true to avoid some precomputes
 	 * @param MinTriangleCount halt fast collapse if mesh falls below this triangle count
 	 */
-	virtual void FastCollapsePass(double MinEdgeLength, int Rounds = 1, bool bMeshIsClosedHint = false, uint32 MinTriangleCount = 0);
+	DYNAMICMESH_API virtual void FastCollapsePass(double MinEdgeLength, int Rounds = 1, bool bMeshIsClosedHint = false, uint32 MinTriangleCount = 0);
 
 	
 
@@ -185,7 +185,7 @@ protected:
 
 
 	/** Top-level function that does the simplification */
-	virtual void DoSimplify();
+	DYNAMICMESH_API virtual void DoSimplify();
 
 
 
@@ -214,14 +214,14 @@ protected:
 	double SeamEdgeWeight =  256.;
 
 	TArray<FQuadricErrorType> vertQuadrics;
-	virtual void InitializeVertexQuadrics();
+	DYNAMICMESH_API virtual void InitializeVertexQuadrics();
 
 	TMap<int, FSeamQuadricType> seamQuadrics;
-	virtual void InitializeSeamQuadrics();
+	DYNAMICMESH_API virtual void InitializeSeamQuadrics();
 
 	TArray<double> triAreas;
 	TArray<FQuadricErrorType> triQuadrics;
-	virtual void InitializeTriQuadrics();
+	DYNAMICMESH_API virtual void InitializeTriQuadrics();
 
 	FDynamicMeshNormalOverlay* NormalOverlay;
 
@@ -261,7 +261,7 @@ protected:
 		}
 	};
 
-	virtual void InitializeQueue();
+	DYNAMICMESH_API virtual void InitializeQueue();
 
 	// return point that minimizes quadric error for edge [ea,eb]
 	FVector3d OptimalPoint(int eid, const FQuadricErrorType& q, int ea, int eb);
@@ -277,7 +277,7 @@ protected:
 
 
 	// update queue weight for each edge in vertex one-ring and rebuild and quadrics necessary 
-	virtual void UpdateNeighborhood(const FDynamicMesh3::FEdgeCollapseInfo& collapseInfo);
+	DYNAMICMESH_API virtual void UpdateNeighborhood(const FDynamicMesh3::FEdgeCollapseInfo& collapseInfo);
 
 
 	virtual void Reproject() 
@@ -339,14 +339,14 @@ protected:
 
 
 	// Project vertices onto projection target. 
-	virtual void FullProjectionPass();
+	DYNAMICMESH_API virtual void FullProjectionPass();
 
-	virtual void ProjectVertex(int vID, IProjectionTarget* targetIn);
+	DYNAMICMESH_API virtual void ProjectVertex(int vID, IProjectionTarget* targetIn);
 
 	// used by collapse-edge to get projected position for new vertex
-	virtual FVector3d GetProjectedCollapsePosition(int vid, const FVector3d& vNewPos);
+	DYNAMICMESH_API virtual FVector3d GetProjectedCollapsePosition(int vid, const FVector3d& vNewPos);
 
-	virtual void ApplyToProjectVertices(const TFunction<void(int)>& apply_f);
+	DYNAMICMESH_API virtual void ApplyToProjectVertices(const TFunction<void(int)>& apply_f);
 
 
 	/**
