add support for direct triangle edge selection in geometry script bevel and underlying c++ code

diff --git a/Engine/Plugins/Runtime/GeometryProcessing/Source/DynamicMesh/Private/Operations/MeshBevel.cpp b/Engine/Plugins/Runtime/GeometryProcessing/Source/DynamicMesh/Private/Operations/MeshBevel.cpp
--- a/Engine/Plugins/Runtime/GeometryProcessing/Source/DynamicMesh/Private/Operations/MeshBevel.cpp
+++ b/Engine/Plugins/Runtime/GeometryProcessing/Source/DynamicMesh/Private/Operations/MeshBevel.cpp
@@ -129,6 +129,210 @@ void FMeshBevel::InitializeFromGroupTopologyEdges(const FDynamicMesh3& Mesh, con
 	BuildVertexSets(Mesh);
 }
 
+void FMeshBevel::InitializeFromTriangleEdges(const FDynamicMesh3& Mesh, TConstArrayView<int32> TriangleEdges, TFunctionRef<bool(int32)> IsCornerVertex)
+{
+	ResultInfo = FGeometryResult(EGeometryResultType::InProgress);
+
+	// VID to TriangleEdges Index map, used for walking chains/loops
+	TMultiMap<int32, int32> V2EIdx;
+	TArray<uint8> EdgeUsed;
+	EdgeUsed.SetNumZeroed(TriangleEdges.Num());
+	TSet<int32> AllEdgeVIDs;
+
+	for (int32 Idx = 0; Idx < TriangleEdges.Num(); ++Idx)
+	{
+		int32 EID = TriangleEdges[Idx];
+		if (Mesh.IsBoundaryEdge(EID))
+		{
+			// bevel doesn't support boundary edges; mark them as processed without adding them to the bevel data
+			EdgeUsed[Idx] = 1;
+			continue;
+		}
+		FIndex2i VIDs = Mesh.GetEdgeV(EID);
+		V2EIdx.Add(VIDs.A, Idx);
+		V2EIdx.Add(VIDs.B, Idx);
+		AllEdgeVIDs.Add(VIDs.A);
+		AllEdgeVIDs.Add(VIDs.B);
+	}
+
+	// edge chains, stored as [VID,EID] pairs and a EndChain or EndLoop as the EID at the end of each continguous chain or loop respectively
+	constexpr int32 EndChain = -1, EndLoop = -2;
+	TArray<TPair<int32,int32>> Chains;
+
+	auto ProcessEdgeIdx = [&EdgeUsed, &TriangleEdges, &Mesh, &IsCornerVertex, &V2EIdx, &Chains, EndChain, EndLoop](int32 VID, int32 EdgeIdx)
+	{
+		EdgeUsed[EdgeIdx] = 1;
+		int32 WalkEID = TriangleEdges[EdgeIdx];
+		Chains.Emplace(VID, WalkEID);
+		int32 StartVID = VID;
+		int32 WalkVID = VID;
+		int32 MaxIters = TriangleEdges.Num() + 2; // safety to prevent infinite loop
+		while (--MaxIters > 0)
+		{
+			FIndex2i EdgeV = Mesh.GetEdgeV(WalkEID);
+			// step WalkVID to next vertex
+			WalkVID = EdgeV[1 - EdgeV.IndexOf(WalkVID)];
+			// detect loops
+			if (WalkVID == StartVID)
+			{
+				Chains.Emplace(WalkVID, EndLoop);
+				break;
+			}
+			// stop chains at specified corner vertices
+			if (IsCornerVertex(WalkVID))
+			{
+				Chains.Emplace(WalkVID, EndChain);
+				break;
+			}
+			// non-loop case, look for a next edge
+			int32 WalkEIdx = -1;
+			int32 ECount = 0;
+			for (TMultiMap<int32, int32>::TConstKeyIterator It = V2EIdx.CreateConstKeyIterator(WalkVID); It; ++It)
+			{
+				ECount++;
+				if (EdgeUsed[It.Value()] != 0)
+				{
+					continue;
+				}
+				WalkEIdx = It.Value();
+			}
+			// if next vertex had valence != 2, it's the end of this chain
+			if (ECount != 2)
+			{
+				Chains.Emplace(WalkVID, EndChain);
+				break;
+			}
+			// walk to the next edge and mark it as processed
+			EdgeUsed[WalkEIdx] = 1;
+			WalkEID = TriangleEdges[WalkEIdx];
+			Chains.Emplace(WalkVID, WalkEID);
+		}
+		check(MaxIters > 0);
+	};
+
+	// Process all loops and chains that have at least one vertex on multiple selected edges
+	for (int32 VID : AllEdgeVIDs)
+	{
+		int32 NumEdges = V2EIdx.Num(VID);
+		if (IsCornerVertex(VID) || NumEdges != 2)
+		{
+			for (TMultiMap<int32, int32>::TConstKeyIterator It = V2EIdx.CreateConstKeyIterator(VID); It; ++It)
+			{
+				if (!EdgeUsed[It.Value()])
+				{
+					ProcessEdgeIdx(VID, It.Value());
+				}
+			}
+		}
+	}
+
+	// Process remaining (isolated) loops where all vertices had only two adjacent selected edges
+	for (int32 Idx = 0; Idx < EdgeUsed.Num(); ++Idx)
+	{
+		if (!EdgeUsed[Idx])
+		{
+			ProcessEdgeIdx(Mesh.GetEdgeV(TriangleEdges[Idx]).A, Idx);
+		}
+	}
+
+	// Add the edge chains and loops to the bevel data structures
+	TArray<int32> EdgeLoop, VertexLoop;
+	for (int32 ChainIdx = 0, ChainLen = 0; ChainIdx < Chains.Num(); ChainIdx += ChainLen + 1)
+	{
+		for (ChainLen = 0; ChainLen + ChainIdx < Chains.Num() && Chains[ChainIdx + ChainLen].Value >= 0; ++ChainLen)
+		{}
+		check(ChainLen > 0 && ChainIdx + ChainLen < Chains.Num());
+		
+		bool bIsLoop = Chains[ChainIdx + ChainLen].Value == EndLoop;
+		if (bIsLoop)
+		{
+			EdgeLoop.Reset(ChainLen);
+			VertexLoop.Reset(ChainLen);
+			for (int32 Idx = ChainIdx; Idx < ChainIdx + ChainLen; ++Idx)
+			{
+				EdgeLoop.Add(Chains[Idx].Value);
+				VertexLoop.Add(Chains[Idx].Key);
+			}
+			// Note: Could initialize the bevel edge loop data directly rather than going via FEdgeLoop and avoid some copying
+			FEdgeLoop NewLoop;
+			NewLoop.Initialize(&Mesh, VertexLoop, EdgeLoop);
+			AddBevelEdgeLoop(Mesh, NewLoop);
+		}
+		else
+		{
+			FIndex2i CornerVIDs(Chains[ChainIdx].Key, Chains[ChainIdx + ChainLen].Key);
+
+			FBevelEdge Edge;
+			int32 NewBevelEdgeIndex = Edges.Num();
+
+			for (int32 ci = 0; ci < 2; ++ci)
+			{
+				int32 VertexID = CornerVIDs[ci];
+				Edge.bEndpointBoundaryFlag[ci] = Mesh.IsBoundaryVertex(VertexID);
+				int32 IncomingEdgeID = (ci == 0) ? Chains[ChainIdx].Value : Chains[ChainIdx + ChainLen - 1].Value;
+
+				int32 BevelVertexIndex = -1;
+				FBevelVertex* VertInfo = GetBevelVertexFromVertexID(VertexID, &BevelVertexIndex);
+				if (VertInfo == nullptr)
+				{
+					FBevelVertex NewVertex;
+					PRAGMA_DISABLE_DEPRECATION_WARNINGS
+					NewVertex.CornerID = INDEX_NONE;
+					PRAGMA_ENABLE_DEPRECATION_WARNINGS
+					NewVertex.VertexID = VertexID;
+					BevelVertexIndex = Vertices.Num();
+					Vertices.Add(NewVertex);
+					VertexIDToIndexMap.Add(VertexID, BevelVertexIndex);
+					VertInfo = &Vertices[BevelVertexIndex];
+				}
+				VertInfo->IncomingBevelMeshEdges.Add(IncomingEdgeID);
+				PRAGMA_DISABLE_DEPRECATION_WARNINGS
+				VertInfo->IncomingBevelTopoEdges.Add(INDEX_NONE);
+				PRAGMA_ENABLE_DEPRECATION_WARNINGS
+				VertInfo->IncomingBevelEdgeIndices.Add(NewBevelEdgeIndex);
+				Edge.BevelVertices[ci] = BevelVertexIndex;
+			}
+
+			Edge.MeshEdges.Reset(ChainLen);
+			for (int32 Idx = ChainIdx; Idx < ChainIdx + ChainLen; ++Idx)
+			{
+				Edge.MeshEdges.Add(Chains[Idx].Value);
+			}
+			Edge.MeshVertices.Reset(ChainLen + 1);
+			for (int32 Idx = ChainIdx; Idx < ChainIdx + ChainLen + 1; ++Idx)
+			{
+				Edge.MeshVertices.Add(Chains[Idx].Key);
+			}
+			PRAGMA_DISABLE_DEPRECATION_WARNINGS
+			Edge.GroupEdgeID = INDEX_NONE;
+			PRAGMA_ENABLE_DEPRECATION_WARNINGS
+
+			Edge.MeshEdgeTris.Reserve(Edge.MeshEdges.Num());
+			for (int32 eid : Edge.MeshEdges)
+			{
+				Edge.MeshEdgeTris.Add(Mesh.GetEdgeT(eid));
+			}
+
+			Edge.InitialPositions.Reserve(Edge.MeshVertices.Num());
+			for (int32 vid : Edge.MeshVertices)
+			{
+				Edge.InitialPositions.Add(Mesh.GetVertex(vid));
+			}
+
+			Edge.EdgeIndex = NewBevelEdgeIndex;
+			Edges.Add(MoveTemp(Edge));
+		}
+
+		if (ResultInfo.CheckAndSetCancelled(Progress))
+		{
+			return;
+		}
+	}
+
+	// precompute topological information necessary to apply bevel to vertices/edges/loops
+	BuildVertexSets(Mesh);
+}
+
 
 bool FMeshBevel::InitializeFromGroupTopologyFaces(const FDynamicMesh3& Mesh, const FGroupTopology& Topology, const TArray<int32>& GroupFaces)
 {
diff --git a/Engine/Plugins/Runtime/GeometryProcessing/Source/DynamicMesh/Public/Operations/MeshBevel.h b/Engine/Plugins/Runtime/GeometryProcessing/Source/DynamicMesh/Public/Operations/MeshBevel.h
--- a/Engine/Plugins/Runtime/GeometryProcessing/Source/DynamicMesh/Public/Operations/MeshBevel.h
+++ b/Engine/Plugins/Runtime/GeometryProcessing/Source/DynamicMesh/Public/Operations/MeshBevel.h
@@ -101,6 +101,20 @@ public:
 	 */
 	void InitializeFromGroupTopologyEdges(const FDynamicMesh3& Mesh, const FGroupTopology& Topology, const TArray<int32>& GroupEdges);
 
+	/**
+	 * Initialize the bevel with the specified mesh triangle edges
+	 * @param IsCornerVertex Function to determine whether a vertex should be treated as a corner when beveling. Only affects vertices with exactly two adjacent TriangleEdges; otherwise corners will be automatically detected.
+	 */
+	void InitializeFromTriangleEdges(const FDynamicMesh3& Mesh, TConstArrayView<int32> TriangleEdges, TFunctionRef<bool(int32)> IsCornerVertex);
+
+	/**
+	 * Initialize the bevel with the specified mesh triangle edges
+	 */
+	void InitializeFromTriangleEdges(const FDynamicMesh3& Mesh, TConstArrayView<int32> TriangleEdges)
+	{
+		InitializeFromTriangleEdges(Mesh, TriangleEdges, [](int32) {return false;});
+	}
+
 	/**
 	* Initialize the bevel with the specified faces of a GroupTopology
 	* @return false if any selection-bowtie vertices were found, in this case we cannot compute the bevel
diff --git a/Engine/Plugins/Runtime/GeometryScripting/Source/GeometryScriptingCore/Private/MeshModelingFunctions.cpp b/Engine/Plugins/Runtime/GeometryScripting/Source/GeometryScriptingCore/Private/MeshModelingFunctions.cpp
--- a/Engine/Plugins/Runtime/GeometryScripting/Source/GeometryScriptingCore/Private/MeshModelingFunctions.cpp
+++ b/Engine/Plugins/Runtime/GeometryScripting/Source/GeometryScriptingCore/Private/MeshModelingFunctions.cpp
@@ -837,7 +837,19 @@ UDynamicMesh* UGeometryScriptLibrary_MeshModelingFunctions::ApplyMeshInsetOutset
 }
 
 
-
+UDynamicMesh* UGeometryScriptLibrary_MeshModelingFunctions::ApplyMeshBevelEdgeSelection(
+	UDynamicMesh* TargetMesh,
+	FGeometryScriptMeshSelection Selection,
+	FGeometryScriptMeshBevelSelectionOptions BevelOptions,
+	UGeometryScriptDebug* Debug)
+{
+	return ApplyMeshBevelSelection(
+		TargetMesh,
+		Selection,
+		EGeometryScriptMeshBevelSelectionMode::SelectedEdges,
+		BevelOptions,
+		Debug);
+}
 
 UDynamicMesh* UGeometryScriptLibrary_MeshModelingFunctions::ApplyMeshBevelSelection(
 	UDynamicMesh* TargetMesh,
@@ -863,7 +875,26 @@ UDynamicMesh* UGeometryScriptLibrary_MeshModelingFunctions::ApplyMeshBevelSelect
 
 		TUniquePtr<FGroupTopology> GroupTopology;
 
-		if (BevelMode == EGeometryScriptMeshBevelSelectionMode::TriangleArea)
+		if (BevelMode == EGeometryScriptMeshBevelSelectionMode::SelectedEdges)
+		{
+			TArray<int32> Edges;
+			Selection.ConvertToMeshIndexArray(EditMesh, Edges, EGeometryScriptIndexType::Edge);
+			// Note: Using group junctions as corner vertices tends to help the bevel match the result of a similar polygroup bevel in the UI
+			// (though is often very similar to the no-corner bevel; necessary corners where more than 2 selected edges meet are automatically detected)
+			if (EditMesh.HasTriangleGroups())
+			{
+				Bevel.InitializeFromTriangleEdges(EditMesh, Edges, [&EditMesh](int32 VID)
+					{
+						return EditMesh.IsGroupJunctionVertex(VID);
+					});
+			}
+			else
+			{
+				Bevel.InitializeFromTriangleEdges(EditMesh, Edges, [](int32) {return false;});
+			}
+			
+		}
+		else if (BevelMode == EGeometryScriptMeshBevelSelectionMode::TriangleArea)
 		{
 			TArray<int32> Triangles;
 			Selection.ConvertToMeshIndexArray(EditMesh, Triangles, EGeometryScriptIndexType::Triangle);
diff --git a/Engine/Plugins/Runtime/GeometryScripting/Source/GeometryScriptingCore/Public/GeometryScript/MeshModelingFunctions.h b/Engine/Plugins/Runtime/GeometryScripting/Source/GeometryScriptingCore/Public/GeometryScript/MeshModelingFunctions.h
--- a/Engine/Plugins/Runtime/GeometryScripting/Source/GeometryScriptingCore/Public/GeometryScript/MeshModelingFunctions.h
+++ b/Engine/Plugins/Runtime/GeometryScripting/Source/GeometryScriptingCore/Public/GeometryScript/MeshModelingFunctions.h
@@ -253,7 +253,9 @@ enum class EGeometryScriptMeshBevelSelectionMode : uint8
 	/** Convert the selection to PolyGroups and bevel all the PolyGroup Edges of the selected PolyGroups */
 	AllPolygroupEdges = 1 UMETA(DisplayName = "All PolyGroup Edges"),
 	/** Convert the selection to PolyGroups and bevel all the PolyGroup Edges that are between selected PolyGroups */
-	SharedPolygroupEdges = 2 UMETA(DisplayName = "Shared PolyGroup Edges")
+	SharedPolygroupEdges = 2 UMETA(DisplayName = "Shared PolyGroup Edges"),
+	/** Convert the selection to Edges (if needed) and bevel them */
+	SelectedEdges = 3
 };
 
 
@@ -380,10 +382,23 @@ public:
 	/**
 	 * Apply a Mesh Bevel operation to parts of TargetMesh using the BevelOptions settings.
 	 * @param Selection specifies which mesh edges to Bevel
-	 * @param BevelMode specifies how Selection should be converted to a Triangle Region or set of PolyGroup Edges
 	 * @param BevelOptions settings for the Bevel Operation
 	 */
-	UFUNCTION(BlueprintCallable, Category = "GeometryScript|Modeling", meta=(ScriptMethod))
+	UFUNCTION(BlueprintCallable, Category = "GeometryScript|Modeling", meta = (ScriptMethod))
+	static UPARAM(DisplayName = "Target Mesh") UDynamicMesh*
+	ApplyMeshBevelEdgeSelection(
+		UDynamicMesh* TargetMesh,
+		FGeometryScriptMeshSelection Selection,
+		FGeometryScriptMeshBevelSelectionOptions BevelOptions,
+		UGeometryScriptDebug* Debug = nullptr);
+
+	/**
+	 * Apply a Mesh Bevel operation to parts of TargetMesh using the BevelOptions settings, with additional options to handle region selections
+	 * @param Selection specifies which mesh edges to Bevel
+	 * @param BevelMode specifies whether Selection should be optionally converted to a Triangle Region or set of PolyGroup Edges
+	 * @param BevelOptions settings for the Bevel Operation
+	 */
+	UFUNCTION(BlueprintCallable, Category = "GeometryScript|Modeling", meta=(ScriptMethod, DisplayName = "Apply Mesh Bevel Region Selection"))
 	static UPARAM(DisplayName = "Target Mesh") UDynamicMesh* 
 	ApplyMeshBevelSelection(
 		UDynamicMesh* TargetMesh,
