add geometry script method to set or remove sharp normal on edge selections

to implement this, also generalize the uv mesh seam insertion methods to support operating on more overlay types

diff --git a/Engine/Plugins/Runtime/GeometryProcessing/Source/DynamicMesh/Private/Parameterization/DynamicMeshUVEditor.cpp b/Engine/Plugins/Runtime/GeometryProcessing/Source/DynamicMesh/Private/Parameterization/DynamicMeshUVEditor.cpp
--- a/Engine/Plugins/Runtime/GeometryProcessing/Source/DynamicMesh/Private/Parameterization/DynamicMeshUVEditor.cpp
+++ b/Engine/Plugins/Runtime/GeometryProcessing/Source/DynamicMesh/Private/Parameterization/DynamicMeshUVEditor.cpp
@@ -18,6 +18,7 @@
 #include "Parameterization/MeshUVPacking.h"
 #include "Solvers/MeshParameterizationSolvers.h"
 #include "Parameterization/MeshRegionGraph.h"
+#include "DynamicMeshEditor.h"
 
 #include "Async/ParallelFor.h"
 
@@ -819,120 +820,6 @@ bool FDynamicMeshUVEditor::SetTriangleUVsFromConformal(const TArray<int32>& Tria
 
 
 
-// Assuming that SplitVtx is either on a mesh boundary or UV seam, find the two sets of one-ring triangles
-// on either side of edge BaseEdgeID that are edge-connected in the UV overlay. Either set could be empty.
-static bool FindSeamTriSplitSets_BoundaryVtx(const FDynamicMesh3* Mesh, const FDynamicMeshUVOverlay* UVOverlay, 
-	int32 SplitVtx, 
-	int32 BaseEdgeID,
-	TArray<int32> SplitTriSets[2])
-{
-	FIndex2i StartTris = Mesh->GetEdgeT(BaseEdgeID);
-	check(Mesh->IsTriangle(StartTris.A) && Mesh->IsTriangle(StartTris.B));
-
-	for (int32 si = 0; si < 2; ++si)
-	{
-		int32 StartTri = StartTris[si];
-		SplitTriSets[si].Add(StartTri);
-		
-		int32 EdgeOtherTri = StartTris[si==0?1:0];
-		int32 CurTri = StartTri;
-		int32 PrevTri = EdgeOtherTri;
-
-		bool bDone = false;
-		while (!bDone)
-		{
-			FIndex3i NextTri = UE::Geometry::FindNextAdjacentTriangleAroundVtx(Mesh, SplitVtx, CurTri, PrevTri,
-				[&](int32 Tri0, int32 Tri1, int32 Edge) { return UVOverlay->AreTrianglesConnected(Tri0, Tri1); }
-			);
-			if (NextTri.A != IndexConstants::InvalidID)
-			{
-				if (NextTri.A == EdgeOtherTri)		// if we looped around, SplitVtx is not on a boundary and we need to abort
-				{
-					return false;
-				}
-
-				SplitTriSets[si].Add(NextTri.A);
-				PrevTri = CurTri;
-				CurTri = NextTri.A;
-			}
-			else
-			{
-				bDone = true;
-			}
-		}
-	}
-
-	return true;
-}
-
-
-
-// If we want to cut the UV one-ring at SplitVtx with the edge sequence [PrevBaseEdgeID,NextBaseEdgeID], assuming both
-// edges are connected to SplitVtx, we need to find the connected sets of triangles on either side of NextBaseEdgeID
-// (assuming PrevBaseEdgeID was already handled). To do that we walk around the uv-connected-one-ring away from NextBaseEdgeID
-// in either direction.
-static bool FindSeamTriSplitSets_InteriorVtx(const FDynamicMesh3* Mesh, const FDynamicMeshUVOverlay* UVOverlay, 
-	int32 SplitVtx, 
-	int32 PrevBaseEdgeID, int32 NextBaseEdgeID,
-	TArray<int32> SplitTriSets[2] )
-{
-	FIndex2i StartTris = Mesh->GetEdgeT(NextBaseEdgeID);
-	check(Mesh->IsTriangle(StartTris.A) && Mesh->IsTriangle(StartTris.B));
-
-	for (int32 si = 0; si < 2; ++si)
-	{
-		int32 StartTri = StartTris[si];
-		SplitTriSets[si].Add(StartTri);
-
-		int32 EdgeOtherTri = StartTris[si == 0 ? 1 : 0];
-		int32 CurTri = StartTri;
-		int32 PrevTri = EdgeOtherTri;
-
-		bool bDone = false;
-		while (!bDone)
-		{
-			FIndex3i NextTri = UE::Geometry::FindNextAdjacentTriangleAroundVtx(Mesh, SplitVtx, CurTri, PrevTri,
-				[&](int32 Tri0, int32 Tri1, int32 Edge) { return UVOverlay->AreTrianglesConnected(Tri0, Tri1) && Edge != PrevBaseEdgeID && Edge != NextBaseEdgeID; }
-			);
-			if (NextTri.A != IndexConstants::InvalidID)
-			{
-				if (NextTri.A == EdgeOtherTri)		// if we somehow looped around, then the arguments were bad and we need to abort
-				{
-					return false;
-				}
-
-				SplitTriSets[si].Add(NextTri.A);
-				PrevTri = CurTri;
-				CurTri = NextTri.A;
-			}
-			else
-			{
-				bDone = true;
-			}
-		}
-	}
-
-	return true;
-}
-
-
-// Find the UV element for MeshVertexID that is contained in the "first" triangle of MeshEdgeID
-static int32 FindUVElementForVertex(const FDynamicMesh3* Mesh, const FDynamicMeshUVOverlay* UVOverlay, int32 MeshVertexID, int32 MeshEdgeID)
-{
-	FIndex2i Tris = Mesh->GetEdgeT(MeshEdgeID);
-	FIndex3i Tri0 = Mesh->GetTriangle(Tris.A);
-	FIndex3i UVTri0 = UVOverlay->GetTriangle(Tris.A);
-	for (int32 j = 0; j < 3; ++j)
-	{
-		if (Tri0[j] == MeshVertexID)
-		{
-			return UVTri0[j];
-		}
-	}
-	return IndexConstants::InvalidID;
-}
-
-
 void UE::Geometry::FDynamicMeshUVEditor::MakeSureUVsAreSet(const TSet<int32>& Triangles, 
 	FUVEditResult* Result, TSet<int32>* ChangedTrianglesOut)
 {
@@ -981,191 +868,12 @@ void UE::Geometry::FDynamicMeshUVEditor::MakeSureUVsAreSet(const TSet<int32>& Tr
 
 bool FDynamicMeshUVEditor::RemoveSeamsAtEdges(const TSet<int32>& EidsToRemoveAsSeams)
 {
-	for (int32 Eid : EidsToRemoveAsSeams)
-	{
-		FIndex2i EdgeVids = Mesh->GetEdgeV(Eid);
-
-		// The following logic closely resemebles the logic within FDynamicMeshOverlay::IsSeamEdge,
-		// but since we want to know both if it's a seam and use many of the intermediate values
-		// computed by IsSeamEdge, we replicate it here.
-
-		FIndex2i Tris = Mesh->GetEdgeT(Eid);
-		if (Tris.B == FDynamicMesh3::InvalidID)
-		{
-			continue; // Technically a seam, but we don't want this one because there's no opposite edge to merge.
-		}
-
-		bool bASet = UVOverlay->IsSetTriangle(Tris.A), bBSet = UVOverlay->IsSetTriangle(Tris.B);
-		if (!bASet || !bBSet)
-		{
-			continue; // Similar problem as the above case - could be a seam here, but if one triangle
-					  // in the Overlay isn't set, there's nothing to merge.
-		}
-
-		FIndex3i Triangle0 = UVOverlay->GetTriangle(Tris.A);
-		FIndex3i BaseTriangle0 = Mesh->GetTriangle(Tris.A);
-		int idx_base_a0 = BaseTriangle0.IndexOf(EdgeVids.A);
-		int idx_base_b0 = BaseTriangle0.IndexOf(EdgeVids.B);
-
-		FIndex3i Triangle1 = UVOverlay->GetTriangle(Tris.B);
-		FIndex3i BaseTriangle1 = Mesh->GetTriangle(Tris.B);
-		int idx_base_a1 = BaseTriangle1.IndexOf(EdgeVids.A);
-		int idx_base_b1 = BaseTriangle1.IndexOf(EdgeVids.B);
-
-		int el_a_tri0 = Triangle0[idx_base_a0];
-		int el_b_tri0 = Triangle0[idx_base_b0];
-		int el_a_tri1 = Triangle1[idx_base_a1];
-		int el_b_tri1 = Triangle1[idx_base_b1];
-
-		// This shouldn't ever be the case, but lets just check. If true,
-		// it would indicate that there's a joined seam but somehow our pairwise
-		// vertex matching didn't work above for some reason.
-		ensure(!(el_a_tri0 == el_b_tri1 && el_b_tri0 == el_a_tri1));
-
-		if (el_a_tri0 != el_a_tri1)
-		{
-			UVOverlay->MergeElement(el_a_tri0, el_a_tri1);
-		}
-		if (el_b_tri0 != el_b_tri1)
-		{
-			UVOverlay->MergeElement(el_b_tri0, el_b_tri1);
-
-		}
-	}
-	return true;
+	return FDynamicMeshEditor::RemoveSeamsAtEdges(EidsToRemoveAsSeams, UVOverlay);
 }
 
-
 bool FDynamicMeshUVEditor::CreateSeamsAtEdges(const TSet<int32>& EidsToMakeIntoSeams, FUVEditResult* Result)
 {
-	/**
-	 * @param TidsUntilNextSeam Tids in a fan away from Eid around Vid until the next seam. Will
-	 *  be the entire one-ring if the given Eid is the only seam attached to Vid.
-	 */
-	auto DoesVertHaveAnotherSeamAttached = [this, &EidsToMakeIntoSeams](int32 Eid, int32 Vid, TArray<int32>& TidsUntilNextSeam)
-	{
-		TidsUntilNextSeam.Reset();
-
-		FIndex2i EdgeTids = Mesh->GetEdgeT(Eid);
-		int32 CurrentTid = EdgeTids.A;
-		int32 PreviousTid = EdgeTids.B;
-
-		int32 MaxNumTids = Mesh->GetVtxEdgeCount(Vid);
-		while (true) 
-		{
-			TidsUntilNextSeam.Add(CurrentTid);
-			if (TidsUntilNextSeam.Num() > MaxNumTids) // sanity check
-			{
-				ensure(false);
-				return false;
-			}
-
-			FIndex3i NextTriResult = UE::Geometry::FindNextAdjacentTriangleAroundVtx(Mesh, Vid, CurrentTid, PreviousTid,
-				[this, &EidsToMakeIntoSeams, Eid](int32 Tri0, int32 Tri1, int32 EidBetween) {
-					// We stop at seam or at future seam not including our starting edge (if we walk across the starting
-					// edge we will stop at a later check- we just don't want to think that we reached a true seam).
-					return UVOverlay->AreTrianglesConnected(Tri0, Tri1) && !(EidBetween != Eid && EidsToMakeIntoSeams.Contains(EidBetween)); }
-			);
-			int32 NextTid = NextTriResult.A;
-
-			// See if we've come to a (current or future) seam or wrapped around
-			if (NextTid == IndexConstants::InvalidID)
-			{
-				return true;
-			}
-			else if (NextTid == EdgeTids.A)
-			{
-				return false;
-			}
-			PreviousTid = CurrentTid;
-			CurrentTid = NextTid;
-		}
-	};
-
-	auto SplitEdgeVertElement = [this, Result](int32 Eid, int32 Vid, const TArray<int32>& TidsToModify) {
-		int32 ElementID = FindUVElementForVertex(Mesh, UVOverlay, Vid, Eid);
-		if (ensure(ElementID != IndexConstants::InvalidID))
-		{
-			int32 NewElementID = UVOverlay->SplitElement(ElementID, TidsToModify);
-			if (Result)
-			{
-				Result->NewUVElements.Add(NewElementID);
-			}
-		}
-	};
-
-	for (int32 Eid : EidsToMakeIntoSeams)
-	{
-		FIndex2i EdgeVids = Mesh->GetEdgeV(Eid);
-		bool VertNeedsSplitting[2] = { false, false };
-		TArray<int32> VertTidsToUseForSplit[2];
-
-		if (UVOverlay->IsSeamEdge(Eid))
-		{
-			// If this is already a seam edge, make sure its endpoints are not bowties. We create such edges 
-			// sometimes as we split adjacent edges, and it is also likely to be what the user wants when
-			// selecting edges adjacent to a bowtie.
-			for (int i = 0; i < 2; ++i)
-			{
-				if (UVOverlay->IsBowtieInOverlay(EdgeVids[i]))
-				{
-					UVOverlay->SplitBowtiesAtVertex(EdgeVids[i], 
-						Result ? &Result->NewUVElements : nullptr);
-				}
-			}
-			continue;
-		}
-
-		// If we're not already a seam, then it's going to become one. A vert needs to get split if it
-		// has a present or future seam attached, or else it will become a bowtie if we just split the
-		// other vert.
-		for (int i = 0; i < 2; ++i)
-		{
-			VertNeedsSplitting[i] = DoesVertHaveAnotherSeamAttached(Eid, EdgeVids[i], VertTidsToUseForSplit[i]);
-		}
-			
-		// If neither absolutely has to get split, then one needs to get split anyway so that we
-		// make the edge into a seam.
-		if (!VertNeedsSplitting[0] && !VertNeedsSplitting[1])
-		{
-			// We'll go halfway around the triangle one-ring, rounding up
-			int32 NumTrisToDisconnect = (VertTidsToUseForSplit[0].Num() + 1) / 2;
-
-			// In doing this, we are splitting an adjacent edge and therefore might inadvertantly
-			// introduce a bowtie on the other vert of that edge, which we would like to avoid. 
-			// So keep track of that potential bowtie.
-			int32 OtherSplitEdge = Mesh->FindEdgeFromTriPair(
-				VertTidsToUseForSplit[0][NumTrisToDisconnect - 1],
-				VertTidsToUseForSplit[0][NumTrisToDisconnect]);
-			int32 PotentialBowtieVid = ensure(OtherSplitEdge != IndexConstants::InvalidID) ?
-				IndexUtil::FindEdgeOtherVertex(Mesh->GetEdgeV(OtherSplitEdge), EdgeVids[0])
-				: IndexConstants::InvalidID;
-
-			// Perform the split
-			VertTidsToUseForSplit[0].SetNum(NumTrisToDisconnect);
-			SplitEdgeVertElement(Eid, EdgeVids[0], VertTidsToUseForSplit[0]);
-
-			// Deal with the bowtie if we created one
-			if (ensure(PotentialBowtieVid != IndexConstants::InvalidID) 
-				&& UVOverlay->IsBowtieInOverlay(PotentialBowtieVid))
-			{
-				UVOverlay->SplitBowtiesAtVertex(PotentialBowtieVid,
-					Result ? &Result->NewUVElements : nullptr);
-			}
-		}
-		else
-		{
-			for (int i = 0; i < 2; ++i)
-			{
-				if (VertNeedsSplitting[i])
-				{
-					SplitEdgeVertElement(Eid, EdgeVids[i], VertTidsToUseForSplit[i]);
-				}
-			}
-		}//end if at least one side needed splitting
-	}//end for each edge
-
-	return true;
+	return FDynamicMeshEditor::CreateSeamsAtEdges(EidsToMakeIntoSeams, UVOverlay, Result ? &Result->NewUVElements : nullptr);
 }
 
 
diff --git a/Engine/Plugins/Runtime/GeometryScripting/Source/GeometryScriptingCore/Private/MeshNormalsFunctions.cpp b/Engine/Plugins/Runtime/GeometryScripting/Source/GeometryScriptingCore/Private/MeshNormalsFunctions.cpp
--- a/Engine/Plugins/Runtime/GeometryScripting/Source/GeometryScriptingCore/Private/MeshNormalsFunctions.cpp
+++ b/Engine/Plugins/Runtime/GeometryScripting/Source/GeometryScriptingCore/Private/MeshNormalsFunctions.cpp
@@ -5,6 +5,7 @@
 #include "DynamicMesh/DynamicMesh3.h"
 #include "DynamicMesh/MeshNormals.h"
 #include "DynamicMesh/MeshTangents.h"
+#include "DynamicMeshEditor.h" // for sharp edge creation/deletion via overlay seam methods
 #include "Operations/RepairOrientation.h"
 #include "DynamicMesh/DynamicMeshAABBTree3.h"
 #include "Polygroups/PolygroupSet.h"
@@ -195,6 +196,76 @@ UDynamicMesh* UGeometryScriptLibrary_MeshNormalsFunctions::RecomputeNormalsForMe
 
 
 
+UDynamicMesh* UGeometryScriptLibrary_MeshNormalsFunctions::SetSplitNormalsAlongSelectedEdges(
+	UDynamicMesh* TargetMesh,
+	FGeometryScriptMeshSelection Selection,
+	bool bSplit,
+	bool bRecalculateNormals,
+	FGeometryScriptCalculateNormalsOptions CalculateOptions,
+	bool bDeferChangeNotifications,
+	UGeometryScriptDebug* Debug)
+{
+	if (TargetMesh == nullptr)
+	{
+		UE::Geometry::AppendError(Debug, EGeometryScriptErrorType::InvalidInputs, LOCTEXT("SetSplitNormalsAlongSelectedEdges_InvalidInput", "SetSplitNormalsAlongSelectedEdges: TargetMesh is Null"));
+		return TargetMesh;
+	}
+
+	TargetMesh->EditMesh([&](FDynamicMesh3& EditMesh)
+	{
+		if (EditMesh.HasAttributes() == false || EditMesh.Attributes()->PrimaryNormals() == nullptr)
+		{
+			UE::Geometry::AppendError(Debug, EGeometryScriptErrorType::InvalidInputs, LOCTEXT("SetSplitNormalsAlongSelectedEdges_NoAttributes", "SetSplitNormalsAlongSelectedEdges: TargetMesh has no Normals attribute enabled"));
+			return;
+		}
+
+		TArray<int32> EdgeArr;
+		Selection.ConvertToMeshIndexArray(EditMesh, EdgeArr, EGeometryScriptIndexType::Edge);
+		TSet<int32> EdgeSet(EdgeArr);
+		// split/merge sharp normals (and tangents/bitangents if available)
+		for (int32 NormalLayerIdx = 0; NormalLayerIdx < EditMesh.Attributes()->NumNormalLayers(); ++NormalLayerIdx)
+		{
+			if (bSplit)
+			{
+				FDynamicMeshEditor::CreateSeamsAtEdges(EdgeSet, EditMesh.Attributes()->GetNormalLayer(NormalLayerIdx));
+			}
+			else
+			{
+				FDynamicMeshEditor::RemoveSeamsAtEdges(EdgeSet, EditMesh.Attributes()->GetNormalLayer(NormalLayerIdx));
+			}
+		}
+		if (bRecalculateNormals)
+		{
+			FDynamicMeshNormalOverlay* Normals = EditMesh.Attributes()->PrimaryNormals();
+			TSet<int32> VIDs;
+			TArray<int32> Elements;
+			auto ProcessVert = [&VIDs, &Elements, Normals](int32 VID)
+			{
+				bool bAlreadyAdded = false;
+				VIDs.Add(VID, &bAlreadyAdded);
+				if (!bAlreadyAdded)
+				{
+					Normals->EnumerateVertexElements(VID, [&Elements](int32 TID, int32 ElID, const FVector3f&)
+						{
+							Elements.Add(ElID);
+							return true;
+						}, /*bFindUniqueElements*/true);
+				}
+			};
+			for (int32 EID : EdgeArr)
+			{
+				FIndex2i EdgeV = EditMesh.GetEdgeV(EID);
+				ProcessVert(EdgeV.A);
+				ProcessVert(EdgeV.B);
+			}
+			FMeshNormals::RecomputeOverlayElementNormals(EditMesh, Elements, CalculateOptions.bAngleWeighted, CalculateOptions.bAngleWeighted);
+		}
+
+	}, EDynamicMeshChangeType::GeneralEdit, EDynamicMeshAttributeChangeFlags::NormalsTangents, bDeferChangeNotifications);
+
+	return TargetMesh;
+}
+
 
 UDynamicMesh* UGeometryScriptLibrary_MeshNormalsFunctions::ComputeSplitNormals( 
 	UDynamicMesh* TargetMesh, 
diff --git a/Engine/Plugins/Runtime/GeometryScripting/Source/GeometryScriptingCore/Public/GeometryScript/MeshNormalsFunctions.h b/Engine/Plugins/Runtime/GeometryScripting/Source/GeometryScriptingCore/Public/GeometryScript/MeshNormalsFunctions.h
--- a/Engine/Plugins/Runtime/GeometryScripting/Source/GeometryScriptingCore/Public/GeometryScript/MeshNormalsFunctions.h
+++ b/Engine/Plugins/Runtime/GeometryScripting/Source/GeometryScriptingCore/Public/GeometryScript/MeshNormalsFunctions.h
@@ -153,6 +153,26 @@ public:
 		FGeometryScriptCalculateNormalsOptions CalculateOptions,
 		UGeometryScriptDebug* Debug = nullptr);
 
+	/**
+	 * Set or remove split normals (aka sharp normals) for all edges in the Selection
+	 * @param TargetMesh The mesh to update
+	 * @param Selection Which edges to operate on
+	 * @param bSplit Whether to split normals along the selected edges; if false, they will be merged instead
+	 * @param bRecalculateNormals Whether to recalculate normals along edges where they were split/merged
+	 * @param CalculateOptions Options for computing the normals, if bRecalculateNormals is true
+	 * @param bDeferChangeNotifications If true, no mesh change notification will be sent. Set to true if performing many changes in a loop.
+	 */
+	UFUNCTION(BlueprintCallable, Category = "GeometryScript|Normals", meta = (ScriptMethod))
+	static UPARAM(DisplayName = "Target Mesh") UDynamicMesh*
+	SetSplitNormalsAlongSelectedEdges(
+		UDynamicMesh* TargetMesh,
+		FGeometryScriptMeshSelection Selection,
+		bool bSplit = true,
+		bool bRecalculateNormals = true,
+		FGeometryScriptCalculateNormalsOptions CalculateOptions = FGeometryScriptCalculateNormalsOptions(),
+		bool bDeferChangeNotifications = false,
+		UGeometryScriptDebug* Debug = nullptr);
+
 	/**
 	 * Set the triangle-vertex normals for the given TriangleID on the TargetMesh. This will
 	 * create unique triangle-vertex normals, ie it will create hard edges / split normals in 
diff --git a/Engine/Source/Runtime/GeometryCore/Private/DynamicMeshEditor.cpp b/Engine/Source/Runtime/GeometryCore/Private/DynamicMeshEditor.cpp
--- a/Engine/Source/Runtime/GeometryCore/Private/DynamicMeshEditor.cpp
+++ b/Engine/Source/Runtime/GeometryCore/Private/DynamicMeshEditor.cpp
@@ -8,6 +8,7 @@
 #include "DynamicSubmesh3.h"
 #include "DynamicMesh/DynamicVertexSkinWeightsAttribute.h"
 #include "DynamicMesh/DynamicBoneAttribute.h"
+#include "DynamicMesh/MeshIndexUtil.h"
 #include "DynamicMesh/MeshNormals.h"
 #include "MeshQueries.h"
 #include "Selections/MeshConnectedComponents.h"
@@ -15,6 +16,212 @@
 using namespace UE::Geometry;
 
 
+// local static helper functions for seam handling
+namespace UE::Private::OverlaySeamHelpers
+{
+	template<typename RealType, int OverlayDim>
+	static bool CreateSeamsAtEdges(const TSet<int32>& EidsToMakeIntoSeams, FDynamicMesh3* Mesh, TDynamicMeshOverlay<RealType, OverlayDim>* Overlay, TArray<int32>* AddedElementIDs)
+	{
+		auto FindElementForVertex = [Mesh, Overlay](int32 MeshVertexID, int32 MeshEdgeID) -> int32
+		{
+			FIndex2i Tris = Mesh->GetEdgeT(MeshEdgeID);
+			FIndex3i Tri0 = Mesh->GetTriangle(Tris.A);
+			FIndex3i UVTri0 = Overlay->GetTriangle(Tris.A);
+			for (int32 j = 0; j < 3; ++j)
+			{
+				if (Tri0[j] == MeshVertexID)
+				{
+					return UVTri0[j];
+				}
+			}
+			return IndexConstants::InvalidID;
+		};
+
+		/**
+		 * @param TidsUntilNextSeam Tids in a fan away from Eid around Vid until the next seam. Will
+		 *  be the entire one-ring if the given Eid is the only seam attached to Vid.
+		 */
+		auto DoesVertHaveAnotherSeamAttached = [&EidsToMakeIntoSeams, Mesh, Overlay, AddedElementIDs](int32 Eid, int32 Vid, TArray<int32>& TidsUntilNextSeam)
+		{
+			TidsUntilNextSeam.Reset();
+
+			FIndex2i EdgeTids = Mesh->GetEdgeT(Eid);
+			int32 CurrentTid = EdgeTids.A;
+			int32 PreviousTid = EdgeTids.B;
+
+			int32 MaxNumTids = Mesh->GetVtxEdgeCount(Vid);
+			while (true) 
+			{
+				TidsUntilNextSeam.Add(CurrentTid);
+				if (TidsUntilNextSeam.Num() > MaxNumTids) // sanity check
+				{
+					ensure(false);
+					return false;
+				}
+
+				FIndex3i NextTriResult = UE::Geometry::FindNextAdjacentTriangleAroundVtx(Mesh, Vid, CurrentTid, PreviousTid,
+					[&EidsToMakeIntoSeams, Overlay, Eid](int32 Tri0, int32 Tri1, int32 EidBetween) {
+						// We stop at seam or at future seam not including our starting edge (if we walk across the starting
+						// edge we will stop at a later check- we just don't want to think that we reached a true seam).
+						return Overlay->AreTrianglesConnected(Tri0, Tri1) && !(EidBetween != Eid && EidsToMakeIntoSeams.Contains(EidBetween)); }
+				);
+				int32 NextTid = NextTriResult.A;
+
+				// See if we've come to a (current or future) seam or wrapped around
+				if (NextTid == IndexConstants::InvalidID)
+				{
+					return true;
+				}
+				else if (NextTid == EdgeTids.A)
+				{
+					return false;
+				}
+				PreviousTid = CurrentTid;
+				CurrentTid = NextTid;
+			}
+		};
+
+		auto SplitEdgeVertElement = [Mesh, Overlay, AddedElementIDs, &FindElementForVertex](int32 Eid, int32 Vid, const TArray<int32>& TidsToModify) {
+			int32 ElementID = FindElementForVertex(Vid, Eid);
+			if (ensure(ElementID != IndexConstants::InvalidID))
+			{
+				int32 NewElementID = Overlay->SplitElement(ElementID, TidsToModify);
+				if (AddedElementIDs)
+				{
+					AddedElementIDs->Add(NewElementID);
+				}
+			}
+		};
+
+		for (int32 Eid : EidsToMakeIntoSeams)
+		{
+			FIndex2i EdgeVids = Mesh->GetEdgeV(Eid);
+			bool VertNeedsSplitting[2] = { false, false };
+			TArray<int32> VertTidsToUseForSplit[2];
+
+			if (Overlay->IsSeamEdge(Eid))
+			{
+				// If this is already a seam edge, make sure its endpoints are not bowties. We create such edges 
+				// sometimes as we split adjacent edges, and it is also likely to be what the user wants when
+				// selecting edges adjacent to a bowtie.
+				for (int i = 0; i < 2; ++i)
+				{
+					if (Overlay->IsBowtieInOverlay(EdgeVids[i]))
+					{
+						Overlay->SplitBowtiesAtVertex(EdgeVids[i], AddedElementIDs);
+					}
+				}
+				continue;
+			}
+
+			// If we're not already a seam, then it's going to become one. A vert needs to get split if it
+			// has a present or future seam attached, or else it will become a bowtie if we just split the
+			// other vert.
+			for (int i = 0; i < 2; ++i)
+			{
+				VertNeedsSplitting[i] = DoesVertHaveAnotherSeamAttached(Eid, EdgeVids[i], VertTidsToUseForSplit[i]);
+			}
+			
+			// If neither absolutely has to get split, then one needs to get split anyway so that we
+			// make the edge into a seam.
+			if (!VertNeedsSplitting[0] && !VertNeedsSplitting[1])
+			{
+				// We'll go halfway around the triangle one-ring, rounding up
+				int32 NumTrisToDisconnect = (VertTidsToUseForSplit[0].Num() + 1) / 2;
+
+				// In doing this, we are splitting an adjacent edge and therefore might inadvertantly
+				// introduce a bowtie on the other vert of that edge, which we would like to avoid. 
+				// So keep track of that potential bowtie.
+				int32 OtherSplitEdge = Mesh->FindEdgeFromTriPair(
+					VertTidsToUseForSplit[0][NumTrisToDisconnect - 1],
+					VertTidsToUseForSplit[0][NumTrisToDisconnect]);
+				int32 PotentialBowtieVid = ensure(OtherSplitEdge != IndexConstants::InvalidID) ?
+					IndexUtil::FindEdgeOtherVertex(Mesh->GetEdgeV(OtherSplitEdge), EdgeVids[0])
+					: IndexConstants::InvalidID;
+
+				// Perform the split
+				VertTidsToUseForSplit[0].SetNum(NumTrisToDisconnect);
+				SplitEdgeVertElement(Eid, EdgeVids[0], VertTidsToUseForSplit[0]);
+
+				// Deal with the bowtie if we created one
+				if (ensure(PotentialBowtieVid != IndexConstants::InvalidID) 
+					&& Overlay->IsBowtieInOverlay(PotentialBowtieVid))
+				{
+					Overlay->SplitBowtiesAtVertex(PotentialBowtieVid, AddedElementIDs);
+				}
+			}
+			else
+			{
+				for (int i = 0; i < 2; ++i)
+				{
+					if (VertNeedsSplitting[i])
+					{
+						SplitEdgeVertElement(Eid, EdgeVids[i], VertTidsToUseForSplit[i]);
+					}
+				}
+			}//end if at least one side needed splitting
+		}//end for each edge
+
+		return true;
+	}
+
+	template<typename RealType, int OverlayDim>
+	bool RemoveSeamsAtEdges(const TSet<int32>& EidsToRemoveAsSeams, FDynamicMesh3* Mesh, TDynamicMeshOverlay<RealType, OverlayDim>* Overlay)
+	{
+		for (int32 Eid : EidsToRemoveAsSeams)
+		{
+			FIndex2i EdgeVids = Mesh->GetEdgeV(Eid);
+
+			// The following logic closely resemebles the logic within FDynamicMeshOverlay::IsSeamEdge,
+			// but since we want to know both if it's a seam and use many of the intermediate values
+			// computed by IsSeamEdge, we replicate it here.
+
+			FIndex2i Tris = Mesh->GetEdgeT(Eid);
+			if (Tris.B == FDynamicMesh3::InvalidID)
+			{
+				continue; // Technically a seam, but we don't want this one because there's no opposite edge to merge.
+			}
+
+			bool bASet = Overlay->IsSetTriangle(Tris.A), bBSet = Overlay->IsSetTriangle(Tris.B);
+			if (!bASet || !bBSet)
+			{
+				continue; // Similar problem as the above case - could be a seam here, but if one triangle
+						  // in the Overlay isn't set, there's nothing to merge.
+			}
+
+			FIndex3i Triangle0 = Overlay->GetTriangle(Tris.A);
+			FIndex3i BaseTriangle0 = Mesh->GetTriangle(Tris.A);
+			int idx_base_a0 = BaseTriangle0.IndexOf(EdgeVids.A);
+			int idx_base_b0 = BaseTriangle0.IndexOf(EdgeVids.B);
+
+			FIndex3i Triangle1 = Overlay->GetTriangle(Tris.B);
+			FIndex3i BaseTriangle1 = Mesh->GetTriangle(Tris.B);
+			int idx_base_a1 = BaseTriangle1.IndexOf(EdgeVids.A);
+			int idx_base_b1 = BaseTriangle1.IndexOf(EdgeVids.B);
+
+			int el_a_tri0 = Triangle0[idx_base_a0];
+			int el_b_tri0 = Triangle0[idx_base_b0];
+			int el_a_tri1 = Triangle1[idx_base_a1];
+			int el_b_tri1 = Triangle1[idx_base_b1];
+
+			// This shouldn't ever be the case, but lets just check. If true,
+			// it would indicate that there's a joined seam but somehow our pairwise
+			// vertex matching didn't work above for some reason.
+			ensure(!(el_a_tri0 == el_b_tri1 && el_b_tri0 == el_a_tri1));
+
+			if (el_a_tri0 != el_a_tri1)
+			{
+				Overlay->MergeElement(el_a_tri0, el_a_tri1);
+			}
+			if (el_b_tri0 != el_b_tri1)
+			{
+				Overlay->MergeElement(el_b_tri0, el_b_tri1);
+
+			}
+		}
+		return true;
+	}
+}
 
 void FDynamicMeshEditResult::GetAllTriangles(TArray<int>& TrianglesOut) const
 {
@@ -2372,6 +2579,23 @@ void FDynamicMeshEditor::AppendElementSubset(
 	}
 }
 
+bool FDynamicMeshEditor::RemoveSeamsAtEdges(const TSet<int32>& EidsToRemoveAsSeams, TDynamicMeshOverlay<float, 2>* Overlay)
+{
+	return UE::Private::OverlaySeamHelpers::RemoveSeamsAtEdges<float, 2>(EidsToRemoveAsSeams, Overlay->GetParentMesh(), Overlay);
+}
+bool FDynamicMeshEditor::RemoveSeamsAtEdges(const TSet<int32>& EidsToRemoveAsSeams, TDynamicMeshOverlay<float, 3>* Overlay)
+{
+	return UE::Private::OverlaySeamHelpers::RemoveSeamsAtEdges<float, 3>(EidsToRemoveAsSeams, Overlay->GetParentMesh(), Overlay);
+}
+bool FDynamicMeshEditor::CreateSeamsAtEdges(const TSet<int32>& EidsToMakeIntoSeams, TDynamicMeshOverlay<float, 2>* Overlay, TArray<int32>* AddedElementIDs)
+{
+	return UE::Private::OverlaySeamHelpers::CreateSeamsAtEdges<float, 2>(EidsToMakeIntoSeams, Overlay->GetParentMesh(), Overlay, AddedElementIDs);
+}
+bool FDynamicMeshEditor::CreateSeamsAtEdges(const TSet<int32>& EidsToMakeIntoSeams, TDynamicMeshOverlay<float, 3>* Overlay, TArray<int32>* AddedElementIDs)
+{
+	return UE::Private::OverlaySeamHelpers::CreateSeamsAtEdges<float, 3>(EidsToMakeIntoSeams, Overlay->GetParentMesh(), Overlay, AddedElementIDs);
+}
+
 template GEOMETRYCORE_API void FDynamicMeshEditor::AppendElementSubset(
 	const FDynamicMesh3*,
 	const TSet<int>&,
diff --git a/Engine/Source/Runtime/GeometryCore/Public/DynamicMeshEditor.h b/Engine/Source/Runtime/GeometryCore/Public/DynamicMeshEditor.h
--- a/Engine/Source/Runtime/GeometryCore/Public/DynamicMeshEditor.h
+++ b/Engine/Source/Runtime/GeometryCore/Public/DynamicMeshEditor.h
@@ -608,6 +608,46 @@ public:
 	 */
 	static GEOMETRYCORE_API bool SplitMesh(const FDynamicMesh3* SourceMesh, TArray<FDynamicMesh3>& SplitMeshes, TFunctionRef<int(int)> TriIDToMeshID, int DeleteMeshID = -1);
 
+	/**
+	 * Merge any seams in the given 2D attribute Overlay along the given mesh edge IDs
+	 *
+	 * @param EidsToRemoveAsSeams list of edges to remove seams from
+	 * @return true on success
+	 */
+	static GEOMETRYCORE_API bool RemoveSeamsAtEdges(const TSet<int32>& EidsToRemoveAsSeams, TDynamicMeshOverlay<float, 2>* Overlay);
+	/**
+	 * Merge any seams in the given (3D attribute) Overlay along the given mesh edge IDs
+	 *
+	 * @param EidsToRemoveAsSeams list of edges to remove seams from
+	 * @return true on success
+	 */
+	static GEOMETRYCORE_API bool RemoveSeamsAtEdges(const TSet<int32>& EidsToRemoveAsSeams, TDynamicMeshOverlay<float, 3>* Overlay);
+
+	/**
+	 * Cut existing 2D attribute overlay topology with a set of edges. This allows for creating partial seams/darts, interior cuts, etc.
+	 *
+	 * Avoids creating bowties in the overlay. In cases where an edge is not next to any present or future seams/borders, some
+	 * adjacent edge will be picked to be made into a seam as well, since it's impossible to make the original
+	 * into a seam otherwise.
+
+	 * @param EidsToMakeIntoSeams list of edges to turn into seams
+	 * @param AddedElementIDs if non-null, list of new elements created along the path will be stored here (not ordered)
+	 * @return true on success
+	 */
+	static GEOMETRYCORE_API bool CreateSeamsAtEdges(const TSet<int32>& EidsToMakeIntoSeams, TDynamicMeshOverlay<float, 2>* Overlay, TArray<int32>* AddedElementIDs = nullptr);
+	/**
+	 * Cut existing 3D attribute overlay topology with a set of edges. This allows for creating partial seams/darts, interior cuts, etc.
+	 *
+	 * Avoids creating bowties in the overlay. In cases where an edge is not next to any present or future seams/borders, some
+	 * adjacent edge will be picked to be made into a seam as well, since it's impossible to make the original
+	 * into a seam otherwise.
+
+	 * @param EidsToMakeIntoSeams list of edges to turn into seams
+	 * @param AddedElementIDs if non-null, list of new elements created along the path will be stored here (not ordered)
+	 * @return true on success
+	 */
+	static GEOMETRYCORE_API bool CreateSeamsAtEdges(const TSet<int32>& EidsToMakeIntoSeams, TDynamicMeshOverlay<float, 3>* Overlay, TArray<int32>* AddedElementIDs = nullptr);
+
 };
 
 
