[DynamicMesh]
* Fixed clang editor linker errors

diff --git a/Engine/Plugins/Runtime/GeometryProcessing/Source/DynamicMesh/Private/MeshSimplification.cpp b/Engine/Plugins/Runtime/GeometryProcessing/Source/DynamicMesh/Private/MeshSimplification.cpp
--- a/Engine/Plugins/Runtime/GeometryProcessing/Source/DynamicMesh/Private/MeshSimplification.cpp
+++ b/Engine/Plugins/Runtime/GeometryProcessing/Source/DynamicMesh/Private/MeshSimplification.cpp
@@ -1635,7 +1635,17 @@ bool TMeshSimplification<QuadricErrorType>::RemoveIsolatedTriangle(int tID)
 
 }
 
+template <typename QuadricErrorType>
+void TMeshSimplification<QuadricErrorType>::OnEdgeCollapse(int edgeID, int va, int vb, const FDynamicMesh3::FEdgeCollapseInfo& collapseInfo)
+{
+	// this is for subclasses...
+}
 
+template <typename QuadricErrorType>
+void TMeshSimplification<QuadricErrorType>::OnRemoveIsolatedTriangle(int tId)
+{
+	// this is for subclasses
+}
 
 // Project vertices onto projection target. 
 // We can do projection in parallel if we have .net 
diff --git a/Engine/Plugins/Runtime/GeometryProcessing/Source/DynamicMesh/Public/MeshSimplification.h b/Engine/Plugins/Runtime/GeometryProcessing/Source/DynamicMesh/Public/MeshSimplification.h
--- a/Engine/Plugins/Runtime/GeometryProcessing/Source/DynamicMesh/Public/MeshSimplification.h
+++ b/Engine/Plugins/Runtime/GeometryProcessing/Source/DynamicMesh/Public/MeshSimplification.h
@@ -326,17 +326,11 @@ protected:
 	bool RemoveIsolatedTriangle(int tID);
 
 
-	// subclasses can override these to implement custom behavior...
-	virtual void OnEdgeCollapse(int edgeID, int va, int vb, const FDynamicMesh3::FEdgeCollapseInfo& collapseInfo)
-	{
-		// this is for subclasses...
-	}
-
-	virtual void OnRemoveIsolatedTriangle(int tId)
-	{
-		// this is for subclasses
-	}
+	// subclasses can override this to implement custom behavior...
+	DYNAMICMESH_API virtual void OnEdgeCollapse(int edgeID, int va, int vb, const FDynamicMesh3::FEdgeCollapseInfo& collapseInfo);
 
+	// subclasses can override this to implement custom behavior...
+	DYNAMICMESH_API virtual void OnRemoveIsolatedTriangle(int tId);
 
 	// Project vertices onto projection target. 
 	DYNAMICMESH_API virtual void FullProjectionPass();
@@ -443,8 +437,5 @@ typedef TMeshSimplification< FAttrBasedQuadricErrord >  FAttrMeshSimplification;
 typedef TMeshSimplification < FVolPresQuadricErrord >   FVolPresMeshSimplification;
 typedef TMeshSimplification< FQuadricErrord >           FQEMSimplification;
 
-template<>
-void TMeshSimplification<FAttrBasedQuadricErrord>::OnEdgeCollapse(int edgeID, int va, int vb, const FDynamicMesh3::FEdgeCollapseInfo& collapseInfo);
-
 } // end namespace UE::Geometry
 } // end namespace UE
