[DynamicMesh]
* Added dll export on the template instantiations. This is needed for some platforms

diff --git a/Engine/Plugins/Runtime/GeometryProcessing/Source/DynamicMesh/Private/MeshSimplification.cpp b/Engine/Plugins/Runtime/GeometryProcessing/Source/DynamicMesh/Private/MeshSimplification.cpp
--- a/Engine/Plugins/Runtime/GeometryProcessing/Source/DynamicMesh/Private/MeshSimplification.cpp
+++ b/Engine/Plugins/Runtime/GeometryProcessing/Source/DynamicMesh/Private/MeshSimplification.cpp
@@ -1719,7 +1719,7 @@ FVector3d TMeshSimplification<QuadricErrorType>::GetProjectedCollapsePosition(in
 
 // Custom behavior for FAttrBasedQuadric simplifier.
 template<>
-void DYNAMICMESH_API TMeshSimplification<FAttrBasedQuadricErrord>::OnEdgeCollapse(int edgeID, int va, int vb, const FDynamicMesh3::FEdgeCollapseInfo& collapseInfo)
+void TMeshSimplification<FAttrBasedQuadricErrord>::OnEdgeCollapse(int edgeID, int va, int vb, const FDynamicMesh3::FEdgeCollapseInfo& collapseInfo)
 {
 
 	// Update the normal
@@ -1766,9 +1766,15 @@ namespace Geometry
 // These are explicit instantiations of the templates that are exported from the shared lib.
 // Only these instantiations of the template can be used.
 // This is necessary because we have placed most of the templated functions in this .cpp file, instead of the header.
+#if PLATFORM_COMPILER_CLANG
+template class DYNAMICMESH_API TMeshSimplification< FAttrBasedQuadricErrord >;
+template class DYNAMICMESH_API TMeshSimplification< FVolPresQuadricErrord >;
+template class DYNAMICMESH_API TMeshSimplification< FQuadricErrord >;
+#else
 template class TMeshSimplification< FAttrBasedQuadricErrord >;
 template class TMeshSimplification< FVolPresQuadricErrord >;
 template class TMeshSimplification< FQuadricErrord >;
+#endif
 
 } // end namespace UE::Geometry
 } // end namespace UE
\ No newline at end of file
