[DynamicMesh]
* Fixed compile errors on microsoft platforms using clang when building with merged modules

diff --git a/Engine/Plugins/Runtime/GeometryProcessing/Source/DynamicMesh/Private/MeshSimplification.cpp b/Engine/Plugins/Runtime/GeometryProcessing/Source/DynamicMesh/Private/MeshSimplification.cpp
--- a/Engine/Plugins/Runtime/GeometryProcessing/Source/DynamicMesh/Private/MeshSimplification.cpp
+++ b/Engine/Plugins/Runtime/GeometryProcessing/Source/DynamicMesh/Private/MeshSimplification.cpp
@@ -1766,7 +1766,7 @@ namespace Geometry
 // These are explicit instantiations of the templates that are exported from the shared lib.
 // Only these instantiations of the template can be used.
 // This is necessary because we have placed most of the templated functions in this .cpp file, instead of the header.
-#if PLATFORM_COMPILER_CLANG
+#if PLATFORM_COMPILER_CLANG && !PLATFORM_MICROSOFT
 template class DYNAMICMESH_API TMeshSimplification< FAttrBasedQuadricErrord >;
 template class DYNAMICMESH_API TMeshSimplification< FVolPresQuadricErrord >;
 template class DYNAMICMESH_API TMeshSimplification< FQuadricErrord >;
