ModelingTools: Fixed crash in MeshMapBaker.ComputeUVCharts for non-compact DynamicMeshes.

diff --git a/Engine/Plugins/Runtime/GeometryProcessing/Source/DynamicMesh/Private/Sampling/MeshMapBaker.cpp b/Engine/Plugins/Runtime/GeometryProcessing/Source/DynamicMesh/Private/Sampling/MeshMapBaker.cpp
--- a/Engine/Plugins/Runtime/GeometryProcessing/Source/DynamicMesh/Private/Sampling/MeshMapBaker.cpp
+++ b/Engine/Plugins/Runtime/GeometryProcessing/Source/DynamicMesh/Private/Sampling/MeshMapBaker.cpp
@@ -988,7 +988,11 @@ bool FMeshMapBaker::EvaluateIsInFilterRegion(const FVector2d& Dist)
 
 void FMeshMapBaker::ComputeUVCharts(const FDynamicMesh3& Mesh, TArray<int32>& MeshUVCharts)
 {
-	MeshUVCharts.SetNumZeroed(Mesh.TriangleCount());
+	MeshUVCharts.SetNumUninitialized(Mesh.MaxTriangleID());
+	for (int32& ChartId : MeshUVCharts)
+	{
+		ChartId = IndexConstants::InvalidID;
+	}
 	if (const FDynamicMeshUVOverlay* UVOverlay = Mesh.Attributes() ? Mesh.Attributes()->PrimaryUV() : nullptr)
 	{
 		FMeshConnectedComponents UVComponents(&Mesh);
