more consistently add/test both possible "triangle edge IDs" in various geometry selection util methods, when dealing with triangle edge selections

diff --git a/Engine/Plugins/Runtime/GeometryProcessing/Source/DynamicMesh/Private/Selections/GeometrySelectionUtil.cpp b/Engine/Plugins/Runtime/GeometryProcessing/Source/DynamicMesh/Private/Selections/GeometrySelectionUtil.cpp
--- a/Engine/Plugins/Runtime/GeometryProcessing/Source/DynamicMesh/Private/Selections/GeometrySelectionUtil.cpp
+++ b/Engine/Plugins/Runtime/GeometryProcessing/Source/DynamicMesh/Private/Selections/GeometrySelectionUtil.cpp
@@ -2102,10 +2102,19 @@ bool UE::Geometry::MakeSelectAllSelection(
 		{
 			for (int32 eid : Mesh.EdgeIndicesItr())
 			{
-				FGeoSelectionID ID = FGeoSelectionID::MeshEdge(Mesh.GetTriEdgeIDFromEdgeID(eid));
-				if ( SelectionIDPredicate(ID) )
+				// Test if both half-edges pass the edge selection predicate
+				bool bShouldSelect = true;
+				Mesh.EnumerateTriEdgeIDsFromEdgeID(eid, [&SelectionIDPredicate, &bShouldSelect](FMeshTriEdgeID TriEdgeID)
 				{
-					AllSelection.Selection.Add(ID.Encoded());
+					bShouldSelect = bShouldSelect && SelectionIDPredicate(FGeoSelectionID::MeshEdge(TriEdgeID));
+				});
+				if (bShouldSelect)
+				{
+					// Select both half-edges
+					Mesh.EnumerateTriEdgeIDsFromEdgeID(eid, [&AllSelection](FMeshTriEdgeID TriEdgeID)
+					{
+						AllSelection.Selection.Add(FGeoSelectionID::MeshEdge(TriEdgeID).Encoded());
+					});
 				}
 			}
 		}
@@ -2227,12 +2236,21 @@ bool UE::Geometry::MakeSelectAllConnectedSelection(
 			FMeshConnectedComponents::GrowToConnectedEdges(Mesh, CurIndices, ConnectedEdges, nullptr,
 				[&](int32 FromEdgeID, int32 ToEdgeID) {
 					FMeshTriEdgeID ToTriEdgeID = Mesh.GetTriEdgeIDFromEdgeID(ToEdgeID), FromTriEdgeID = Mesh.GetTriEdgeIDFromEdgeID(FromEdgeID);
-					return SelectionIDPredicate(FGeoSelectionID::MeshEdge(ToTriEdgeID)) && 
+					// Grow if both half-edges pass the predicate
+					bool bToEdgeID_SelectionPredicate = true;
+					Mesh.EnumerateTriEdgeIDsFromEdgeID(ToEdgeID, [&SelectionIDPredicate, &bToEdgeID_SelectionPredicate](FMeshTriEdgeID TestTriEdgeID)
+					{
+						bToEdgeID_SelectionPredicate = bToEdgeID_SelectionPredicate && SelectionIDPredicate(FGeoSelectionID::MeshEdge(TestTriEdgeID));
+					});
+					return bToEdgeID_SelectionPredicate &&
 							IsConnectedPredicate( FGeoSelectionID::MeshEdge(FromTriEdgeID), FGeoSelectionID::MeshEdge(ToTriEdgeID) );
 				});
 			for (int32 EdgeID : ConnectedEdges)
 			{
-				AllConnectedSelection.Selection.Add( FGeoSelectionID::MeshEdge(Mesh.GetTriEdgeIDFromEdgeID(EdgeID)).Encoded() );
+				Mesh.EnumerateTriEdgeIDsFromEdgeID(EdgeID, [&AllConnectedSelection](FMeshTriEdgeID TriEdgeID)
+				{
+					AllConnectedSelection.Selection.Add(FGeoSelectionID::MeshEdge(TriEdgeID).Encoded());
+				});
 			}
 		}
 		else if (AllConnectedSelection.ElementType == EGeometryElementType::Face)
@@ -2686,10 +2704,19 @@ bool UE::Geometry::MakeBoundaryConnectedSelection(
 			}
 			for (int32 EdgeID : AdjacentEdges)
 			{
-				FGeoSelectionID MeshEdgeID = FGeoSelectionID::MeshEdge(Mesh.GetTriEdgeIDFromEdgeID(EdgeID));
-				if (SelectionIDPredicate(MeshEdgeID))
+				// Test if both half-edges pass the edge selection predicate
+				bool bShouldSelect = true;
+				Mesh.EnumerateTriEdgeIDsFromEdgeID(EdgeID, [&SelectionIDPredicate, &bShouldSelect](FMeshTriEdgeID TriEdgeID)
 				{
-					BoundaryConnectedSelection.Selection.Add(MeshEdgeID.Encoded());
+					bShouldSelect = bShouldSelect && SelectionIDPredicate(FGeoSelectionID::MeshEdge(TriEdgeID));
+					});
+				if (bShouldSelect)
+				{
+					// Select both half-edges
+					Mesh.EnumerateTriEdgeIDsFromEdgeID(EdgeID, [&BoundaryConnectedSelection](FMeshTriEdgeID TriEdgeID)
+					{
+						BoundaryConnectedSelection.Selection.Add(FGeoSelectionID::MeshEdge(TriEdgeID).Encoded());
+					});
 				}
 			}
 		}
