Mesh Smoothing tool.  Make the smoothing respect the interrupter

diff --git a/Engine/Plugins/Runtime/GeometryProcessing/Source/DynamicMesh/Private/Solvers/Internal/MeshDiffusionIntegrator.cpp b/Engine/Plugins/Runtime/GeometryProcessing/Source/DynamicMesh/Private/Solvers/Internal/MeshDiffusionIntegrator.cpp
--- a/Engine/Plugins/Runtime/GeometryProcessing/Source/DynamicMesh/Private/Solvers/Internal/MeshDiffusionIntegrator.cpp
+++ b/Engine/Plugins/Runtime/GeometryProcessing/Source/DynamicMesh/Private/Solvers/Internal/MeshDiffusionIntegrator.cpp
@@ -5,7 +5,7 @@
 
 using namespace UE::Geometry;
 
-FMeshDiffusionIntegrator::FMeshDiffusionIntegrator(const FDynamicMesh3& DynamicMesh, const ELaplacianWeightScheme Scheme)
+FMeshDiffusionIntegrator::FMeshDiffusionIntegrator(const FDynamicMesh3& DynamicMesh, const ELaplacianWeightScheme Scheme, FProgressCancel* Progress)
 {
 	Id = 0;
 	bIsSymmetric = false;
@@ -13,6 +13,7 @@ FMeshDiffusionIntegrator::FMeshDiffusionIntegrator(const FDynamicMesh3& DynamicM
 
 	VertexCount = DynamicMesh.VertexCount();
 
+	ProgressCancel = Progress;
 }
 
 void FMeshDiffusionIntegrator::Initialize(const FDynamicMesh3& DynamicMesh, const ELaplacianWeightScheme Scheme)
@@ -86,6 +87,7 @@ void FMeshDiffusionIntegrator::Integrate_ForwardEuler(const int32 NumSteps, cons
 	Id = 0;
 	for (int32 s = 0; s < NumSteps; ++s)
 	{
+		if (ProgressCancel && ProgressCancel->Cancelled()) return;
 
 		int32 SrcBuffer = Id;
 		Id = 1 - Id;
@@ -145,6 +147,7 @@ void FMeshDiffusionIntegrator::Integrate_BackwardEuler(const EMatrixSolverType M
 		Id = 0;
 		for (int32 s = 0; s < NumSteps; ++s)
 		{
+			if (ProgressCancel && ProgressCancel->Cancelled()) return;
 
 			int32 SrcBuffer = Id;
 			Id = 1 - Id;
@@ -164,6 +167,7 @@ void FMeshDiffusionIntegrator::Integrate_BackwardEuler(const EMatrixSolverType M
 		Id = 0;
 		for (int32 s = 0; s < NumSteps; ++s)
 		{
+			if (ProgressCancel && ProgressCancel->Cancelled()) return;
 
 			int32 SrcBuffer = Id;
 			Id = 1 - Id;
diff --git a/Engine/Plugins/Runtime/GeometryProcessing/Source/DynamicMesh/Private/Solvers/Internal/MeshDiffusionIntegrator.h b/Engine/Plugins/Runtime/GeometryProcessing/Source/DynamicMesh/Private/Solvers/Internal/MeshDiffusionIntegrator.h
--- a/Engine/Plugins/Runtime/GeometryProcessing/Source/DynamicMesh/Private/Solvers/Internal/MeshDiffusionIntegrator.h
+++ b/Engine/Plugins/Runtime/GeometryProcessing/Source/DynamicMesh/Private/Solvers/Internal/MeshDiffusionIntegrator.h
@@ -8,6 +8,7 @@
 #include "Solvers/MeshLinearization.h"
 
 #include "MatrixSolver.h"
+#include "Util/ProgressCancel.h"
 
 namespace UE
 {
@@ -17,14 +18,14 @@ namespace Geometry
 
 /**
  * FMeshDiffusionIntegrator solves diffusion problems over vertices of a 3D triangle mesh.
- *
+ * Note: if canceled using the optional FProgressCancel, the results may be in an unusable state. 
  * (Could likely be generalized to other graph diffusion problems)
  */
 class FMeshDiffusionIntegrator
 {
 public:
 
-	FMeshDiffusionIntegrator(const FDynamicMesh3& DynamicMesh, const ELaplacianWeightScheme Scheme);
+	FMeshDiffusionIntegrator(const FDynamicMesh3& DynamicMesh, const ELaplacianWeightScheme Scheme, FProgressCancel* Progress = nullptr);
 	virtual ~FMeshDiffusionIntegrator() {}
 
 	void Integrate_ForwardEuler(const int32 NumSteps, const double Speed);
@@ -75,6 +76,8 @@ protected:
 
 	FSOAPositions                       Tmp[2];
 	int32                               Id; // double buffer id
+
+	FProgressCancel* ProgressCancel;
 };
 
 
diff --git a/Engine/Plugins/Runtime/GeometryProcessing/Source/DynamicMesh/Private/Solvers/Internal/MeshDiffusionSmoothing.h b/Engine/Plugins/Runtime/GeometryProcessing/Source/DynamicMesh/Private/Solvers/Internal/MeshDiffusionSmoothing.h
--- a/Engine/Plugins/Runtime/GeometryProcessing/Source/DynamicMesh/Private/Solvers/Internal/MeshDiffusionSmoothing.h
+++ b/Engine/Plugins/Runtime/GeometryProcessing/Source/DynamicMesh/Private/Solvers/Internal/MeshDiffusionSmoothing.h
@@ -12,13 +12,14 @@ namespace Geometry
 
 /**
  * Solve mesh diffusion using Laplacian matrix of various types
+ * Note: if canceled using the optional FProgressCancel, the results may be in an unusable state. 
  */
 class FLaplacianDiffusionMeshSmoother : public FMeshDiffusionIntegrator
 {
 public:
 
-	FLaplacianDiffusionMeshSmoother(const FDynamicMesh3& DynamicMesh, const ELaplacianWeightScheme Scheme)
-		: FMeshDiffusionIntegrator(DynamicMesh, Scheme)
+	FLaplacianDiffusionMeshSmoother(const FDynamicMesh3& DynamicMesh, const ELaplacianWeightScheme Scheme, FProgressCancel* Progress = nullptr)
+		: FMeshDiffusionIntegrator(DynamicMesh, Scheme, Progress)
 	{
 		Initialize(DynamicMesh, Scheme);
 	}
@@ -41,13 +42,14 @@ protected:
 
 /**
  * Solve mesh diffusion using Biharmonic (ie squared Laplacian) matrix of various types
+ * Note: if canceled using the optional FProgressCancel, the results may be in an unusable state. 
  */
 class  FBiHarmonicDiffusionMeshSmoother : public FMeshDiffusionIntegrator
 {
 public:
 
-	FBiHarmonicDiffusionMeshSmoother(const FDynamicMesh3& DynamicMesh, const ELaplacianWeightScheme Scheme)
-		: FMeshDiffusionIntegrator(DynamicMesh, Scheme)
+	FBiHarmonicDiffusionMeshSmoother(const FDynamicMesh3& DynamicMesh, const ELaplacianWeightScheme Scheme, FProgressCancel* Progress = nullptr)
+		: FMeshDiffusionIntegrator(DynamicMesh, Scheme, Progress)
 	{
 		Initialize(DynamicMesh, Scheme);
 	}
diff --git a/Engine/Plugins/Runtime/GeometryProcessing/Source/DynamicMesh/Private/Solvers/MeshSmoothing.cpp b/Engine/Plugins/Runtime/GeometryProcessing/Source/DynamicMesh/Private/Solvers/MeshSmoothing.cpp
--- a/Engine/Plugins/Runtime/GeometryProcessing/Source/DynamicMesh/Private/Solvers/MeshSmoothing.cpp
+++ b/Engine/Plugins/Runtime/GeometryProcessing/Source/DynamicMesh/Private/Solvers/MeshSmoothing.cpp
@@ -4,11 +4,12 @@
 #include "Solvers/Internal/MeshDiffusionSmoothing.h"
 #include "Solvers/Internal/ConstrainedMeshSmoothers.h"
 #include "MeshWeights.h"
+#include "Util/ProgressCancel.h"
 
 using namespace UE::Geometry;
 
 void UE::MeshDeformation::ComputeSmoothing_BiHarmonic(const ELaplacianWeightScheme WeightScheme, const FDynamicMesh3& OriginalMesh,
-	const double Speed, const double Intensity, const int32 NumIterations, TArray<FVector3d>& PositionArray)
+	const double Speed, const double Intensity, const int32 NumIterations, TArray<FVector3d>& PositionArray, FProgressCancel* Progress)
 {
 	// This is equivalent to taking a single backward Euler time step of bi-harmonic diffusion
 	// where L is the Laplacian (Del^2) , and L^T L is an approximation of the Del^4.
@@ -42,7 +43,7 @@ void UE::MeshDeformation::ComputeSmoothing_BiHarmonic(const ELaplacianWeightSche
 
 	const double TimeStep = Speed * FMath::Min(Intensity, 1.e6);
 
-	FBiHarmonicDiffusionMeshSmoother BiHarmonicDiffusionSmoother(OriginalMesh, WeightScheme);
+	FBiHarmonicDiffusionMeshSmoother BiHarmonicDiffusionSmoother(OriginalMesh, WeightScheme, Progress);
 
 	BiHarmonicDiffusionSmoother.Integrate_BackwardEuler(MatrixSolverType, NumIterations, TimeStep);
 
@@ -91,7 +92,7 @@ void UE::MeshDeformation::ComputeSmoothing_ImplicitBiHarmonicPCG(const ELaplacia
 }
 
 void UE::MeshDeformation::ComputeSmoothing_Diffusion(const ELaplacianWeightScheme WeightScheme, const FDynamicMesh3& OriginalMesh, bool bForwardEuler,
-	const double Speed, const double Intensity, const int32 IterationCount, TArray<FVector3d>& PositionArray)
+	const double Speed, const double Intensity, const int32 IterationCount, TArray<FVector3d>& PositionArray, FProgressCancel* Progress)
 {
 #ifndef EIGEN_MPL2_ONLY
 	const EMatrixSolverType MatrixSolverType = EMatrixSolverType::LTL;
@@ -112,7 +113,7 @@ void UE::MeshDeformation::ComputeSmoothing_Diffusion(const ELaplacianWeightSchem
 #endif
 	if (IterationCount < 1) return;
 
-	FLaplacianDiffusionMeshSmoother Smoother(OriginalMesh, WeightScheme);
+	FLaplacianDiffusionMeshSmoother Smoother(OriginalMesh, WeightScheme, Progress);
 
 	if (bForwardEuler)
 	{
@@ -129,7 +130,7 @@ void UE::MeshDeformation::ComputeSmoothing_Diffusion(const ELaplacianWeightSchem
 
 
 void UE::MeshDeformation::ComputeSmoothing_Forward(bool bUniformWeightScheme, bool bSmoothBoundary, const FDynamicMesh3& OriginalMesh,
-	TFunctionRef<double(int VID, bool bOnBoundary)> GetSmoothingAlpha, const int32 NumIterations, TArray<FVector3d>& PositionArray)
+	TFunctionRef<double(int VID, bool bOnBoundary)> GetSmoothingAlpha, const int32 NumIterations, TArray<FVector3d>& PositionArray, FProgressCancel* Progress)
 {
 	int32 NV = OriginalMesh.MaxVertexID();
 
@@ -157,6 +158,8 @@ void UE::MeshDeformation::ComputeSmoothing_Forward(bool bUniformWeightScheme, bo
 
 	for (int32 Iter = 0; Iter < NumIterations; ++Iter)
 	{
+		if (Progress && Progress->Cancelled()) return;
+
 		// calculate smoothed positions of interior vertices
 		ParallelFor(NV, [bUniformWeightScheme, &OriginalMesh, &GetSmoothingAlpha, &IsBoundary, &CurrentPositions, &NextPositions](int32 VID)
 		{
diff --git a/Engine/Plugins/Runtime/GeometryProcessing/Source/DynamicMesh/Public/Solvers/MeshSmoothing.h b/Engine/Plugins/Runtime/GeometryProcessing/Source/DynamicMesh/Public/Solvers/MeshSmoothing.h
--- a/Engine/Plugins/Runtime/GeometryProcessing/Source/DynamicMesh/Public/Solvers/MeshSmoothing.h
+++ b/Engine/Plugins/Runtime/GeometryProcessing/Source/DynamicMesh/Public/Solvers/MeshSmoothing.h
@@ -5,6 +5,7 @@
 #include "CoreMinimal.h"
 #include "DynamicMesh/DynamicMesh3.h"
 #include "Solvers/MeshLaplacian.h"
+#include "Util/ProgressCancel.h"
 
 namespace UE
 {
@@ -34,9 +35,10 @@ namespace UE
 		* weight = 1 / (k * Sqrt[dt] )
 		*
 		* The result is returned in the PositionArray
+		* Note: if canceled using the optional FProgressCancel, the results may be in an unusable state. 
 		*/
 		void DYNAMICMESH_API ComputeSmoothing_BiHarmonic(const ELaplacianWeightScheme WeightingScheme, const FDynamicMesh3& OriginalMesh,
-			const double Speed, const double Weight, const int32 NumIterations, TArray<FVector3d>& PositionArray);
+			const double Speed, const double Weight, const int32 NumIterations, TArray<FVector3d>& PositionArray, FProgressCancel* Progress = nullptr);
 
 		void  DYNAMICMESH_API ComputeSmoothing_ImplicitBiHarmonicPCG(const ELaplacianWeightScheme WeightScheme, const FDynamicMesh3& OriginalMesh,
 			const double Speed, const double Weight, const int32 MaxIterations, TArray<FVector3d>& PositionArray);
@@ -51,9 +53,10 @@ namespace UE
 		* with dt = Speed / Max(|w_ii|)
 		*
 		* here w_ii are the diagonal values of L
+		* Note: if canceled using the optional FProgressCancel, the results may be in an unusable state. 
 		*/
 		void  DYNAMICMESH_API ComputeSmoothing_Diffusion(const ELaplacianWeightScheme WeightScheme, const FDynamicMesh3& OriginalMesh, bool bForwardEuler,
-			const double Speed, double Weight, const int32 NumIterations, TArray<FVector3d>& PositionArray);
+			const double Speed, double Weight, const int32 NumIterations, TArray<FVector3d>& PositionArray, FProgressCancel* Progress = nullptr);
 
 
 		/**
@@ -66,9 +69,11 @@ namespace UE
 		 * @param GetSmoothingAlpha The amount to lerp each vertex towards the centroid, per iteration
 		 * @param NumIterations Number of iterations to smooth
 		 * @param PositionsArray Vertex positions (input and output)
+		 * @param Progress An optional interrupter that can terminate the iterations, leaving the results in an unusable state. 
 		 */
 		void DYNAMICMESH_API ComputeSmoothing_Forward(bool bUniformWeightScheme, bool bSmoothBoundary, const FDynamicMesh3& OriginalMesh,
-			TFunctionRef<double(int VID, bool bBoundary)> GetSmoothingAlpha, const int32 NumIterations, TArray<FVector3d>& PositionArray);
+			TFunctionRef<double(int VID, bool bBoundary)> GetSmoothingAlpha, const int32 NumIterations, TArray<FVector3d>& PositionArray, 
+			class FProgressCancel* Progress = nullptr);
 
 	}
 }
\ No newline at end of file
diff --git a/Engine/Plugins/Runtime/MeshModelingToolset/Source/ModelingOperators/Private/SmoothingOps/IterativeSmoothingOp.cpp b/Engine/Plugins/Runtime/MeshModelingToolset/Source/ModelingOperators/Private/SmoothingOps/IterativeSmoothingOp.cpp
--- a/Engine/Plugins/Runtime/MeshModelingToolset/Source/ModelingOperators/Private/SmoothingOps/IterativeSmoothingOp.cpp
+++ b/Engine/Plugins/Runtime/MeshModelingToolset/Source/ModelingOperators/Private/SmoothingOps/IterativeSmoothingOp.cpp
@@ -23,16 +23,21 @@ void FIterativeSmoothingOp::CalculateResult(FProgressCancel* Progress)
 	{
 		if (SmoothOptions.bUniform)
 		{
-			Smooth_MeanValue();
+			Smooth_MeanValue(Progress);
 		}
 		else
 		{
-			Smooth_Implicit_Cotan();
+			Smooth_Implicit_Cotan(Progress);
 		}
 	}
 	else
 	{
-		Smooth_Forward(SmoothOptions.bUniform);
+		Smooth_Forward(SmoothOptions.bUniform, Progress);
+	}
+
+	if (Progress && Progress->Cancelled())
+	{
+		return;
 	}
 
 	// Copy the results back into the result mesh and update normals
@@ -53,24 +58,24 @@ double FIterativeSmoothingOp::GetSmoothAlpha(int32 VertexID, bool bIsBoundary)
 }
 
 
-void FIterativeSmoothingOp::Smooth_Implicit_Cotan()
+void FIterativeSmoothingOp::Smooth_Implicit_Cotan(FProgressCancel* Progress)
 {
 	double Intensity = 1.;
 	UE::MeshDeformation::ComputeSmoothing_BiHarmonic(ELaplacianWeightScheme::ClampedCotangent, 
-		*ResultMesh, SmoothOptions.SmoothAlpha, Intensity, SmoothOptions.Iterations, PositionBuffer);
+		*ResultMesh, SmoothOptions.SmoothAlpha, Intensity, SmoothOptions.Iterations, PositionBuffer, Progress);
 }
 
 
-void FIterativeSmoothingOp::Smooth_MeanValue()
+void FIterativeSmoothingOp::Smooth_MeanValue(FProgressCancel* Progress)
 {
 	double Intensity = 1.;
 	UE::MeshDeformation::ComputeSmoothing_BiHarmonic(ELaplacianWeightScheme::MeanValue,
-		*ResultMesh, SmoothOptions.SmoothAlpha, Intensity, SmoothOptions.Iterations, PositionBuffer);
+		*ResultMesh, SmoothOptions.SmoothAlpha, Intensity, SmoothOptions.Iterations, PositionBuffer, Progress);
 }
 
 
-void FIterativeSmoothingOp::Smooth_Forward(bool bUniform)
+void FIterativeSmoothingOp::Smooth_Forward(bool bUniform, FProgressCancel* Progress)
 {
 	UE::MeshDeformation::ComputeSmoothing_Forward(bUniform, SmoothOptions.bSmoothBoundary,
-		*ResultMesh, [this](int VID, bool bBoundary) { return GetSmoothAlpha(VID, bBoundary); }, SmoothOptions.Iterations, PositionBuffer);
+		*ResultMesh, [this](int VID, bool bBoundary) { return GetSmoothAlpha(VID, bBoundary); }, SmoothOptions.Iterations, PositionBuffer, Progress);
 }
diff --git a/Engine/Plugins/Runtime/MeshModelingToolset/Source/ModelingOperators/Public/SmoothingOps/IterativeSmoothingOp.h b/Engine/Plugins/Runtime/MeshModelingToolset/Source/ModelingOperators/Public/SmoothingOps/IterativeSmoothingOp.h
--- a/Engine/Plugins/Runtime/MeshModelingToolset/Source/ModelingOperators/Public/SmoothingOps/IterativeSmoothingOp.h
+++ b/Engine/Plugins/Runtime/MeshModelingToolset/Source/ModelingOperators/Public/SmoothingOps/IterativeSmoothingOp.h
@@ -21,6 +21,7 @@ public:
 	~FIterativeSmoothingOp() override {};
 
 	// Apply smoothing. results in an updated ResultMesh
+    //Note: if canceled using the optional FProgressCancel, the results may be in an unusable state.
 	void CalculateResult(FProgressCancel* Progress) override;
 
 private:
@@ -28,13 +29,13 @@ private:
 	double GetSmoothAlpha(int32 VertexID, bool bIsBoundary);
 
 	// uniform iterative smoothing
-	void Smooth_Forward(bool bUniform);
+	void Smooth_Forward(bool bUniform, FProgressCancel* Progress = nullptr);
 
 	// cotan smoothing iterations
-	void Smooth_Implicit_Cotan();
+	void Smooth_Implicit_Cotan(FProgressCancel* Progress = nullptr);
 
 	// mean value smoothing iterations
-	void Smooth_MeanValue();
+	void Smooth_MeanValue(FProgressCancel* Progress = nullptr);
 };
 
 } // end namespace UE::Geometry
