Make PolygroupSet::CheckExists a const method

diff --git a/Engine/Plugins/Runtime/GeometryProcessing/Source/DynamicMesh/Private/Polygroups/PolygroupSet.cpp b/Engine/Plugins/Runtime/GeometryProcessing/Source/DynamicMesh/Private/Polygroups/PolygroupSet.cpp
--- a/Engine/Plugins/Runtime/GeometryProcessing/Source/DynamicMesh/Private/Polygroups/PolygroupSet.cpp
+++ b/Engine/Plugins/Runtime/GeometryProcessing/Source/DynamicMesh/Private/Polygroups/PolygroupSet.cpp
@@ -6,7 +6,7 @@
 using namespace UE::Geometry;
 
 
-bool FPolygroupLayer::CheckExists(const FDynamicMesh3* Mesh)
+bool FPolygroupLayer::CheckExists(const FDynamicMesh3* Mesh) const
 {
 	if (Mesh)
 	{
diff --git a/Engine/Plugins/Runtime/GeometryProcessing/Source/DynamicMesh/Public/Polygroups/PolygroupSet.h b/Engine/Plugins/Runtime/GeometryProcessing/Source/DynamicMesh/Public/Polygroups/PolygroupSet.h
--- a/Engine/Plugins/Runtime/GeometryProcessing/Source/DynamicMesh/Public/Polygroups/PolygroupSet.h
+++ b/Engine/Plugins/Runtime/GeometryProcessing/Source/DynamicMesh/Public/Polygroups/PolygroupSet.h
@@ -42,7 +42,7 @@ struct DYNAMICMESH_API FPolygroupLayer
 	}
 
 	/** @return true if the specified layer (default or extended) exist and is initialized on the given Mesh */
-	bool CheckExists(const FDynamicMesh3* Mesh);
+	bool CheckExists(const FDynamicMesh3* Mesh) const;
 };
 
 
