fix grid size overflowing int for convex hull and convex decomposition methods

diff --git a/Engine/Plugins/Runtime/GeometryProcessing/Source/DynamicMesh/Private/Operations/MeshConvexHull.cpp b/Engine/Plugins/Runtime/GeometryProcessing/Source/DynamicMesh/Private/Operations/MeshConvexHull.cpp
--- a/Engine/Plugins/Runtime/GeometryProcessing/Source/DynamicMesh/Private/Operations/MeshConvexHull.cpp
+++ b/Engine/Plugins/Runtime/GeometryProcessing/Source/DynamicMesh/Private/Operations/MeshConvexHull.cpp
@@ -216,7 +216,8 @@ void FMeshConvexHull::GridSample(const FDynamicMesh3& Mesh,
 	FAxisAlignedBox3d Bounds = Mesh.GetBounds();
 	Bounds.Min = Bounds.Min - 1e-4;			// Pad to avoid problems with vertices lying exactly on bounding box
 	Bounds.Max = Bounds.Max + 1e-4;
-	const double GridCellSize = Bounds.MaxDim() / (double)GridResolutionMaxAxis;
+	// Clamp grid resolution to prevent overflowing int
+	const double GridCellSize = Bounds.MaxDim() / (double)FMath::Clamp(GridResolutionMaxAxis, 2, 1290);
 
 	FBoundsGridIndexer3d Indexer(Bounds, GridCellSize);
 	const FVector3i GridResolution = Indexer.GridResolution();
diff --git a/Engine/Plugins/Runtime/GeometryScripting/Source/GeometryScriptingCore/Public/GeometryScript/CollisionFunctions.h b/Engine/Plugins/Runtime/GeometryScripting/Source/GeometryScriptingCore/Public/GeometryScript/CollisionFunctions.h
--- a/Engine/Plugins/Runtime/GeometryScripting/Source/GeometryScriptingCore/Public/GeometryScript/CollisionFunctions.h
+++ b/Engine/Plugins/Runtime/GeometryScripting/Source/GeometryScriptingCore/Public/GeometryScript/CollisionFunctions.h
@@ -167,7 +167,7 @@ struct GEOMETRYSCRIPTINGCORE_API FComputeNegativeSpaceOptions
 	bool bOnlyConnectedToHull = false;
 
 	/** When performing Voxel Search, maximum number of voxels to use along each dimension */
-	UPROPERTY(EditAnywhere, BlueprintReadWrite, Category = NegativeSpace, meta = (ClampMin = 4, ClampMax = 4096))
+	UPROPERTY(EditAnywhere, BlueprintReadWrite, Category = NegativeSpace, meta = (ClampMin = 4, ClampMax = 1290))
 	int32 MaxVoxelsPerDim = 128;
 
 	/** Approximate number of spheres to consider when covering negative space */
diff --git a/Engine/Source/Runtime/GeometryCore/Private/CompGeom/ConvexDecomposition3.cpp b/Engine/Source/Runtime/GeometryCore/Private/CompGeom/ConvexDecomposition3.cpp
--- a/Engine/Source/Runtime/GeometryCore/Private/CompGeom/ConvexDecomposition3.cpp
+++ b/Engine/Source/Runtime/GeometryCore/Private/CompGeom/ConvexDecomposition3.cpp
@@ -339,6 +339,8 @@ bool FSphereCovering::AddNegativeSpace(const TFastWindingTree<FDynamicMesh3>& Sp
 		{
 			MarchingCubes.Bounds = HullAABB.GetBoundingBox();
 			MarchingCubes.Bounds.Expand(SampleSettings.MinRadius + SampleSettings.ReduceRadiusMargin + UE_DOUBLE_KINDA_SMALL_NUMBER);
+			// enforce max voxels per dim after bounds expand
+			MarchingCubes.CubeSize = FMath::Max(MarchingCubes.CubeSize, MarchingCubes.Bounds.MaxDim() / (double)SampleSettings.MaxVoxelsPerDim);
 			const double MinRadSq = SampleSettings.MinRadius * SampleSettings.MinRadius;
 			MarchingCubes.Implicit = [&HullMeshWrap, &HullAABB, &Spatial, WindingSign, &SampleSettings, MinRadSq](const FVector3d& Pt) -> double
 			{
diff --git a/Engine/Source/Runtime/GeometryCore/Public/CompGeom/ConvexDecomposition3.h b/Engine/Source/Runtime/GeometryCore/Public/CompGeom/ConvexDecomposition3.h
--- a/Engine/Source/Runtime/GeometryCore/Public/CompGeom/ConvexDecomposition3.h
+++ b/Engine/Source/Runtime/GeometryCore/Public/CompGeom/ConvexDecomposition3.h
@@ -148,7 +148,7 @@ struct FNegativeSpaceSampleSettings
 		MinSpacing = FMath::Max(0.0, MinSpacing);
 		ReduceRadiusMargin = FMath::Max(0.0, ReduceRadiusMargin);
 		MinRadius = FMath::Max(0.0, MinRadius);
-		MaxVoxelsPerDim = FMath::Clamp(MaxVoxelsPerDim, 4, 4096);
+		MaxVoxelsPerDim = FMath::Clamp(MaxVoxelsPerDim, 4, 1290);
 		MinCustomMarginScale = FMath::Clamp(MinCustomMarginScale, .1, 1.0);
 	}
 
