CIS fix -- change mesh convex hull grid sampling acceptable minimum grid sample to 1, b/c tests expect us to allow that

diff --git a/Engine/Plugins/Runtime/GeometryProcessing/Source/DynamicMesh/Private/Operations/MeshConvexHull.cpp b/Engine/Plugins/Runtime/GeometryProcessing/Source/DynamicMesh/Private/Operations/MeshConvexHull.cpp
--- a/Engine/Plugins/Runtime/GeometryProcessing/Source/DynamicMesh/Private/Operations/MeshConvexHull.cpp
+++ b/Engine/Plugins/Runtime/GeometryProcessing/Source/DynamicMesh/Private/Operations/MeshConvexHull.cpp
@@ -217,7 +217,7 @@ void FMeshConvexHull::GridSample(const FDynamicMesh3& Mesh,
 	Bounds.Min = Bounds.Min - 1e-4;			// Pad to avoid problems with vertices lying exactly on bounding box
 	Bounds.Max = Bounds.Max + 1e-4;
 	// Clamp grid resolution to prevent overflowing int
-	const double GridCellSize = Bounds.MaxDim() / (double)FMath::Clamp(GridResolutionMaxAxis, 2, 1290);
+	const double GridCellSize = Bounds.MaxDim() / (double)FMath::Clamp(GridResolutionMaxAxis, 1, 1290);
 
 	FBoundsGridIndexer3d Indexer(Bounds, GridCellSize);
 	const FVector3i GridResolution = Indexer.GridResolution();
