MeshModelingTools: Fix topology deformer breaking on coincident verts. Also avoid a minor computation that is only used in debug mode.

diff --git a/Engine/Plugins/Runtime/GeometryProcessing/Source/DynamicMesh/Private/Operations/GroupTopologyDeformer.cpp b/Engine/Plugins/Runtime/GeometryProcessing/Source/DynamicMesh/Private/Operations/GroupTopologyDeformer.cpp
--- a/Engine/Plugins/Runtime/GeometryProcessing/Source/DynamicMesh/Private/Operations/GroupTopologyDeformer.cpp
+++ b/Engine/Plugins/Runtime/GeometryProcessing/Source/DynamicMesh/Private/Operations/GroupTopologyDeformer.cpp
@@ -317,8 +317,22 @@ void FGroupTopologyDeformer::ComputeEncoding()
 			{
 				FVector3d BorderPos = Mesh->GetVertex(Face.BoundaryVerts[k]);
 				FVector3d DeltaVec = Pos - BorderPos;
-				double Weight = 1.0 / DeltaVec.SquaredLength();
 				VtxEncoding.Deltas[k] = DeltaVec;
+
+				double DeltaSquaredLength = DeltaVec.SquaredLength();
+				if (DeltaSquaredLength == 0)
+				{
+					// This weight should be 1, and everything else should be zero.
+					for (int32 i = 0; i < VtxEncoding.Weights.Num(); ++i)
+					{
+						VtxEncoding.Weights[i] = 0;
+					}
+					VtxEncoding.Weights[k] = 1;
+					WeightSum = 1;
+					break;
+				}
+
+				double Weight = 1.0 / DeltaVec.SquaredLength();
 				VtxEncoding.Weights[k] = Weight;
 				WeightSum += Weight;
 			}
@@ -326,7 +340,9 @@ void FGroupTopologyDeformer::ComputeEncoding()
 			for (int k = 0; k < NumBoundaryV; ++k)
 			{
 				VtxEncoding.Weights[k] /= WeightSum;
+#if DO_GUARD_SLOW
 				Reconstruct += VtxEncoding.Weights[k] * (Mesh->GetVertex(Face.BoundaryVerts[k]) + VtxEncoding.Deltas[k]);
+#endif
 			}
 			checkSlow(Distance(Reconstruct, Pos) < 0.0001);   // sanity check
 		}
