Fix scaling field and expose the correct ideal edge length parameter for the ftetwild wrapper

diff --git a/Engine/Plugins/Experimental/ChaosFlesh/Source/ChaosFleshDeprecatedNodes/Private/Dataflow/ChaosFleshCreateTetrahedralCollectionNode.cpp b/Engine/Plugins/Experimental/ChaosFlesh/Source/ChaosFleshDeprecatedNodes/Private/Dataflow/ChaosFleshCreateTetrahedralCollectionNode.cpp
--- a/Engine/Plugins/Experimental/ChaosFlesh/Source/ChaosFleshDeprecatedNodes/Private/Dataflow/ChaosFleshCreateTetrahedralCollectionNode.cpp
+++ b/Engine/Plugins/Experimental/ChaosFlesh/Source/ChaosFleshDeprecatedNodes/Private/Dataflow/ChaosFleshCreateTetrahedralCollectionNode.cpp
@@ -230,7 +230,7 @@ void FGenerateTetrahedralCollectionDataflowNodes::EvaluateTetWild(
 		Params.EpsRel = EpsRel;
 		Params.MaxIts = MaxIterations;
 		Params.StopEnergy = StopEnergy;
-		Params.IdealEdgeLength = IdealEdgeLength;
+		Params.IdealEdgeLengthRel = IdealEdgeLengthRel;
 
 		Params.bInvertOutputTets = bInvertOutputTets;
 
diff --git a/Engine/Plugins/Experimental/ChaosFlesh/Source/ChaosFleshDeprecatedNodes/Public/Dataflow/ChaosFleshCreateTetrahedralCollectionNode.h b/Engine/Plugins/Experimental/ChaosFlesh/Source/ChaosFleshDeprecatedNodes/Public/Dataflow/ChaosFleshCreateTetrahedralCollectionNode.h
--- a/Engine/Plugins/Experimental/ChaosFlesh/Source/ChaosFleshDeprecatedNodes/Public/Dataflow/ChaosFleshCreateTetrahedralCollectionNode.h
+++ b/Engine/Plugins/Experimental/ChaosFlesh/Source/ChaosFleshDeprecatedNodes/Public/Dataflow/ChaosFleshCreateTetrahedralCollectionNode.h
@@ -21,10 +21,13 @@ namespace UE {
 }
 
 
+// Note: disabled deprecation warnings due to deprecated member: IdealEdgeLength. Can remove these disable/enable pragmas once that member variable is removed.
+PRAGMA_DISABLE_DEPRECATION_WARNINGS
 USTRUCT(meta = (DataflowFlesh, Deprecated = "5.4"))
 struct FGenerateTetrahedralCollectionDataflowNodes : public FDataflowNode
 {
 	GENERATED_USTRUCT_BODY()
+	PRAGMA_ENABLE_DEPRECATION_WARNINGS
 	DATAFLOW_NODE_DEFINE_INTERNAL(FGenerateTetrahedralCollectionDataflowNodes, "GenerateTetrahedralCollection", "Flesh", "")
 	DATAFLOW_NODE_RENDER_TYPE("SurfaceRender", FGeometryCollection::StaticType(), "Collection")
 
@@ -48,9 +51,12 @@ public:
 	// TetWild
 	//
 
-	//! Energy at which to stop optimizing tet quality and accept the result.
+	//! Desired relative edge length, as a fraction of bounding box size
 	UPROPERTY(EditAnywhere, Category = "Dataflow", meta = (ClampMin = "0.0", EditCondition = "Method == TetMeshingMethod::TetWild", EditConditionHides))
-	double IdealEdgeLength = 1.0;
+	double IdealEdgeLengthRel = 0.05;
+
+	UE_DEPRECATED(5.6, "Use IdealEdgeLengthRel instead, which is relative to the bounding box size")
+	double IdealEdgeLength = 0.05;
 
 	//! Maximum number of optimization iterations.
 	UPROPERTY(EditAnywhere, Category = "Dataflow", meta = (ClampMin = "1", EditCondition = "Method == TetMeshingMethod::TetWild", EditConditionHides))
diff --git a/Engine/Plugins/Experimental/ChaosFlesh/Source/ChaosFleshNodes/Private/Dataflow/ChaosFleshCreateTetrahedronNode.cpp b/Engine/Plugins/Experimental/ChaosFlesh/Source/ChaosFleshNodes/Private/Dataflow/ChaosFleshCreateTetrahedronNode.cpp
--- a/Engine/Plugins/Experimental/ChaosFlesh/Source/ChaosFleshNodes/Private/Dataflow/ChaosFleshCreateTetrahedronNode.cpp
+++ b/Engine/Plugins/Experimental/ChaosFlesh/Source/ChaosFleshNodes/Private/Dataflow/ChaosFleshCreateTetrahedronNode.cpp
@@ -275,7 +275,7 @@ void FCreateTetrahedronDataflowNode::EvaluateTetWild(
 		Params.EpsRel = EpsRel;
 		Params.MaxIts = MaxIterations;
 		Params.StopEnergy = StopEnergy;
-		Params.IdealEdgeLength = IdealEdgeLength;
+		Params.IdealEdgeLengthRel = IdealEdgeLengthRel;
 
 		Params.bInvertOutputTets = bInvertOutputTets;
 
diff --git a/Engine/Plugins/Experimental/ChaosFlesh/Source/ChaosFleshNodes/Public/Dataflow/ChaosFleshCreateTetrahedronNode.h b/Engine/Plugins/Experimental/ChaosFlesh/Source/ChaosFleshNodes/Public/Dataflow/ChaosFleshCreateTetrahedronNode.h
--- a/Engine/Plugins/Experimental/ChaosFlesh/Source/ChaosFleshNodes/Public/Dataflow/ChaosFleshCreateTetrahedronNode.h
+++ b/Engine/Plugins/Experimental/ChaosFlesh/Source/ChaosFleshNodes/Public/Dataflow/ChaosFleshCreateTetrahedronNode.h
@@ -17,10 +17,13 @@ namespace UE {
 	}
 }
 
+// Note: disabled deprecation warnings due to deprecated member: IdealEdgeLength. Can remove these disable/enable pragmas once that member variable is removed.
+PRAGMA_DISABLE_DEPRECATION_WARNINGS
 USTRUCT(meta = (DataflowFlesh))
 struct FCreateTetrahedronDataflowNode : public FDataflowNode
 {
 	GENERATED_USTRUCT_BODY()
+	PRAGMA_ENABLE_DEPRECATION_WARNINGS
 	DATAFLOW_NODE_DEFINE_INTERNAL(FCreateTetrahedronDataflowNode, "CreateTetrahedron", "Flesh", "")
 	DATAFLOW_NODE_RENDER_TYPE("SurfaceRender", FGeometryCollection::StaticType(), "Collection")
 
@@ -52,9 +55,12 @@ public:
 	// TetWild
 	//
 
-	//! Energy at which to stop optimizing tet quality and accept the result.
+	//! Desired relative edge length, as a fraction of bounding box size
 	UPROPERTY(EditAnywhere, Category = "Dataflow", meta = (ClampMin = "0.0", EditCondition = "Method == TetMeshingMethod::TetWild", EditConditionHides))
-	double IdealEdgeLength = 1.0;
+	double IdealEdgeLengthRel = 0.05;
+
+	UE_DEPRECATED(5.6, "Use IdealEdgeLengthRel instead, which is relative to the bounding box size")
+	double IdealEdgeLength = 0.05;
 
 	//! Maximum number of optimization iterations.
 	UPROPERTY(EditAnywhere, Category = "Dataflow", meta = (ClampMin = "1", EditCondition = "Method == TetMeshingMethod::TetWild", EditConditionHides))
diff --git a/Engine/Plugins/Runtime/GeometryProcessing/Source/GeometryAlgorithms/Private/FTetWildWrapper.cpp b/Engine/Plugins/Runtime/GeometryProcessing/Source/GeometryAlgorithms/Private/FTetWildWrapper.cpp
--- a/Engine/Plugins/Runtime/GeometryProcessing/Source/GeometryAlgorithms/Private/FTetWildWrapper.cpp
+++ b/Engine/Plugins/Runtime/GeometryProcessing/Source/GeometryAlgorithms/Private/FTetWildWrapper.cpp
@@ -44,21 +44,21 @@ namespace
 		Params.apply_sizing_field = InParams.bApplySizing;
 		if (Params.apply_sizing_field)
 		{
-			Params.V_sizing_field.resize(InParams.SizingFieldVertices.Num(), 3);
-			Params.T_sizing_field.resize(InParams.SizingFieldTets.Num(), 4);
+			Params.V_sizing_field.resize(InParams.SizingFieldVertices.Num() * 3, 1);
+			Params.T_sizing_field.resize(InParams.SizingFieldTets.Num() * 4, 1);
 			Params.values_sizing_field.resize(InParams.SizingFieldValues.Num(), 1);
 			for (int32 Idx = 0; Idx < InParams.SizingFieldVertices.Num(); ++Idx)
 			{
 				for (int32 SubIdx = 0; SubIdx < 3; ++SubIdx)
 				{
-					Params.V_sizing_field.row(Idx)[SubIdx] = InParams.SizingFieldVertices[Idx][SubIdx];
+					Params.V_sizing_field(Idx * 3 + SubIdx) = InParams.SizingFieldVertices[Idx][SubIdx];
 				}
 			}
 			for (int32 Idx = 0; Idx < InParams.SizingFieldTets.Num(); ++Idx)
 			{
 				for (int32 SubIdx = 0; SubIdx < 4; ++SubIdx)
 				{
-					Params.T_sizing_field.row(Idx)[SubIdx] = InParams.SizingFieldTets[Idx][SubIdx];
+					Params.T_sizing_field(Idx * 4 + SubIdx) = InParams.SizingFieldTets[Idx][SubIdx];
 				}
 			}
 			for (int32 Idx = 0; Idx < InParams.SizingFieldValues.Num(); ++Idx)
@@ -66,7 +66,7 @@ namespace
 				Params.values_sizing_field[Idx] = InParams.SizingFieldValues[Idx];
 			}
 		}
-		Params.ideal_edge_length = InParams.IdealEdgeLength;
+		Params.ideal_edge_length_rel = InParams.IdealEdgeLengthRel;
 		Params.eps_rel = InParams.EpsRel;
 		Params.max_its = InParams.MaxIts;
 		Params.stop_energy = InParams.StopEnergy;
diff --git a/Engine/Plugins/Runtime/GeometryProcessing/Source/GeometryAlgorithms/Public/FTetWildWrapper.h b/Engine/Plugins/Runtime/GeometryProcessing/Source/GeometryAlgorithms/Public/FTetWildWrapper.h
--- a/Engine/Plugins/Runtime/GeometryProcessing/Source/GeometryAlgorithms/Public/FTetWildWrapper.h
+++ b/Engine/Plugins/Runtime/GeometryProcessing/Source/GeometryAlgorithms/Public/FTetWildWrapper.h
@@ -33,7 +33,10 @@ namespace UE::Geometry::FTetWild
 		TArray<FIntVector4> SizingFieldTets;
 		TArray<double> SizingFieldValues;
 
-		// Ideal edge length
+		// Ideal edge length relative to bounding box size
+		double IdealEdgeLengthRel = 0.05;
+
+		UE_DEPRECATED(5.6, "Use IdealEdgeLengthRel instead.")
 		double IdealEdgeLength = 0.05;
 
 		// Relative tolerance, controlling how closely the mesh must follow the input surface
