ModelingTools: Added alpha support for vertex color bakes in MeshPropertyMapEvaluator.

diff --git a/Engine/Plugins/Runtime/GeometryProcessing/Source/DynamicMesh/Private/Sampling/MeshPropertyMapEvaluator.cpp b/Engine/Plugins/Runtime/GeometryProcessing/Source/DynamicMesh/Private/Sampling/MeshPropertyMapEvaluator.cpp
--- a/Engine/Plugins/Runtime/GeometryProcessing/Source/DynamicMesh/Private/Sampling/MeshPropertyMapEvaluator.cpp
+++ b/Engine/Plugins/Runtime/GeometryProcessing/Source/DynamicMesh/Private/Sampling/MeshPropertyMapEvaluator.cpp
@@ -37,6 +37,11 @@ void FMeshPropertyMapEvaluator::Setup(const FMeshBaseBaker& Baker, FEvaluationCo
 	Context.AccumulateMode = EAccumulateMode::Add;
 	Context.DataLayout = DataLayout();
 
+	if (Property == EMeshPropertyMapType::MaterialID || Property == EMeshPropertyMapType::PolyGroupID)
+	{
+		Context.AccumulateMode = EAccumulateMode::Overwrite;
+	}
+
 	Bounds = DetailSampler->GetBounds();
 	for (int32 j = 0; j < 3; ++j)
 	{
@@ -47,56 +52,65 @@ void FMeshPropertyMapEvaluator::Setup(const FMeshBaseBaker& Baker, FEvaluationCo
 		}
 	}
 
-	DefaultValue = FVector3f::Zero();
-	switch (this->Property)
+PRAGMA_DISABLE_DEPRECATION_WARNINGS
+	DefaultValue = GetDefaultValue(this->Property);
+PRAGMA_ENABLE_DEPRECATION_WARNINGS
+	DefaultValue4f = GetDefaultValue(this->Property);
+}
+
+const TArray<FMeshMapEvaluator::EComponents>& FMeshPropertyMapEvaluator::DataLayout() const
+{
+	static const TArray<EComponents> Layout{ EComponents::Float4 };
+	return Layout;
+}
+
+FVector4f FMeshPropertyMapEvaluator::GetDefaultValue(const EMeshPropertyMapType InProperty) const
+{
+	FVector4f Value;
+	switch (InProperty)
 	{
 	case EMeshPropertyMapType::Position:
-		DefaultValue = PositionToColor(Bounds.Center(), Bounds);
+		Value = PositionToColor(Bounds.Center(), Bounds);
 		break;
 	case EMeshPropertyMapType::FacetNormal:
-		DefaultValue = NormalToColor(FVector3d::UnitZ());
+		Value = NormalToColor(FVector3f::UnitZ());
 		break;
 	case EMeshPropertyMapType::Normal:
-		DefaultValue = NormalToColor(FVector3d::UnitZ());
+		Value = NormalToColor(FVector3f::UnitZ());
 		break;
 	case EMeshPropertyMapType::UVPosition:
-		DefaultValue = UVToColor(FVector2f::Zero());
+		Value = UVToColor(FVector2f::Zero());
 		break;
 	case EMeshPropertyMapType::MaterialID:
 	case EMeshPropertyMapType::PolyGroupID:
-		DefaultValue = FVector3f(LinearColors::LightPink3f());
-		Context.AccumulateMode = EAccumulateMode::Overwrite;
+		Value = FVector4f(LinearColors::LightPink3f(), 1.0f);
 		break;
 	case EMeshPropertyMapType::VertexColor:
-		DefaultValue = FVector3f::One();
+		Value = FVector4f::One();
 		break;
 	}
+	return Value;
 }
 
-const TArray<FMeshMapEvaluator::EComponents>& FMeshPropertyMapEvaluator::DataLayout() const
-{
-	static const TArray<EComponents> Layout{ EComponents::Float3 };
-	return Layout;
-}
 
 template <bool bUseDetailNormalMap>
 void FMeshPropertyMapEvaluator::EvaluateSample(float*& Out, const FCorrespondenceSample& Sample, void* EvalData)
 {
 	const FMeshPropertyMapEvaluator* Eval = static_cast<FMeshPropertyMapEvaluator*>(EvalData);
-	const FVector3f SampleResult = Eval->SampleFunction<bUseDetailNormalMap>(Sample);
+	const FVector4f SampleResult = Eval->SampleFunction<bUseDetailNormalMap>(Sample);
 	WriteToBuffer(Out, SampleResult);
 }
 
 void FMeshPropertyMapEvaluator::EvaluateDefault(float*& Out, void* EvalData)
 {
-	WriteToBuffer(Out, FVector3f::Zero());
+	WriteToBuffer(Out, FVector4f::Zero());
 }
 
 void FMeshPropertyMapEvaluator::EvaluateColor(const int DataIdx, float*& In, FVector4f& Out, void* EvalData)
 {
 	// TODO: Move property color space transformation from EvaluateSample/Default to here.
-	Out = FVector4f(In[0], In[1], In[2], 1.0f);
-	In += 3;
+	Out = FVector4f(In[0], In[1], In[2], In[3]);
+	In += 4;
 }
 
 void FMeshPropertyMapEvaluator::EvaluateChannel(const int DataIdx, float*& In, float& Out, void* EvalData)
@@ -107,11 +121,11 @@ void FMeshPropertyMapEvaluator::EvaluateChannel(const int DataIdx, float*& In, f
 }
 
 template <bool bUseDetailNormalMap>
-FVector3f FMeshPropertyMapEvaluator::SampleFunction(const FCorrespondenceSample& SampleData) const
+FVector4f FMeshPropertyMapEvaluator::SampleFunction(const FCorrespondenceSample& SampleData) const
 {
 	const void* DetailMesh = SampleData.DetailMesh;
 	const FVector3d& DetailBaryCoords = SampleData.DetailBaryCoords;
-	FVector3f Color = DefaultValue;
+	FVector4f Color = DefaultValue4f;
 	const int32 DetailTriID = SampleData.DetailTriID;
 
 	switch (this->Property)
@@ -125,7 +139,7 @@ FVector3f FMeshPropertyMapEvaluator::SampleFunction(const FCorrespondenceSample&
 	case EMeshPropertyMapType::FacetNormal:
 	{
 		const FVector3d FacetNormal = DetailSampler->GetTriNormal(DetailMesh, DetailTriID);
-		Color = NormalToColor(FacetNormal);
+		Color = NormalToColor(FVector3f(FacetNormal));
 	}
 	break;
 	case EMeshPropertyMapType::Normal:
@@ -175,7 +189,7 @@ FVector3f FMeshPropertyMapEvaluator::SampleFunction(const FCorrespondenceSample&
 				}
 			}
 
-			Color = NormalToColor(FVector3d(DetailNormal));
+			Color = NormalToColor(DetailNormal);
 		}
 	}
 	break;
@@ -199,14 +213,14 @@ FVector3f FMeshPropertyMapEvaluator::SampleFunction(const FCorrespondenceSample&
 		FVector4f DetailColor;
 		if (DetailSampler->TriBaryInterpolateColor(DetailMesh, DetailTriID, DetailBaryCoords, DetailColor))
 		{
-			Color = FVector3f(DetailColor.X, DetailColor.Y, DetailColor.Z);
+			Color = DetailColor;
 		}
 	}
 	break;
 	case EMeshPropertyMapType::PolyGroupID:
 	{
 		const int32 GroupID = DetailSampler->GetPolyGroupID(DetailMesh, DetailTriID);
-		Color = LinearColors::SelectColor<FVector3f>(GroupID);
+		Color = LinearColors::SelectColor<FVector4f>(GroupID);
 	}
 	break;
 	}
diff --git a/Engine/Plugins/Runtime/GeometryProcessing/Source/DynamicMesh/Public/Sampling/MeshPropertyMapEvaluator.h b/Engine/Plugins/Runtime/GeometryProcessing/Source/DynamicMesh/Public/Sampling/MeshPropertyMapEvaluator.h
--- a/Engine/Plugins/Runtime/GeometryProcessing/Source/DynamicMesh/Public/Sampling/MeshPropertyMapEvaluator.h
+++ b/Engine/Plugins/Runtime/GeometryProcessing/Source/DynamicMesh/Public/Sampling/MeshPropertyMapEvaluator.h
@@ -56,6 +56,8 @@ public:
 	static void EvaluateChannel(const int DataIdx, float*& In, float& Out, void* EvalData);
 
 protected:
+	FVector4f GetDefaultValue(EMeshPropertyMapType InProperty) const;
+	
 	// Cached data
 	const IMeshBakerDetailSampler* DetailSampler = nullptr;
 
@@ -65,31 +67,35 @@ protected:
 	
 	bool bHasDetailNormalTextures = false;
 	FAxisAlignedBox3d Bounds;
-	FVector3f DefaultValue = FVector3f::Zero();
 
+	UE_DEPRECATED(5.6, "DefaultValue is deprecated. Please use GetDefaultValue() instead.")
+	FVector3f DefaultValue = FVector3f::Zero();
+	
 private:
-	static FVector3f NormalToColor(const FVector3d Normal) 
+	static FVector4f NormalToColor(const FVector3f Normal) 
 	{
-		return (FVector3f)((Normal + FVector3d::One()) * 0.5);
+		return FVector4f((Normal + FVector3f::One()) * 0.5f, 1.0f);
 	}
 
-	static FVector3f UVToColor(const FVector2f UV)
+	static FVector4f UVToColor(const FVector2f UV)
 	{
 		const float X = FMathf::Clamp(UV.X, 0.0, 1.0);
 		const float Y = FMathf::Clamp(UV.Y, 0.0, 1.0);
-		return FVector3f(X, Y, 0);
+		return FVector4f(X, Y, 0.f, 1.f);
 	}
 
-	static FVector3f PositionToColor(const FVector3d Position, const FAxisAlignedBox3d SafeBounds)
+	static FVector4f PositionToColor(const FVector3d Position, const FAxisAlignedBox3d SafeBounds)
 	{
-		double X = (Position.X - SafeBounds.Min.X) / SafeBounds.Width();
-		double Y = (Position.Y - SafeBounds.Min.Y) / SafeBounds.Height();
-		double Z = (Position.Z - SafeBounds.Min.Z) / SafeBounds.Depth();
-		return (FVector3f)FVector3d(X, Y, Z);
+		const float X = static_cast<float>((Position.X - SafeBounds.Min.X) / SafeBounds.Width());
+		const float Y = static_cast<float>((Position.Y - SafeBounds.Min.Y) / SafeBounds.Height());
+		const float Z = static_cast<float>((Position.Z - SafeBounds.Min.Z) / SafeBounds.Depth());
+		return FVector4f(X, Y, Z, 1.f);
 	}
 
 	template <bool bUseDetailNormalMap>
-	FVector3f SampleFunction(const FCorrespondenceSample& SampleData) const;
+	FVector4f SampleFunction(const FCorrespondenceSample& SampleData) const;
+
+	FVector4f DefaultValue4f = FVector4f();
 };
 
 } // end namespace UE::Geometry
