Geometry Script: Add support for transfering vertex colors between two dissimilar meshes

diff --git a/Engine/Plugins/Runtime/GeometryProcessing/Source/DynamicMesh/Private/Operations/TransferDynamicMeshAttributes.cpp b/Engine/Plugins/Runtime/GeometryProcessing/Source/DynamicMesh/Private/Operations/TransferDynamicMeshAttributes.cpp
new file mode 100644
--- /dev/null
+++ b/Engine/Plugins/Runtime/GeometryProcessing/Source/DynamicMesh/Private/Operations/TransferDynamicMeshAttributes.cpp
@@ -0,0 +1,578 @@
+// Copyright Epic Games, Inc. All Rights Reserved.
+
+#include "Operations/TransferDynamicMeshAttributes.h"
+#include "DynamicMesh/DynamicMesh3.h"
+#include "DynamicMesh/DynamicMeshAttributeSet.h"
+#include "DynamicMesh/DynamicMeshAABBTree3.h"
+#include "DynamicMesh/MeshNormals.h"
+#include "Util/ProgressCancel.h"
+#include "Async/ParallelFor.h"
+#include "TransformTypes.h"
+#include "Algo/Count.h"
+#include "Solvers/Internal/QuadraticProgramming.h"
+#include "Solvers/LaplacianMatrixAssembly.h"
+#include "Operations/SmoothDynamicMeshAttributes.h"
+
+using namespace UE::Geometry;
+
+namespace TransferValuesLocals 
+{
+	static constexpr int32 NumElements = 4;
+	
+	/**
+	 * Given a triangle and point on a triangle (via barycentric coordinates), compute the color.
+	 * 
+	 * @param OutColor Interpolated color for a vertex with Bary barycentric coordinates
+	 * @param TriVertices The vertices of a triangle containing the point we are interpolating the colors for
+	 * @param Bary Barycentric coordinates of the point
+	 * @param InColorAttribute Attribute containing colors of the mesh that TriVertices belong to
+	 */
+
+	void InterpolateVertexAttribute(FVector4f& OutColor,
+									const FIndex3i& TriVertices,
+									const FVector3f& Bary,
+									const FDynamicMeshColorOverlay* InColorAttribute)
+	{
+		FVector4f Value1, Value2, Value3;
+		InColorAttribute->GetElement(TriVertices[0], Value1);
+		InColorAttribute->GetElement(TriVertices[1], Value2);
+		InColorAttribute->GetElement(TriVertices[2], Value3);
+
+		const float Alpha = Bary[0], Beta = Bary[1], Theta  = Bary[2];
+		for (int32 Idx = 0; Idx < NumElements; ++Idx)
+		{
+			OutColor[Idx] = Alpha * Value1[Idx] + Beta * Value2[Idx] + Theta * Value3[Idx];
+		}
+	}
+
+	FDynamicMeshColorOverlay* GetOrCreateColorAttribute(FDynamicMesh3& InMesh)
+	{
+		checkSlow(InMesh.HasAttributes());
+
+		FDynamicMeshAttributeSet* MeshAttributes = InMesh.Attributes();
+		if (!MeshAttributes->HasPrimaryColors())
+		{
+			MeshAttributes->EnablePrimaryColors();
+		}
+		
+		return MeshAttributes->PrimaryColors();
+	}
+
+	static FVector3f ToUENormal(const FVector3d& Normal)
+	{
+		return FVector3f((float)Normal.X, (float)Normal.Y, (float)Normal.Z);
+	}
+}
+
+FTransferVertexColorAttribute::FTransferVertexColorAttribute(
+	const FDynamicMesh3* InSourceMesh,
+	const FDynamicMeshAABBTree3* InSourceBVH)
+	: SourceMesh(InSourceMesh)
+	, SourceBVH(InSourceBVH)
+{
+	// If the BVH for the source mesh was not specified then create one
+	if (SourceBVH == nullptr)
+	{
+		InternalSourceBVH = MakeUnique<FDynamicMeshAABBTree3>(SourceMesh);
+	}
+}
+
+FTransferVertexColorAttribute::~FTransferVertexColorAttribute() 
+{}
+
+bool FTransferVertexColorAttribute::Cancelled()
+{
+	return (Progress == nullptr) ? false : Progress->Cancelled();
+}
+
+EOperationValidationResult FTransferVertexColorAttribute::Validate()
+{	
+	if (SourceMesh == nullptr) 
+	{
+		return EOperationValidationResult::Failed_UnknownReason;
+	}
+
+	// Either BVH was passed by the caller or was created internally in the constructor
+	if (SourceBVH == nullptr && InternalSourceBVH.IsValid() == false) 
+	{
+		return EOperationValidationResult::Failed_UnknownReason;
+	}
+
+	if (!SourceMesh->HasAttributes()) 
+	{
+		return EOperationValidationResult::Failed_UnknownReason;
+	}
+
+	if (!SourceMesh->Attributes()->HasPrimaryColors()) 
+	{
+		return EOperationValidationResult::Failed_UnknownReason;
+	}
+
+	return EOperationValidationResult::Ok;
+}
+
+bool FTransferVertexColorAttribute::TransferColorsToMesh(FDynamicMesh3& InOutTargetMesh)
+{	
+	using namespace TransferValuesLocals;
+
+	if (Validate() != EOperationValidationResult::Ok) 
+	{
+		return false;
+	}
+
+	if (!InOutTargetMesh.HasAttributes())
+	{
+		InOutTargetMesh.EnableAttributes(); 
+	}
+
+	// If we need to compare normals, make sure both the target and the source meshes have per-vertex normals data
+	TUniquePtr<FMeshNormals> InternalTargetMeshNormals;
+	if (NormalThreshold >= 0)
+	{
+		if (!SourceMesh->HasVertexNormals() && !InternalSourceMeshNormals)
+		{
+			// only do this once for the source mesh in case of subsequent calls to the method
+			InternalSourceMeshNormals = MakeUnique<FMeshNormals>(SourceMesh);
+			InternalSourceMeshNormals->ComputeVertexNormals();
+		}
+
+		if (!InOutTargetMesh.HasVertexNormals())
+		{
+			InternalTargetMeshNormals = MakeUnique<FMeshNormals>(&InOutTargetMesh);
+			InternalTargetMeshNormals->ComputeVertexNormals();
+		}
+	}
+	
+	FDynamicMeshColorOverlay* TargetColors = GetOrCreateColorAttribute(InOutTargetMesh);
+	checkSlow(TargetColors);
+
+	bool bFailed = false;
+
+	// compute the transfer only for the subset of vertices if necessary  
+	const bool bUseSubset = !TargetVerticesSubset.IsEmpty();
+	const int32 NumVerticesToTransfer = bUseSubset ? TargetVerticesSubset.Num() : InOutTargetMesh.MaxVertexID();
+	
+	if (TransferMethod == ETransferMethod::ClosestPointOnSurface)
+	{
+		MatchedVertices.Init(false, NumVerticesToTransfer);
+		
+		ParallelFor(NumVerticesToTransfer, [this, &InOutTargetMesh, &TargetColors, &InternalTargetMeshNormals, bUseSubset](int32 InVertexID)
+		{
+			if (Cancelled()) 
+			{
+				return;
+			}
+
+			const int32 VertexID = bUseSubset ? TargetVerticesSubset[InVertexID] : InVertexID;
+			if (InOutTargetMesh.IsVertex(VertexID)) 
+			{
+				const FVector3d Point = InOutTargetMesh.GetVertex(VertexID);
+
+				FVector3f Normal = FVector3f::UnitY();
+				if (NormalThreshold >= 0) 
+				{
+					const bool bHasNormals = InOutTargetMesh.HasVertexNormals();
+					if (ensure(bHasNormals || InternalTargetMeshNormals))
+					{
+						Normal = bHasNormals ? InOutTargetMesh.GetVertexNormal(VertexID) : ToUENormal(InternalTargetMeshNormals->GetNormals()[VertexID]);
+					}
+				}
+
+				FVector4f Color;
+				if (TransferColorToPoint(Color, Point, Normal))
+				{
+					TargetColors->SetElement(VertexID, Color);
+					MatchedVertices[VertexID] = true;
+				}
+			}
+
+		}, bUseParallel ? EParallelForFlags::None : EParallelForFlags::ForceSingleThread);
+
+		// If the caller requested to simply find the closest point for all vertices then the number of matched vertices
+		// must be equal to the target mesh vertex count
+		if (SearchRadius < 0 && NormalThreshold < 0)
+		{
+			int32 NumMatched = 0;
+			for (bool Flag : MatchedVertices)
+			{
+				if (Flag) 
+				{ 
+					NumMatched++;
+				}
+			}
+			bFailed = NumMatched != NumVerticesToTransfer;
+		}
+	} 
+	else if (TransferMethod == ETransferMethod::Inpaint)
+	{
+		/**
+         *  Given two meshes, Mesh1 without colors and Mesh2 with colors, assume they are aligned in 3d space.
+         *  For every vertex on Mesh1 find the closest point on the surface of Mesh2 within a radius R. If the difference 
+         *  between the normals of the two points is below the threshold, then it's a match. Otherwise no match.
+         *  So now we have two sets of vertices on Mesh1. One with a match on the source mesh and one without a match.
+         *  For all the vertices with a match, copy values over. For all the vertices without the match, do nothing.
+         *  Now, for all the vertices without a match, try to approximate the values by smoothly interpolating between 
+         *  the values at the known vertices via solving a quadratic problem. 
+         *  
+         *  The solver minimizes an energy
+         *      trace(W^t Q W)
+         *      W \in R^(nxm) is a matrix where n is the number of vertices and m is the number of elements.
+         *      Q \in R^(nxn) is a matrix that combines both Dirichlet and Laplacian energies, Q = -L + L*M^(-1)*L
+         *                    where L is a cotangent Laplacian and M is a mass matrix
+         *  
+		 */
+	
+		MatchedVertices.Init(false, InOutTargetMesh.MaxVertexID());
+ 
+		// because the inpaint algorithm can extract data from regions outside the target vertex subset, a temporary attribute is used to modify the values.
+		// NOTE: make sure to copy the values of the vertex subset into the complete TargetColors attribute before exciting the function. (see CopySubsetColorsIfNeeded)
+		FDynamicMeshColorOverlay SubsetTargetColors;
+		if (bUseSubset)
+		{
+			SubsetTargetColors.Copy(*TargetColors);
+		}
+		FDynamicMeshColorOverlay* EditedColors = bUseSubset ? &SubsetTargetColors : TargetColors;
+ 
+		// For every vertex on the target mesh try to find the match on the source mesh using the distance and normal checks
+		ParallelFor(InOutTargetMesh.MaxVertexID(), [this, &InOutTargetMesh, &EditedColors, &InternalTargetMeshNormals](int32 VertexID)
+		{
+			if (Cancelled()) 
+			{
+				return;
+			}
+			
+			if (InOutTargetMesh.IsVertex(VertexID)) 
+			{
+				// check if we need to force the vertex to not have a match
+				if (ForceInpaint.Num() == InOutTargetMesh.MaxVertexID() && ForceInpaint[VertexID] != 0)
+				{
+					return;
+				}
+ 
+				const FVector3d Point = InOutTargetMesh.GetVertex(VertexID);
+				FVector3f Normal = FVector3f::UnitY();
+				if (NormalThreshold >= 0) 
+				{
+					const bool bHasNormals = InOutTargetMesh.HasVertexNormals();
+					if (ensure(bHasNormals || InternalTargetMeshNormals.IsValid()))
+					{
+						Normal = bHasNormals ? InOutTargetMesh.GetVertexNormal(VertexID) : ToUENormal(InternalTargetMeshNormals->GetNormals()[VertexID]);
+					}
+				}
+ 
+				FVector4f Color;
+				if (TransferColorToPoint(Color, Point, Normal))
+				{
+					EditedColors->SetElement(VertexID, Color);
+					MatchedVertices[VertexID] = true;
+				}
+			}
+		}, bUseParallel ? EParallelForFlags::None : EParallelForFlags::ForceSingleThread);
+ 
+		if (Cancelled()) 
+		{
+			return false;
+		}
+ 
+		int32 NumMatched = 0;
+		if (!bUseSubset)
+		{
+			for (bool Flag : MatchedVertices)
+			{
+				if (Flag) 
+				{ 
+					NumMatched++;
+				}
+			}
+		}
+		else
+		{
+			NumMatched = (int32)Algo::CountIf(TargetVerticesSubset, [this](int32 VertexID)
+			{
+				return MatchedVertices.IsValidIndex(VertexID) && MatchedVertices[VertexID];
+			});
+		}
+ 
+		// If no vertices matched, we have nothing to inpaint.
+		if (NumMatched == 0)
+		{
+			return false;
+		}
+ 
+		auto CopySubsetColorsIfNeeded = [Subset = TargetVerticesSubset, &EditedColors, &TargetColors, &InOutTargetMesh]()
+		{
+			if (EditedColors && EditedColors != TargetColors)
+			{
+				for (const int32 VertexID: Subset)
+				{
+					if (InOutTargetMesh.IsVertex(VertexID))
+					{
+						FVector4f Color;
+						EditedColors->GetElement(VertexID, Color);
+						TargetColors->SetElement(VertexID, Color);
+					}
+				}
+			}
+		};
+		
+		// If all vertices were matched then nothing else to do
+		if (NumMatched == NumVerticesToTransfer)
+		{
+			// copy colors from the subset attribute if using subset
+			CopySubsetColorsIfNeeded();
+			return true;
+		}
+ 
+		// Compute linearization so we can store constraints at linearized indices
+		FVertexLinearization VtxLinearization(InOutTargetMesh, false);
+		const TArray<int32>& ToMeshV = VtxLinearization.ToId();
+		const TArray<int32>& ToIndex = VtxLinearization.ToIndex();
+		
+		// Setup the sparse matrix FixedValues of known (matched) colors and the array (FixedIndices) of the matched vertex IDs
+		FSparseMatrixD FixedValues;
+		FixedValues.resize(NumMatched, NumElements);
+		std::vector<Eigen::Triplet<FSparseMatrixD::Scalar>> FixedValuesTriplets;
+		FixedValuesTriplets.reserve(NumMatched);
+		
+		TArray<int> FixedIndices;
+		FixedIndices.Reserve(NumMatched);
+ 
+		for (int32 VertexID = 0; VertexID < InOutTargetMesh.MaxVertexID(); ++VertexID)
+		{
+			if (InOutTargetMesh.IsVertex(VertexID) && MatchedVertices[VertexID])
+			{
+				FVector4f Data;
+				EditedColors->GetElement(VertexID, Data);
+ 
+				const int32 CurIdx = FixedIndices.Num();
+				for (int32 Idx = 0; Idx < NumElements; ++Idx)
+				{
+					FixedValuesTriplets.emplace_back(CurIdx, Idx, Data[Idx]);
+				}
+ 
+				checkSlow(VertexID < ToIndex.Num());
+				FixedIndices.Add(ToIndex[VertexID]);
+			}
+		}
+		FixedValues.setFromTriplets(FixedValuesTriplets.begin(), FixedValuesTriplets.end());
+ 
+		const int32 NumVerts = VtxLinearization.NumVerts();
+		FEigenSparseMatrixAssembler CotangentAssembler(NumVerts, NumVerts);
+		FEigenSparseMatrixAssembler LaplacianAssembler(NumVerts, NumVerts);
+ 
+		if (bUseIntrinsicLaplacian)
+		{
+			// Construct the Cotangent values matrix
+			UE::MeshDeformation::ConstructFullIDTCotangentLaplacian<double>(InOutTargetMesh, VtxLinearization, CotangentAssembler,
+																			UE::MeshDeformation::ECotangentWeightMode::Default, 
+																			UE::MeshDeformation::ECotangentAreaMode::NoArea);
+ 
+			// Construct the Laplacian with cotangent values scaled by the voronoi area (i.e. M^(-1)*L matrix where M is the mass/stiffness matrix)
+			UE::MeshDeformation::ConstructFullIDTCotangentLaplacian<double>(InOutTargetMesh, VtxLinearization, LaplacianAssembler,
+																			UE::MeshDeformation::ECotangentWeightMode::Default,
+																			UE::MeshDeformation::ECotangentAreaMode::VoronoiArea);
+		}
+		else 
+		{
+			UE::MeshDeformation::ConstructFullCotangentLaplacian<double>(InOutTargetMesh, VtxLinearization, CotangentAssembler,
+																		 UE::MeshDeformation::ECotangentWeightMode::Default,
+																		 UE::MeshDeformation::ECotangentAreaMode::NoArea);
+ 
+			UE::MeshDeformation::ConstructFullCotangentLaplacian<double>(InOutTargetMesh, VtxLinearization, LaplacianAssembler,
+																		 UE::MeshDeformation::ECotangentWeightMode::Default,
+																		 UE::MeshDeformation::ECotangentAreaMode::VoronoiArea);
+		}
+ 
+		FSparseMatrixD CotangentMatrix, MassCotangentMatrix;
+		CotangentAssembler.ExtractResult(CotangentMatrix);
+		LaplacianAssembler.ExtractResult(MassCotangentMatrix);
+ 
+		// -L * L* M^(-1)*L energy
+		FSparseMatrixD Energy = -1*CotangentMatrix + CotangentMatrix*MassCotangentMatrix;
+ 
+		// Solve the QP problem with fixed constraints
+		FSparseMatrixD TargetValues;
+		TArray<int> VariableRows;
+ 
+		// We want the solution TargetValues matrix to only contain the rows representing the variable (non-fixed) rows
+		constexpr bool bVariablesOnly = true;
+ 		bFailed = !FQuadraticProgramming::SolveWithFixedConstraints(Energy, nullptr, FixedIndices, FixedValues, TargetValues, bVariablesOnly, KINDA_SMALL_NUMBER, &VariableRows);
+		checkSlow((VariableRows.Num() + FixedIndices.Num()) == Energy.rows());
+ 
+		if (!bFailed)
+		{
+			// Transpose so we can efficiently iterate over the col-major matrix. Each column now contains per-vertex values.
+			// Otherwise, we are iterating over rows of a col-major matrix which is slow.
+			FSparseMatrixD TargetValuesTransposed = TargetValues.transpose(); 
+ 
+			// Iterate over every column containing all values for the vertex
+			for (int32 ColIdx = 0; ColIdx < TargetValuesTransposed.outerSize(); ++ColIdx)
+			{
+				FVector4f Data(0.f);
+				FInt32Vector4 N(0);
+				
+				// Iterate over only non-zero rows (i.e. non-zero values)
+				for (FSparseMatrixD::InnerIterator Itr(TargetValuesTransposed, ColIdx); Itr; ++Itr)
+				{
+					const int32 Idx = static_cast<int32>(Itr.row());
+					const float Value = static_cast<float>(Itr.value());
+					Data[Idx] = Value;
+					N[Idx] += 1;
+				}
+
+				// normalize
+				for (int32 Index = 0; Index < NumElements; ++Index)
+				{
+					if (N[Index] > 1)
+					{
+						Data /= static_cast<float>(N[Index]); 
+					}
+				}
+				
+				const int32 VertexIDLinearalized = bVariablesOnly ? static_cast<int32>(VariableRows[ColIdx]) : ColIdx; // linearized vertex ID (matrix row) of the variable in the Energy matrix
+				const int32 VertexID = ToMeshV[VertexIDLinearalized];
+				EditedColors->SetElement(VertexID, Data);
+			}
+ 
+			// copy values from the subset attribute if using subset
+			CopySubsetColorsIfNeeded();
+			
+			// Optional post-processing smoothing of the values at the vertices without a match
+			if (NumSmoothingIterations > 0 && SmoothingStrength > 0)
+			{
+				TArray<int32> VerticesToSmooth;
+				const int32 NumNotMatched = InOutTargetMesh.VertexCount() - NumMatched;
+				VerticesToSmooth.Reserve(NumNotMatched);
+				for (int32 VertexID = 0; VertexID < InOutTargetMesh.MaxVertexID(); ++VertexID)
+				{
+					if (InOutTargetMesh.IsVertex(VertexID) && !MatchedVertices[VertexID])
+					{
+						VerticesToSmooth.Add(VertexID);
+					}
+				}
+
+				FSmoothDynamicMeshAttributes BlurOp(InOutTargetMesh);
+				BlurOp.NumIterations = NumSmoothingIterations;
+				BlurOp.Strength = SmoothingStrength;
+				BlurOp.EdgeWeightMethod = FSmoothDynamicMeshAttributes::EEdgeWeights::CotanWeights; //expose as param
+				BlurOp.Selection = MoveTemp(VerticesToSmooth);
+
+				TArray<bool> ValuesToSmooth;
+				ValuesToSmooth.Init(true, NumElements);
+				
+				ensure( BlurOp.SmoothOverlay(TargetColors, ValuesToSmooth) );
+			}
+		}
+	}
+	else 
+	{
+		checkNoEntry(); // unsupported method
+	}
+
+	if (Cancelled() || bFailed) 
+	{
+		return false;
+	}
+		
+	return true;
+}
+
+bool FTransferVertexColorAttribute::TransferColorToPoint(FVector4f& OutColor, const FVector3d& InPoint, const FVector3f& InNormal) const 
+{	
+	using namespace TransferValuesLocals;
+
+	// Find the containing triangle and the barycentric coordinates of the closest point
+	int32 TriID; 
+	FVector3d Bary;
+	if (!FindClosestPointOnSourceSurface(InPoint, TargetToWorld, TriID, Bary))
+	{
+		return false;
+	}
+
+	const FVector3f BaryF((float)Bary[0], (float)Bary[1], (float)Bary[2]);
+	const FIndex3i TriVertex = SourceMesh->GetTriangle(TriID);
+
+	const FDynamicMeshColorOverlay* SourceColors = SourceMesh->Attributes()->PrimaryColors();
+	if (SearchRadius < 0 && NormalThreshold < 0)
+	{
+		// If the radius and normals are ignored, simply interpolate the values and return the result
+		InterpolateVertexAttribute(OutColor, TriVertex, BaryF, SourceColors);
+	}
+	else
+	{
+		bool bPassedRadiusCheck = true;
+		if (SearchRadius >= 0)
+		{
+			const FVector3d MatchedPoint = SourceMesh->GetTriBaryPoint(TriID, Bary[0], Bary[1], Bary[2]);
+			bPassedRadiusCheck = (InPoint - MatchedPoint).Length() <= SearchRadius;
+		}
+
+		bool bPassedNormalsCheck = true;
+		if (NormalThreshold >= 0)
+		{
+			FVector3f Normal0 = FVector3f::UnitY();
+			FVector3f Normal1 = FVector3f::UnitY();
+			FVector3f Normal2 = FVector3f::UnitY();
+			const bool bHasSourceNormals = SourceMesh->HasVertexNormals();
+			if (ensure(bHasSourceNormals || InternalSourceMeshNormals.IsValid()))
+			{
+				Normal0 = bHasSourceNormals ? SourceMesh->GetVertexNormal(TriVertex[0]) : ToUENormal(InternalSourceMeshNormals->GetNormals()[TriVertex[0]]);
+				Normal1 = bHasSourceNormals ? SourceMesh->GetVertexNormal(TriVertex[1]) : ToUENormal(InternalSourceMeshNormals->GetNormals()[TriVertex[1]]);
+				Normal2 = bHasSourceNormals ? SourceMesh->GetVertexNormal(TriVertex[2]) : ToUENormal(InternalSourceMeshNormals->GetNormals()[TriVertex[2]]);
+			}
+
+			const FVector3f MatchedNormal = Normalized(BaryF[0]*Normal0 + BaryF[1]*Normal1 + BaryF[2]*Normal2);
+			const FVector3f InNormalNormalized = Normalized(InNormal);
+			const float NormalAngle = FMathf::ACos(InNormalNormalized.Dot(MatchedNormal));
+			bPassedNormalsCheck = (double)NormalAngle <= NormalThreshold;
+
+			if (!bPassedNormalsCheck && LayeredMeshSupport)
+			{
+				// try again with a flipped normal
+				bPassedNormalsCheck = (double)(TMathUtil<float>::Pi - NormalAngle) <= NormalThreshold;
+			}
+		}
+		
+		if (bPassedRadiusCheck && bPassedNormalsCheck)
+		{
+			InterpolateVertexAttribute(OutColor, TriVertex, BaryF, SourceColors);
+		}
+		else
+		{
+			return false;
+		}
+	}
+
+	return true;
+}
+
+bool FTransferVertexColorAttribute::FindClosestPointOnSourceSurface(const FVector3d& InPoint, const FTransformSRT3d& InToWorld, int32& NearTriID, FVector3d& Bary) const
+{
+	IMeshSpatial::FQueryOptions Options;
+	double NearestDistSqr;
+	
+	const FVector3d WorldPoint = InToWorld.TransformPosition(InPoint);
+	if (SourceBVH != nullptr) 
+	{ 
+		NearTriID = SourceBVH->FindNearestTriangle(WorldPoint, NearestDistSqr, Options);
+	}
+	else 
+	{
+		NearTriID = InternalSourceBVH->FindNearestTriangle(WorldPoint, NearestDistSqr, Options);
+	}
+
+	if (!ensure(NearTriID != IndexConstants::InvalidID))
+	{
+		return false;
+	}
+
+	const FDistPoint3Triangle3d Query = TMeshQueries<FDynamicMesh3>::TriangleDistance(*SourceMesh, NearTriID, WorldPoint);
+	const FVector3d NearestPnt = Query.ClosestTrianglePoint;
+	const FIndex3i TriVertex = SourceMesh->GetTriangle(NearTriID);
+
+	Bary = VectorUtil::BarycentricCoords(NearestPnt, SourceMesh->GetVertexRef(TriVertex.A),
+													 SourceMesh->GetVertexRef(TriVertex.B),
+													 SourceMesh->GetVertexRef(TriVertex.C));
+
+	return true;
+}
diff --git a/Engine/Plugins/Runtime/GeometryProcessing/Source/DynamicMesh/Public/Operations/TransferDynamicMeshAttributes.h b/Engine/Plugins/Runtime/GeometryProcessing/Source/DynamicMesh/Public/Operations/TransferDynamicMeshAttributes.h
new file mode 100644
--- /dev/null
+++ b/Engine/Plugins/Runtime/GeometryProcessing/Source/DynamicMesh/Public/Operations/TransferDynamicMeshAttributes.h
@@ -0,0 +1,184 @@
+// Copyright Epic Games, Inc. All Rights Reserved.
+
+#pragma once
+
+#include "GeometryTypes.h"
+#include "UObject/NameTypes.h"
+#include "TransformTypes.h"
+
+// Forward declarations
+class FProgressCancel;
+
+namespace UE::Geometry
+{
+	class FDynamicMesh3;
+	template<typename MeshType> class TMeshAABBTree3;
+	typedef TMeshAABBTree3<FDynamicMesh3> FDynamicMeshAABBTree3;
+	class FMeshNormals;
+}
+
+
+namespace UE
+{
+namespace Geometry
+{
+
+/**
+ * Transfer vertex colors from one mesh (source) to another (target).
+ */
+
+class DYNAMICMESH_API FTransferVertexColorAttribute
+{
+public:
+
+	enum class ETransferMethod : uint8
+	{
+        // For every vertex on the target mesh, find the closest point on the surface of the source mesh. If that point 
+        // is within the SearchRadius, and their normals differ by less than the NormalThreshold, then we directly copy  
+        // the weights from the source point to the target mesh vertex.
+		ClosestPointOnSurface = 0,
+
+        // Same as the ClosestPointOnSurface but for all the vertices we didn't copy the weights directly, automatically 
+		// compute the smooth weights.
+		Inpaint = 1
+	};
+
+	//
+	// Optional Inputs
+	//
+	
+    /** Set this to be able to cancel the running operation. */
+	FProgressCancel* Progress = nullptr;
+
+	/** Enable/disable multi-threading. */
+	bool bUseParallel = true;
+	
+	/** The transfer method to compute the bone weights. */
+	ETransferMethod TransferMethod = ETransferMethod::ClosestPointOnSurface;
+
+	/** Transform applied to the input target mesh or target point before transfer. */
+	FTransformSRT3d TargetToWorld = FTransformSRT3d::Identity();
+
+	//
+	// Optional Inputs for ETransferMethod::Inpaint method
+	//
+
+	/**  Radius for searching the closest point. If negative, all points are considered. */
+	double SearchRadius = -1;
+
+	/** 
+	 * Maximum angle (in radians) difference between target and source point normals to be considered a match. 
+	 * If negative, normals are ignored.
+	 */
+	double NormalThreshold = -1;
+	
+	/** 
+	 * If true, when the closest point doesn't pass the normal threshold test, will try again with a flipped normal. 
+	 * This helps with layered meshes where the "inner" and "outer" layers are close to each other but whose normals 
+	 * are pointing in the opposite directions. 
+	 */
+	bool LayeredMeshSupport = false;
+
+	/** The number of optional post-processing smoothing iterations applied to the vertices without the match. */
+	int32 NumSmoothingIterations = 0; 
+
+	/** The strength of each post-processing smoothing iteration. */
+	float SmoothingStrength = 0.0f;
+
+	/** If true, will use the intrinsic Delaunay mesh to construct sparse Cotangent Laplacian matrix. */
+	bool bUseIntrinsicLaplacian = false;
+
+	/** 
+	 * Optional mask where if ForceInpaint[VertexID] != 0 we want to force the colors for the vertex to be computed  
+	 * automatically.
+	 * 
+	 * @note Only used when TransferMethod == ETransferMethod::Inpaint.
+	 * 		 The size must be equal to the InTargetMesh.MaxVertexID(), otherwise the mask is ignored.
+	 */
+	TArray<float> ForceInpaint;
+
+	/** 
+	 * Optional subset of target mesh vertices to transfer weights to.
+	 * If left empty, skin weights will be transferred to all target mesh vertices.
+	 */
+	TArray<int32> TargetVerticesSubset;
+
+	//
+	// Outputs
+	//
+
+	/** MatchedVertices[VertexID] is set to true for a target mesh vertex ID with a match found, false otherwise. */
+	TArray<bool> MatchedVertices;
+
+protected:
+		
+	/** Source mesh we are transferring colors from. */
+	const FDynamicMesh3* SourceMesh = nullptr;
+	
+	/** 
+	 * The caller can optionally specify the source mesh BVH in case this operator is run on multiple target meshes 
+	 * while the source mesh remains the same. Otherwise BVH tree will be computed.
+	 */
+	const FDynamicMeshAABBTree3* SourceBVH = nullptr;
+
+	/** If the caller doesn't pass BVH for the source mesh then we compute one. */
+	TUniquePtr<FDynamicMeshAABBTree3> InternalSourceBVH;
+
+	/** If the source mesh doesn't have per-vertex normals then compute them */
+	TUniquePtr<FMeshNormals> InternalSourceMeshNormals;
+
+public:
+	
+	/**
+	 * @param InSourceMesh The mesh we are transferring colors from 
+	 * @param SourceBVH Optional source mesh BVH. If not provided, one will be computed internally. 
+	 * 
+	 * @note Assumes that the InSourceMesh has a primary colors attribute.
+	 */
+	FTransferVertexColorAttribute(const FDynamicMesh3* InSourceMesh, const FDynamicMeshAABBTree3* SourceBVH = nullptr); 
+	
+	virtual ~FTransferVertexColorAttribute();
+
+	/**
+	 * @return EOperationValidationResult::Ok if we can apply operation, or error code if we cannot.
+	 */
+	virtual EOperationValidationResult Validate();
+
+	/**
+     * Transfer the colors from the source mesh to the given target mesh and store the result in the primary colors attribute.
+	 * 
+	 * @param InOutTargetMesh	  Target mesh we are transferring colors into
+     * 
+     * @return true if the algorithm succeeds, false if it failed or was canceled by the user.
+	 */
+	virtual bool TransferColorsToMesh(FDynamicMesh3& InOutTargetMesh);
+
+
+	/**
+	 * Compute the color for a given point using the ETransferMethod::ClosestPointOnSurface algorithm.
+	 *
+	 * @param OutColor 		Color computed for the input transformed point.
+	 * @param InPoint		Point for which we are computing the color.
+	 * @param InNormal		Normal at the input point. Should be set if NormalThreshold >= 0.
+	 * 
+	 * @return true if the algorithm succeeds, false if it failed to find the matching point or was canceled by the user.
+	 */
+	bool TransferColorToPoint(FVector4f& OutColor, const FVector3d& InPoint, const FVector3f& InNormal = FVector3f::Zero()) const;
+	
+protected:
+	
+    /** @return if true, abort the computation. */
+	virtual bool Cancelled();
+	
+	/** 
+	 * Find the closest point on the surface of the source mesh and return the ID of the triangle containing it and its 
+	 * barycentric coordinates.
+	 * 
+	 * @return true if point is found, false otherwise
+	 */
+	bool FindClosestPointOnSourceSurface(const FVector3d& InPoint, const FTransformSRT3d& InToWorld, int32& OutTriID, FVector3d& OutBary) const;
+
+};
+
+} // end namespace UE::Geometry
+} // end namespace UE
diff --git a/Engine/Plugins/Runtime/GeometryScripting/Source/GeometryScriptingCore/Private/MeshVertexColorFunctions.cpp b/Engine/Plugins/Runtime/GeometryScripting/Source/GeometryScriptingCore/Private/MeshVertexColorFunctions.cpp
--- a/Engine/Plugins/Runtime/GeometryScripting/Source/GeometryScriptingCore/Private/MeshVertexColorFunctions.cpp
+++ b/Engine/Plugins/Runtime/GeometryScripting/Source/GeometryScriptingCore/Private/MeshVertexColorFunctions.cpp
@@ -7,6 +7,7 @@
 #include "Util/ColorConstants.h"
 #include "Operations/SmoothDynamicMeshAttributes.h"
 #include "DynamicMesh/DynamicMeshOverlay.h"
+#include "Operations/TransferDynamicMeshAttributes.h"
 
 #include UE_INLINE_GENERATED_CPP_BY_NAME(MeshVertexColorFunctions)
 
@@ -478,7 +479,75 @@ UDynamicMesh* UGeometryScriptLibrary_MeshVertexColorFunctions::BlurMeshVertexCol
 	return TargetMesh;
 }
 
+UDynamicMesh* UGeometryScriptLibrary_MeshVertexColorFunctions::TransferVertexColorsFromMesh(
+	UDynamicMesh* SourceMesh,
+	UDynamicMesh* TargetMesh,
+	FGeometryScriptTransferMeshVertexColorsOptions Options,
+	FGeometryScriptMeshSelection Selection,
+	UGeometryScriptDebug* Debug)
+{
+	using namespace UE::Geometry;
+
+	if (SourceMesh == nullptr)
+	{
+		AppendError(Debug, EGeometryScriptErrorType::InvalidInputs, LOCTEXT("TransferVertexColorsFromMesh_InvalidSourceMesh", "TransferVertexColorsFromMesh: Source Mesh is Null"));
+		return TargetMesh;
+	}
+	if (TargetMesh == nullptr)
+	{
+		AppendError(Debug, EGeometryScriptErrorType::InvalidInputs, LOCTEXT("TransferVertexColorsFromMesh_InvalidTargetMesh", "TransferVertexColorsFromMesh: Target Mesh is Null"));
+		return TargetMesh;
+	}
+
+	SourceMesh->ProcessMesh([&](const FDynamicMesh3& ReadMesh)
+	{
+		if (!ReadMesh.HasAttributes() || !ReadMesh.Attributes()->HasPrimaryColors())
+		{
+			AppendError(Debug, EGeometryScriptErrorType::InvalidInputs, LOCTEXT("TransferVertexColorsFromMesh_NoColor", "Source Mesh has no vertex color attribute"));
+			return;
+		}
+
+		FTransferVertexColorAttribute TransferColors(&ReadMesh);
+		TransferColors.TransferMethod = static_cast<FTransferVertexColorAttribute::ETransferMethod>(Options.TransferMethod);
+		TransferColors.bUseParallel = true;
 
+		TargetMesh->EditMesh([&](FDynamicMesh3& EditMesh)
+		{
+			if (!Selection.IsEmpty())
+			{
+				Selection.ConvertToMeshIndexArray(EditMesh, TransferColors.TargetVerticesSubset, EGeometryScriptIndexType::Vertex);
+			}
+			
+			if (!EditMesh.HasAttributes())
+			{
+				EditMesh.EnableAttributes();
+			}
+			
+			if (Options.TransferMethod == ETransferVertexColorMethod::Inpaint)
+			{
+				TransferColors.NormalThreshold = FMathd::DegToRad * Options.NormalThreshold;
+				TransferColors.SearchRadius = Options.RadiusPercentage * EditMesh.GetBounds().DiagonalLength();
+				TransferColors.NumSmoothingIterations = Options.NumSmoothingIterations;
+				TransferColors.SmoothingStrength = Options.SmoothingStrength;
+				TransferColors.LayeredMeshSupport = Options.LayeredMeshSupport;
+			}
+
+			if (TransferColors.Validate() != EOperationValidationResult::Ok)
+			{
+				AppendError(Debug, EGeometryScriptErrorType::OperationFailed, LOCTEXT("TransferVertexColorsFromMesh_ValidationFailed", "TransferVertexColorsFromMesh: Invalid parameters were set for the transfer colors operator"));
+				return;
+			}
+			
+			if (!TransferColors.TransferColorsToMesh(EditMesh))
+			{
+				AppendError(Debug, EGeometryScriptErrorType::OperationFailed, LOCTEXT("TransferVertexColorsFromMesh_TransferFailed", "TransferVertexColorsFromMesh: Failed to transfer the colors"));
+			}
+			
+		}, EDynamicMeshChangeType::AttributeEdit, EDynamicMeshAttributeChangeFlags::Unknown, false);
+	});
+
+	return TargetMesh;
+}
 
 
 
diff --git a/Engine/Plugins/Runtime/GeometryScripting/Source/GeometryScriptingCore/Public/GeometryScript/MeshVertexColorFunctions.h b/Engine/Plugins/Runtime/GeometryScripting/Source/GeometryScriptingCore/Public/GeometryScript/MeshVertexColorFunctions.h
--- a/Engine/Plugins/Runtime/GeometryScripting/Source/GeometryScriptingCore/Public/GeometryScript/MeshVertexColorFunctions.h
+++ b/Engine/Plugins/Runtime/GeometryScripting/Source/GeometryScriptingCore/Public/GeometryScript/MeshVertexColorFunctions.h
@@ -44,6 +44,57 @@ struct GEOMETRYSCRIPTINGCORE_API FGeometryScriptBlurMeshVertexColorsOptions
 	bool Alpha = true;
 };
 
+UENUM(BlueprintType)
+enum class ETransferVertexColorMethod : uint8
+{
+	/** For every vertex on the TargetMesh, find the closest point on the surface of the SourceMesh and transfer 
+	 * colors from it. This is usually a point on the SourceMesh triangle where the colors are computed via 
+	 * interpolation of the colors at the vertices of the triangle via barycentric coordinates.
+	 */
+	ClosestPointOnSurface = 0,
+
+	/** For every vertex on the target mesh, find the closest point on the surface of the source mesh. If that point 
+	 * is within the search radius (controlled via SearchPercentage), and their normals differ by less than the 
+	 * NormalThreshold, then we directly copy the colors from the source point to the target mesh vertex 
+	 * (same as the ClosestPointOnSurface method). For all the vertices we didn't copy the colors directly, 
+	 * automatically compute the smooth colors. 
+	 */
+	Inpaint = 1
+};
+
+USTRUCT(BlueprintType)
+struct GEOMETRYSCRIPTINGCORE_API FGeometryScriptTransferMeshVertexColorsOptions
+{
+	GENERATED_BODY()
+	
+	/** The type of algorithm to use for transferring the colors. */
+	UPROPERTY(EditAnywhere, BlueprintReadWrite, Category = Options)
+	ETransferVertexColorMethod TransferMethod = ETransferVertexColorMethod::ClosestPointOnSurface;
+	
+	/** Defines the search radius as the RadiusPercentage * (input mesh bounding box diagonal). All points not within the search
+	  * radius will be ignored. If negative, all points are considered. Only used in the Inpaint algorithm.*/
+	UPROPERTY(EditAnywhere, BlueprintReadWrite, Category = Options, Meta = (UIMin = -1, UIMax = 2, ClampMin = -1, ClampMax = 2, EditCondition = "TransferMethod == ETransferVertexColorMethod::Inpaint"))
+	double RadiusPercentage = -1;
+
+	/** Maximum angle (in degrees) difference between the target and the source point normals to be considered a match. 
+	 * If negative, normals are ignored. Only used in the Inpaint algorithm.*/
+	UPROPERTY(EditAnywhere, BlueprintReadWrite, Category = Options, Meta = (UIMin = -1, UIMax = 180, ClampMin = -1, ClampMax = 180, EditCondition="TransferMethod == ETransferVertexColorMethod::Inpaint"))
+	double NormalThreshold = -1;
+
+	/** If true, when the closest point doesn't pass the normal threshold test, will try again with a flipped normal. 
+	 * This helps with layered meshes where the "inner" and "outer" layers are close to each other but whose normals 
+	 * are pointing in the opposite directions. */
+	UPROPERTY(EditAnywhere, BlueprintReadWrite, Category = Options, Meta = (EditCondition = "TransferMethod == ETransferVertexColorMethod::Inpaint"))
+	bool LayeredMeshSupport = true;
+
+	/** The number of optional post-processing smoothing iterations applied to the vertices without the match. */
+	UPROPERTY(EditAnywhere, BlueprintReadWrite, Category = Options, meta = (UIMin = 0, UIMax = 100, ClampMin = 0, ClampMax = 100, EditCondition = "TransferMethod == ETransferVertexColorMethod::Inpaint"))
+	int32 NumSmoothingIterations = 0; 
+
+	/** The strength of each post-processing smoothing iteration. */
+	UPROPERTY(EditAnywhere, BlueprintReadWrite, Category = Options, meta = (UIMin = 0, UIMax = 1, ClampMin = 0, ClampMax = 1, EditCondition = "TransferMethod == ETransferVertexColorMethod::Inpaint"))
+	float SmoothingStrength = 0.0f;
+};
 
 
 UCLASS(meta = (ScriptName = "GeometryScript_VertexColors"))
@@ -155,5 +206,24 @@ public:
 		EGeometryScriptBlurColorMode BlurMode = EGeometryScriptBlurColorMode::Uniform,
 		FGeometryScriptBlurMeshVertexColorsOptions Options = FGeometryScriptBlurMeshVertexColorsOptions(),
 		UGeometryScriptDebug* Debug = nullptr);
+
+	/** 
+	 * Transfer the vertex colors from the SourceMesh to the TargetMesh. Assumes that the meshes are aligned. Otherwise, 
+	 * use the TransformMesh geometry script function to align them.
+	 * 
+	 * @param SourceMesh The mesh we are transferring the colors from.
+	 * @param TargetMesh The mesh we are transferring the colors to.
+	 * @param Options The options to set for the transfer colors algorithm.
+	 * @param Selection Optional subset of target mesh vertices to transfer weights to.
+	 * If left empty, colors will be transferred to all target mesh vertices.
+	 */
+	UFUNCTION(BlueprintCallable, Category = "GeometryScript|VertexColor", meta = (ScriptMethod))
+	static UPARAM(DisplayName = "Target Mesh") UDynamicMesh*
+	TransferVertexColorsFromMesh(
+		UDynamicMesh* SourceMesh,
+		UDynamicMesh* TargetMesh,
+		FGeometryScriptTransferMeshVertexColorsOptions Options = FGeometryScriptTransferMeshVertexColorsOptions(),
+		FGeometryScriptMeshSelection Selection = FGeometryScriptMeshSelection(),
+		UGeometryScriptDebug* Debug = nullptr);
 };
 
