Geometry Scripting: Fix index mismatch when transfering colors from one mesh to another.

#rnx

diff --git a/Engine/Plugins/Runtime/GeometryProcessing/Source/DynamicMesh/Private/Operations/TransferDynamicMeshAttributes.cpp b/Engine/Plugins/Runtime/GeometryProcessing/Source/DynamicMesh/Private/Operations/TransferDynamicMeshAttributes.cpp
--- a/Engine/Plugins/Runtime/GeometryProcessing/Source/DynamicMesh/Private/Operations/TransferDynamicMeshAttributes.cpp
+++ b/Engine/Plugins/Runtime/GeometryProcessing/Source/DynamicMesh/Private/Operations/TransferDynamicMeshAttributes.cpp
@@ -23,20 +23,21 @@ namespace TransferValuesLocals
 	 * Given a triangle and point on a triangle (via barycentric coordinates), compute the color.
 	 * 
 	 * @param OutColor Interpolated color for a vertex with Bary barycentric coordinates
-	 * @param TriVertices The vertices of a triangle containing the point we are interpolating the colors for
+	 * @param TriElements The vertices of a triangle containing the point we are interpolating the colors for
 	 * @param Bary Barycentric coordinates of the point
 	 * @param InColorAttribute Attribute containing colors of the mesh that TriVertices belong to
 	 */
 
 	void InterpolateVertexAttribute(FVector4f& OutColor,
-									const FIndex3i& TriVertices,
+									const FIndex3i& TriElements,
 									const FVector3f& Bary,
 									const FDynamicMeshColorOverlay* InColorAttribute)
 	{
 		FVector4f Value1, Value2, Value3;
-		InColorAttribute->GetElement(TriVertices[0], Value1);
-		InColorAttribute->GetElement(TriVertices[1], Value2);
-		InColorAttribute->GetElement(TriVertices[2], Value3);
+
+		InColorAttribute->GetElement(TriElements[0], Value1);
+		InColorAttribute->GetElement(TriElements[1], Value2);
+		InColorAttribute->GetElement(TriElements[2], Value3);
 
 		const float Alpha = Bary[0], Beta = Bary[1], Theta  = Bary[2];
 		for (int32 Idx = 0; Idx < NumElements; ++Idx)
@@ -53,8 +54,11 @@ namespace TransferValuesLocals
 		if (!MeshAttributes->HasPrimaryColors())
 		{
 			MeshAttributes->EnablePrimaryColors();
+
+			// Start with a clean attribute and elements we can write to.
+			MeshAttributes->PrimaryColors()->CreateFromPredicate([](int /*ParentVID*/, int /*TriIDA*/, int /*TriIDB*/){return true;}, 0.f);
 		}
-		
+
 		return MeshAttributes->PrimaryColors();
 	}
 
@@ -142,6 +146,14 @@ bool FTransferVertexColorAttribute::TransferColorsToMesh(FDynamicMesh3& InOutTar
 			InternalTargetMeshNormals->ComputeVertexNormals();
 		}
 	}
+
+	// Task context for the parallel for loops down below, to avoid repeatedly re-allocating an array. 
+	struct FTaskContext
+	{
+		TArray<int32> ElementIDs;
+	};
+		
+	TArray<FTaskContext> TaskContexts;
 	
 	FDynamicMeshColorOverlay* TargetColors = GetOrCreateColorAttribute(InOutTargetMesh);
 	checkSlow(TargetColors);
@@ -155,8 +167,8 @@ bool FTransferVertexColorAttribute::TransferColorsToMesh(FDynamicMesh3& InOutTar
 	if (TransferMethod == ETransferMethod::ClosestPointOnSurface)
 	{
 		MatchedVertices.Init(false, NumVerticesToTransfer);
-		
-		ParallelFor(NumVerticesToTransfer, [this, &InOutTargetMesh, &TargetColors, &InternalTargetMeshNormals, bUseSubset](int32 InVertexID)
+
+		ParallelForWithTaskContext(TaskContexts, NumVerticesToTransfer, [this, &InOutTargetMesh, &TargetColors, &InternalTargetMeshNormals, bUseSubset](FTaskContext& Context, int32 InVertexID)
 		{
 			if (Cancelled()) 
 			{
@@ -181,7 +193,12 @@ bool FTransferVertexColorAttribute::TransferColorsToMesh(FDynamicMesh3& InOutTar
 				FVector4f Color;
 				if (TransferColorToPoint(Color, Point, Normal))
 				{
-					TargetColors->SetElement(VertexID, Color);
+					TargetColors->GetVertexElements(VertexID, Context.ElementIDs);
+
+					for (int32 ElementID: Context.ElementIDs)
+					{
+						TargetColors->SetElement(ElementID, Color);
+					}
 					MatchedVertices[VertexID] = true;
 				}
 			}
@@ -221,7 +238,7 @@ bool FTransferVertexColorAttribute::TransferColorsToMesh(FDynamicMesh3& InOutTar
          *                    where L is a cotangent Laplacian and M is a mass matrix
          *  
 		 */
-	
+
 		MatchedVertices.Init(false, InOutTargetMesh.MaxVertexID());
  
 		// because the inpaint algorithm can extract data from regions outside the target vertex subset, a temporary attribute is used to modify the values.
@@ -234,7 +251,7 @@ bool FTransferVertexColorAttribute::TransferColorsToMesh(FDynamicMesh3& InOutTar
 		FDynamicMeshColorOverlay* EditedColors = bUseSubset ? &SubsetTargetColors : TargetColors;
  
 		// For every vertex on the target mesh try to find the match on the source mesh using the distance and normal checks
-		ParallelFor(InOutTargetMesh.MaxVertexID(), [this, &InOutTargetMesh, &EditedColors, &InternalTargetMeshNormals](int32 VertexID)
+		ParallelForWithTaskContext(TaskContexts, InOutTargetMesh.MaxVertexID(), [this, &InOutTargetMesh, &EditedColors, &InternalTargetMeshNormals](FTaskContext& Context, int32 VertexID)
 		{
 			if (Cancelled()) 
 			{
@@ -263,7 +280,12 @@ bool FTransferVertexColorAttribute::TransferColorsToMesh(FDynamicMesh3& InOutTar
 				FVector4f Color;
 				if (TransferColorToPoint(Color, Point, Normal))
 				{
-					EditedColors->SetElement(VertexID, Color);
+					EditedColors->GetVertexElements(VertexID, Context.ElementIDs);
+
+					for (int32 ElementID: Context.ElementIDs)
+					{
+						EditedColors->SetElement(ElementID, Color);
+					}
 					MatchedVertices[VertexID] = true;
 				}
 			}
@@ -301,15 +323,21 @@ bool FTransferVertexColorAttribute::TransferColorsToMesh(FDynamicMesh3& InOutTar
  
 		auto CopySubsetColorsIfNeeded = [Subset = TargetVerticesSubset, &EditedColors, &TargetColors, &InOutTargetMesh]()
 		{
+			TArray<int32> ElementIDs;
 			if (EditedColors && EditedColors != TargetColors)
 			{
 				for (const int32 VertexID: Subset)
 				{
 					if (InOutTargetMesh.IsVertex(VertexID))
 					{
-						FVector4f Color;
-						EditedColors->GetElement(VertexID, Color);
-						TargetColors->SetElement(VertexID, Color);
+						TargetColors->GetVertexElements(VertexID, ElementIDs);
+
+						for (int32 ElementID: ElementIDs)
+						{
+							FVector4f Color;
+							EditedColors->GetElement(ElementID, Color);
+							TargetColors->SetElement(ElementID, Color);
+						}
 					}
 				}
 			}
@@ -490,13 +518,13 @@ bool FTransferVertexColorAttribute::TransferColorToPoint(FVector4f& OutColor, co
 	}
 
 	const FVector3f BaryF((float)Bary[0], (float)Bary[1], (float)Bary[2]);
-	const FIndex3i TriVertex = SourceMesh->GetTriangle(TriID);
 
 	const FDynamicMeshColorOverlay* SourceColors = SourceMesh->Attributes()->PrimaryColors();
+	const FIndex3i ColorTriElements = SourceColors->GetTriangle(TriID);
 	if (SearchRadius < 0 && NormalThreshold < 0)
 	{
 		// If the radius and normals are ignored, simply interpolate the values and return the result
-		InterpolateVertexAttribute(OutColor, TriVertex, BaryF, SourceColors);
+		InterpolateVertexAttribute(OutColor, ColorTriElements, BaryF, SourceColors);
 	}
 	else
 	{
@@ -516,9 +544,10 @@ bool FTransferVertexColorAttribute::TransferColorToPoint(FVector4f& OutColor, co
 			const bool bHasSourceNormals = SourceMesh->HasVertexNormals();
 			if (ensure(bHasSourceNormals || InternalSourceMeshNormals.IsValid()))
 			{
-				Normal0 = bHasSourceNormals ? SourceMesh->GetVertexNormal(TriVertex[0]) : ToUENormal(InternalSourceMeshNormals->GetNormals()[TriVertex[0]]);
-				Normal1 = bHasSourceNormals ? SourceMesh->GetVertexNormal(TriVertex[1]) : ToUENormal(InternalSourceMeshNormals->GetNormals()[TriVertex[1]]);
-				Normal2 = bHasSourceNormals ? SourceMesh->GetVertexNormal(TriVertex[2]) : ToUENormal(InternalSourceMeshNormals->GetNormals()[TriVertex[2]]);
+				const FIndex3i TriVertices = SourceMesh->GetTriangle(TriID);
+				Normal0 = bHasSourceNormals ? SourceMesh->GetVertexNormal(TriVertices[0]) : ToUENormal(InternalSourceMeshNormals->GetNormals()[TriVertices[0]]);
+				Normal1 = bHasSourceNormals ? SourceMesh->GetVertexNormal(TriVertices[1]) : ToUENormal(InternalSourceMeshNormals->GetNormals()[TriVertices[1]]);
+				Normal2 = bHasSourceNormals ? SourceMesh->GetVertexNormal(TriVertices[2]) : ToUENormal(InternalSourceMeshNormals->GetNormals()[TriVertices[2]]);
 			}
 
 			const FVector3f MatchedNormal = Normalized(BaryF[0]*Normal0 + BaryF[1]*Normal1 + BaryF[2]*Normal2);
@@ -535,7 +564,7 @@ bool FTransferVertexColorAttribute::TransferColorToPoint(FVector4f& OutColor, co
 		
 		if (bPassedRadiusCheck && bPassedNormalsCheck)
 		{
-			InterpolateVertexAttribute(OutColor, TriVertex, BaryF, SourceColors);
+			InterpolateVertexAttribute(OutColor, ColorTriElements, BaryF, SourceColors);
 		}
 		else
 		{
