Geometry Scripting: Fix index mismatch when transfering colors from one mesh to another using inpaint

#rnx

diff --git a/Engine/Plugins/Runtime/GeometryProcessing/Source/DynamicMesh/Private/Operations/TransferDynamicMeshAttributes.cpp b/Engine/Plugins/Runtime/GeometryProcessing/Source/DynamicMesh/Private/Operations/TransferDynamicMeshAttributes.cpp
--- a/Engine/Plugins/Runtime/GeometryProcessing/Source/DynamicMesh/Private/Operations/TransferDynamicMeshAttributes.cpp
+++ b/Engine/Plugins/Runtime/GeometryProcessing/Source/DynamicMesh/Private/Operations/TransferDynamicMeshAttributes.cpp
@@ -240,6 +240,8 @@ bool FTransferVertexColorAttribute::TransferColorsToMesh(FDynamicMesh3& InOutTar
 		 */
 
 		MatchedVertices.Init(false, InOutTargetMesh.MaxVertexID());
+		TArray<FVector4f> MatchedColors;
+		MatchedColors.Init(FVector4f::Zero(), InOutTargetMesh.MaxVertexID());
  
 		// because the inpaint algorithm can extract data from regions outside the target vertex subset, a temporary attribute is used to modify the values.
 		// NOTE: make sure to copy the values of the vertex subset into the complete TargetColors attribute before exciting the function. (see CopySubsetColorsIfNeeded)
@@ -251,7 +253,7 @@ bool FTransferVertexColorAttribute::TransferColorsToMesh(FDynamicMesh3& InOutTar
 		FDynamicMeshColorOverlay* EditedColors = bUseSubset ? &SubsetTargetColors : TargetColors;
  
 		// For every vertex on the target mesh try to find the match on the source mesh using the distance and normal checks
-		ParallelForWithTaskContext(TaskContexts, InOutTargetMesh.MaxVertexID(), [this, &InOutTargetMesh, &EditedColors, &InternalTargetMeshNormals](FTaskContext& Context, int32 VertexID)
+		ParallelForWithTaskContext(TaskContexts, InOutTargetMesh.MaxVertexID(), [this, &InOutTargetMesh, &EditedColors, &InternalTargetMeshNormals, &MatchedColors](FTaskContext& Context, int32 VertexID)
 		{
 			if (Cancelled()) 
 			{
@@ -277,7 +279,7 @@ bool FTransferVertexColorAttribute::TransferColorsToMesh(FDynamicMesh3& InOutTar
 					}
 				}
  
-				FVector4f Color;
+				FVector4f& Color = MatchedColors[VertexID];
 				if (TransferColorToPoint(Color, Point, Normal))
 				{
 					EditedColors->GetVertexElements(VertexID, Context.ElementIDs);
@@ -369,13 +371,12 @@ bool FTransferVertexColorAttribute::TransferColorsToMesh(FDynamicMesh3& InOutTar
 		{
 			if (InOutTargetMesh.IsVertex(VertexID) && MatchedVertices[VertexID])
 			{
-				FVector4f Data;
-				EditedColors->GetElement(VertexID, Data);
+				const FVector4f& Color = MatchedColors[VertexID];
  
 				const int32 CurIdx = FixedIndices.Num();
 				for (int32 Idx = 0; Idx < NumElements; ++Idx)
 				{
-					FixedValuesTriplets.emplace_back(CurIdx, Idx, Data[Idx]);
+					FixedValuesTriplets.emplace_back(CurIdx, Idx, Color[Idx]);
 				}
  
 				checkSlow(VertexID < ToIndex.Num());
@@ -434,6 +435,7 @@ bool FTransferVertexColorAttribute::TransferColorsToMesh(FDynamicMesh3& InOutTar
 			FSparseMatrixD TargetValuesTransposed = TargetValues.transpose(); 
  
 			// Iterate over every column containing all values for the vertex
+			TArray<int32> ElementIDs; 
 			for (int32 ColIdx = 0; ColIdx < TargetValuesTransposed.outerSize(); ++ColIdx)
 			{
 				FVector4f Data(0.f);
@@ -442,10 +444,10 @@ bool FTransferVertexColorAttribute::TransferColorsToMesh(FDynamicMesh3& InOutTar
 				// Iterate over only non-zero rows (i.e. non-zero values)
 				for (FSparseMatrixD::InnerIterator Itr(TargetValuesTransposed, ColIdx); Itr; ++Itr)
 				{
-					const int32 Idx = static_cast<int32>(Itr.row());
+					const int32 Index = static_cast<int32>(Itr.row());
 					const float Value = static_cast<float>(Itr.value());
-					Data[Idx] = Value;
-					N[Idx] += 1;
+					Data[Index] += Value;
+					N[Index]++;
 				}
 
 				// normalize
@@ -453,13 +455,18 @@ bool FTransferVertexColorAttribute::TransferColorsToMesh(FDynamicMesh3& InOutTar
 				{
 					if (N[Index] > 1)
 					{
-						Data /= static_cast<float>(N[Index]); 
+						Data[Index] /= static_cast<float>(N[Index]); 
 					}
 				}
 				
 				const int32 VertexIDLinearalized = bVariablesOnly ? static_cast<int32>(VariableRows[ColIdx]) : ColIdx; // linearized vertex ID (matrix row) of the variable in the Energy matrix
 				const int32 VertexID = ToMeshV[VertexIDLinearalized];
-				EditedColors->SetElement(VertexID, Data);
+
+				EditedColors->GetVertexElements(VertexID, ElementIDs);
+				for (int32 ElementID: ElementIDs)
+				{
+					EditedColors->SetElement(ElementID, Data);
+				}
 			}
  
 			// copy values from the subset attribute if using subset
