Geometry Scripting: added the ability to create and transfer vertex colors on hard edges

if the target's projection onto the source falls exactly on a vertex of the source mesh, there's no way of ensuring that the nearest face is the right one (the first face is returned if this vertex is shared by several faces)
to deal with this situation, the idea is to create vertex instances for each face and slightly offset their positions in order to created a "shrinked" face thus obtain a more relevant projection
use bHardEdges to activate that projection and set the biased ratio between 0 and 1

#rnx

diff --git a/Engine/Plugins/Runtime/GeometryProcessing/Source/DynamicMesh/Private/Operations/TransferDynamicMeshAttributes.cpp b/Engine/Plugins/Runtime/GeometryProcessing/Source/DynamicMesh/Private/Operations/TransferDynamicMeshAttributes.cpp
--- a/Engine/Plugins/Runtime/GeometryProcessing/Source/DynamicMesh/Private/Operations/TransferDynamicMeshAttributes.cpp
+++ b/Engine/Plugins/Runtime/GeometryProcessing/Source/DynamicMesh/Private/Operations/TransferDynamicMeshAttributes.cpp
@@ -46,10 +46,9 @@ namespace TransferValuesLocals
 		}
 	}
 
-	FDynamicMeshColorOverlay* GetOrCreateColorAttribute(FDynamicMesh3& InMesh)
+	FDynamicMeshColorOverlay* GetOrCreateColorAttribute(FDynamicMesh3& InMesh, const bool InSplit)
 	{
 		checkSlow(InMesh.HasAttributes());
-
 		FDynamicMeshAttributeSet* MeshAttributes = InMesh.Attributes();
 		if (!MeshAttributes->HasPrimaryColors())
 		{
@@ -58,7 +57,21 @@ namespace TransferValuesLocals
 			// Start with a clean attribute and elements we can write to.
 			MeshAttributes->PrimaryColors()->CreateFromPredicate([](int /*ParentVID*/, int /*TriIDA*/, int /*TriIDB*/){return true;}, 0.f);
 		}
-
+		
+		if (InSplit)
+		{
+			// create vertex instances for each face if needed
+			FDynamicMeshColorOverlay* ColorOverlay = MeshAttributes->PrimaryColors();
+		
+			ColorOverlay->SplitVerticesWithPredicate(
+				[](int ElementIdx, int TriID) { return true; },
+				[ColorOverlay](int ElementIdx, int TriID, float* FillVect)
+				{
+					const FVector4f CurValue = ColorOverlay->GetElement(ElementIdx);
+					FillVect[0] = CurValue.X; FillVect[1] = CurValue.Y; FillVect[2] = CurValue.Z; FillVect[3] = CurValue.W;
+				});
+		}
+		
 		return MeshAttributes->PrimaryColors();
 	}
 
@@ -66,6 +79,83 @@ namespace TransferValuesLocals
 	{
 		return FVector3f((float)Normal.X, (float)Normal.Y, (float)Normal.Z);
 	}
+
+	struct FTriangleData
+	{
+		FVector3d Normal;
+		FVector3d Centroid;
+	};
+	
+	void GetBiasedElementPositions(const FDynamicMesh3& InOutTargetMesh, const float InRatio, TArray<FVector>& OutElementPositions)
+	{
+		const FDynamicMeshColorOverlay* ColorOverlay = InOutTargetMesh.Attributes()->PrimaryColors();
+		if (!ensure(ColorOverlay))
+		{
+			return;
+		}
+		
+		// get normal + centroid of each face
+		TArray<FTriangleData> TrianglesData;
+		TrianglesData.Reserve(InOutTargetMesh.TriangleCount());
+		double DummyTriArea;
+		for (int TriIdx : InOutTargetMesh.TriangleIndicesItr())
+		{
+			FTriangleData Data;
+			InOutTargetMesh.GetTriInfo(TriIdx, Data.Normal, DummyTriArea, Data.Centroid);
+			TrianglesData.Emplace(MoveTemp(Data));
+		}
+
+		// store the triangle for each element  
+		TArray<int32> ElementToTriangle;
+		ElementToTriangle.Init(INDEX_NONE, ColorOverlay->MaxElementID());
+		for (int TriIdx : InOutTargetMesh.TriangleIndicesItr())
+		{
+			const FIndex3i TriElements = ColorOverlay->GetTriangle(TriIdx);
+			ElementToTriangle[TriElements[0]] = TriIdx;
+			ElementToTriangle[TriElements[1]] = TriIdx;
+			ElementToTriangle[TriElements[2]] = TriIdx;
+		}
+		
+		// clamp the bias between UE_SMALL_NUMBER and 1.0
+		const double Ratio = FMath::Clamp(static_cast<double>(FMath::Abs(InRatio)), UE_KINDA_SMALL_NUMBER, 1.f);
+		
+		// compute biased vertex instance positions (per ElementID)
+		// note that MaxElementID() is used here instead of ElementCount()
+		OutElementPositions.Reset();
+		OutElementPositions.Init(FVector::Zero(), ColorOverlay->MaxElementID());
+		for (const int32 ElementID : ColorOverlay->ElementIndicesItr())
+		{
+			const int32 ParentVertex = ColorOverlay->GetParentVertex(ElementID);
+
+			// initialize with the parent vertex position
+			FVector& VertexPos = OutElementPositions[ElementID];
+			VertexPos = InOutTargetMesh.GetVertexRef(ParentVertex);
+
+			// get the face the element belongs to
+			const int32 TriangleIndex = ElementToTriangle[ElementID];
+			if (TriangleIndex != INDEX_NONE)
+			{
+				const FTriangleData& Triangle = TrianglesData[TriangleIndex];
+
+				// build triangle base with VertexPos - TriCentroid as the X axis and the triangle Normal as Z
+				const FVector& TriCentroid = Triangle.Centroid;
+				const FVector X = (VertexPos - TriCentroid).GetSafeNormal();
+				const FVector Y = Triangle.Normal.Cross(X);
+				const FQuat Rotation = FRotationMatrix::MakeFromXY(X, Y).ToQuat();
+				const FTransform TriTransform(Rotation, Triangle.Centroid);
+
+				// slightly move the vertex along the X axis to "shrink the triangle"
+				VertexPos = TriTransform.InverseTransformPosition(VertexPos);
+				VertexPos.X = VertexPos.X - (VertexPos.X * Ratio);
+				VertexPos = TriTransform.TransformPosition(VertexPos);
+			}
+		}
+	}
+
+	struct FTaskContext
+	{
+		TArray<int32> ElementIDs;
+	};
 }
 
 FTransferVertexColorAttribute::FTransferVertexColorAttribute(
@@ -147,15 +237,8 @@ bool FTransferVertexColorAttribute::TransferColorsToMesh(FDynamicMesh3& InOutTar
 		}
 	}
 
-	// Task context for the parallel for loops down below, to avoid repeatedly re-allocating an array. 
-	struct FTaskContext
-	{
-		TArray<int32> ElementIDs;
-	};
-		
-	TArray<FTaskContext> TaskContexts;
+	FDynamicMeshColorOverlay* TargetColors = GetOrCreateColorAttribute(InOutTargetMesh, bHardEdges);
 	
-	FDynamicMeshColorOverlay* TargetColors = GetOrCreateColorAttribute(InOutTargetMesh);
 	checkSlow(TargetColors);
 
 	bool bFailed = false;
@@ -166,58 +249,14 @@ bool FTransferVertexColorAttribute::TransferColorsToMesh(FDynamicMesh3& InOutTar
 	
 	if (TransferMethod == ETransferMethod::ClosestPointOnSurface)
 	{
-		MatchedVertices.Init(false, NumVerticesToTransfer);
-
-		ParallelForWithTaskContext(TaskContexts, NumVerticesToTransfer, [this, &InOutTargetMesh, &TargetColors, &InternalTargetMeshNormals, bUseSubset](FTaskContext& Context, int32 InVertexID)
-		{
-			if (Cancelled()) 
-			{
-				return;
-			}
-
-			const int32 VertexID = bUseSubset ? TargetVerticesSubset[InVertexID] : InVertexID;
-			if (InOutTargetMesh.IsVertex(VertexID)) 
-			{
-				const FVector3d Point = InOutTargetMesh.GetVertex(VertexID);
-
-				FVector3f Normal = FVector3f::UnitY();
-				if (NormalThreshold >= 0) 
-				{
-					const bool bHasNormals = InOutTargetMesh.HasVertexNormals();
-					if (ensure(bHasNormals || InternalTargetMeshNormals))
-					{
-						Normal = bHasNormals ? InOutTargetMesh.GetVertexNormal(VertexID) : ToUENormal(InternalTargetMeshNormals->GetNormals()[VertexID]);
-					}
-				}
-
-				FVector4f Color;
-				if (TransferColorToPoint(Color, Point, Normal))
-				{
-					TargetColors->GetVertexElements(VertexID, Context.ElementIDs);
-
-					for (int32 ElementID: Context.ElementIDs)
-					{
-						TargetColors->SetElement(ElementID, Color);
-					}
-					MatchedVertices[VertexID] = true;
-				}
-			}
-
-		}, bUseParallel ? EParallelForFlags::None : EParallelForFlags::ForceSingleThread);
-
+		const int32 NumMatched = TransferUsingClosestPoint(InOutTargetMesh, InternalTargetMeshNormals);
+		const int32 NumVerticesToMatch = bHardEdges ? TargetColors->ElementCount() : NumVerticesToTransfer;
+		
 		// If the caller requested to simply find the closest point for all vertices then the number of matched vertices
 		// must be equal to the target mesh vertex count
 		if (SearchRadius < 0 && NormalThreshold < 0)
 		{
-			int32 NumMatched = 0;
-			for (bool Flag : MatchedVertices)
-			{
-				if (Flag) 
-				{ 
-					NumMatched++;
-				}
-			}
-			bFailed = NumMatched != NumVerticesToTransfer;
+			bFailed = NumMatched != NumVerticesToMatch;
 		}
 	} 
 	else if (TransferMethod == ETransferMethod::Inpaint)
@@ -251,7 +290,10 @@ bool FTransferVertexColorAttribute::TransferColorsToMesh(FDynamicMesh3& InOutTar
 			SubsetTargetColors.Copy(*TargetColors);
 		}
 		FDynamicMeshColorOverlay* EditedColors = bUseSubset ? &SubsetTargetColors : TargetColors;
- 
+
+		// Task context for the parallel for loops down below, to avoid repeatedly re-allocating an array.
+		TArray<FTaskContext> TaskContexts;
+		
 		// For every vertex on the target mesh try to find the match on the source mesh using the distance and normal checks
 		ParallelForWithTaskContext(TaskContexts, InOutTargetMesh.MaxVertexID(), [this, &InOutTargetMesh, &EditedColors, &InternalTargetMeshNormals, &MatchedColors](FTaskContext& Context, int32 VertexID)
 		{
@@ -523,7 +565,7 @@ bool FTransferVertexColorAttribute::TransferColorToPoint(FVector4f& OutColor, co
 	{
 		return false;
 	}
-
+	
 	const FVector3f BaryF((float)Bary[0], (float)Bary[1], (float)Bary[2]);
 
 	const FDynamicMeshColorOverlay* SourceColors = SourceMesh->Attributes()->PrimaryColors();
@@ -612,3 +654,138 @@ bool FTransferVertexColorAttribute::FindClosestPointOnSourceSurface(const FVecto
 
 	return true;
 }
+
+int32 FTransferVertexColorAttribute::TransferUsingClosestPoint(FDynamicMesh3& InOutTargetMesh, const TUniquePtr<FMeshNormals>& InTargetMeshNormals)
+{
+	using namespace TransferValuesLocals;
+	
+	if (!ensure(TransferMethod == ETransferMethod::ClosestPointOnSurface))
+	{
+		return 0;
+	}
+	
+	FDynamicMeshAttributeSet* MeshAttributes = InOutTargetMesh.Attributes();
+	if (!ensure(MeshAttributes->HasPrimaryColors()))
+	{
+		return 0;
+	}
+		
+	FDynamicMeshColorOverlay* TargetColors = MeshAttributes->PrimaryColors();
+
+	// transfer using element ids instead of vertices
+	if (bHardEdges)
+	{
+		// compute per-vertex instance biased positions
+		TArray<FVector> BiasedPositions;
+		GetBiasedElementPositions(InOutTargetMesh, BiasRatio, BiasedPositions);
+
+		// note that MaxElementID() is used here instead of ElementCount()
+		const int32 NumElementsToTransfer = TargetColors->MaxElementID();
+		MatchedVertices.Init(false, NumElementsToTransfer);
+
+		ParallelFor(NumElementsToTransfer, [this, &InOutTargetMesh, &TargetColors, &InTargetMeshNormals, &BiasedPositions](int32 InElementID)
+		{
+			if (Cancelled()) 
+			{
+				return;
+			}
+
+			if (TargetColors->IsElement(InElementID))
+			{
+				const int32 VertexID = TargetColors->GetParentVertex(InElementID);
+				if (InOutTargetMesh.IsVertex(VertexID))
+				{
+					const FVector3d& BiasedPoint = BiasedPositions.IsValidIndex(InElementID) ? BiasedPositions[InElementID] : InOutTargetMesh.GetVertexRef(VertexID);
+
+					FVector3f Normal = FVector3f::UnitY();
+					if (NormalThreshold >= 0) 
+					{
+						const bool bHasNormals = InOutTargetMesh.HasVertexNormals();
+						if (ensure(bHasNormals || InTargetMeshNormals))
+						{
+							Normal = bHasNormals ? InOutTargetMesh.GetVertexNormal(VertexID) : ToUENormal(InTargetMeshNormals->GetNormals()[VertexID]);
+						}
+					}
+
+					FVector4f Color;
+					if (TransferColorToPoint(Color, BiasedPoint, Normal))
+					{
+						TargetColors->SetElement(InElementID, Color);								
+						MatchedVertices[InElementID] = true;
+					}
+				}
+			}
+
+		}, bUseParallel ? EParallelForFlags::None : EParallelForFlags::ForceSingleThread);
+
+		// compute matching number
+		int32 NumMatched = 0;
+		for (const int32 ElementID : TargetColors->ElementIndicesItr())
+		{
+			if (MatchedVertices[ElementID])
+			{
+				NumMatched++;
+			}
+		}
+
+		return NumMatched;
+	}
+
+
+	// compute the transfer only for the subset of vertices if necessary  
+	const bool bUseSubset = !TargetVerticesSubset.IsEmpty();
+	const int32 NumVerticesToTransfer = bUseSubset ? TargetVerticesSubset.Num() : InOutTargetMesh.MaxVertexID();
+
+	TArray<FTaskContext> TaskContexts;
+
+	MatchedVertices.Init(false, NumVerticesToTransfer);
+
+	ParallelForWithTaskContext(TaskContexts, NumVerticesToTransfer, [this, &InOutTargetMesh, &TargetColors, &InTargetMeshNormals, bUseSubset](FTaskContext& Context, int32 InVertexID)
+	{
+		if (Cancelled()) 
+		{
+			return;
+		}
+
+		const int32 VertexID = bUseSubset ? TargetVerticesSubset[InVertexID] : InVertexID;
+		if (InOutTargetMesh.IsVertex(VertexID)) 
+		{
+			const FVector3d Point = InOutTargetMesh.GetVertex(VertexID);
+
+			FVector3f Normal = FVector3f::UnitY();
+			if (NormalThreshold >= 0) 
+			{
+				const bool bHasNormals = InOutTargetMesh.HasVertexNormals();
+				if (ensure(bHasNormals || InTargetMeshNormals))
+				{
+					Normal = bHasNormals ? InOutTargetMesh.GetVertexNormal(VertexID) : ToUENormal(InTargetMeshNormals->GetNormals()[VertexID]);
+				}
+			}
+
+			FVector4f Color;
+			if (TransferColorToPoint(Color, Point, Normal))
+			{
+				TargetColors->GetVertexElements(VertexID, Context.ElementIDs);
+
+				for (int32 ElementID: Context.ElementIDs)
+				{
+					TargetColors->SetElement(ElementID, Color);
+				}
+			
+				MatchedVertices[VertexID] = true;
+			}
+		}
+
+	}, bUseParallel ? EParallelForFlags::None : EParallelForFlags::ForceSingleThread);
+
+	int32 NumMatched = 0;
+	for (bool Flag : MatchedVertices)
+	{
+		if (Flag) 
+		{ 
+			NumMatched++;
+		}
+	}
+	
+	return NumMatched;
+}
\ No newline at end of file
diff --git a/Engine/Plugins/Runtime/GeometryProcessing/Source/DynamicMesh/Public/Operations/TransferDynamicMeshAttributes.h b/Engine/Plugins/Runtime/GeometryProcessing/Source/DynamicMesh/Public/Operations/TransferDynamicMeshAttributes.h
--- a/Engine/Plugins/Runtime/GeometryProcessing/Source/DynamicMesh/Public/Operations/TransferDynamicMeshAttributes.h
+++ b/Engine/Plugins/Runtime/GeometryProcessing/Source/DynamicMesh/Public/Operations/TransferDynamicMeshAttributes.h
@@ -110,6 +110,12 @@ public:
 	/** MatchedVertices[VertexID] is set to true for a target mesh vertex ID with a match found, false otherwise. */
 	TArray<bool> MatchedVertices;
 
+	/** Creates vertex instances per triangle to be able to have per-face vertex colors. */
+	bool bHardEdges = false;
+	
+	/** Ratio used to blend a vertex between its position and the center of the face (0 = vertex position, 1 = face centroid) */
+	float BiasRatio = UE_KINDA_SMALL_NUMBER;
+
 protected:
 		
 	/** Source mesh we are transferring colors from. */
@@ -178,6 +184,11 @@ protected:
 	 */
 	bool FindClosestPointOnSourceSurface(const FVector3d& InPoint, const FTransformSRT3d& InToWorld, int32& OutTriID, FVector3d& OutBary) const;
 
+	/**
+     * Transfer the colors from the source mesh to the given target mesh using the closest point algorithm.
+     * @return the number of matching vertices / elements
+	 */
+	int32 TransferUsingClosestPoint(FDynamicMesh3& InOutTargetMesh, const TUniquePtr<FMeshNormals>& InTargetMeshNormals);
 };
 
 } // end namespace UE::Geometry
diff --git a/Engine/Plugins/Runtime/GeometryScripting/Source/GeometryScriptingCore/Private/MeshVertexColorFunctions.cpp b/Engine/Plugins/Runtime/GeometryScripting/Source/GeometryScriptingCore/Private/MeshVertexColorFunctions.cpp
--- a/Engine/Plugins/Runtime/GeometryScripting/Source/GeometryScriptingCore/Private/MeshVertexColorFunctions.cpp
+++ b/Engine/Plugins/Runtime/GeometryScripting/Source/GeometryScriptingCore/Private/MeshVertexColorFunctions.cpp
@@ -510,6 +510,8 @@ UDynamicMesh* UGeometryScriptLibrary_MeshVertexColorFunctions::TransferVertexCol
 		FTransferVertexColorAttribute TransferColors(&ReadMesh);
 		TransferColors.TransferMethod = static_cast<FTransferVertexColorAttribute::ETransferMethod>(Options.TransferMethod);
 		TransferColors.bUseParallel = true;
+		TransferColors.bHardEdges = Options.bHardEdges;
+		TransferColors.BiasRatio = Options.BiasRatio;
 
 		TargetMesh->EditMesh([&](FDynamicMesh3& EditMesh)
 		{
diff --git a/Engine/Plugins/Runtime/GeometryScripting/Source/GeometryScriptingCore/Public/GeometryScript/MeshVertexColorFunctions.h b/Engine/Plugins/Runtime/GeometryScripting/Source/GeometryScriptingCore/Public/GeometryScript/MeshVertexColorFunctions.h
--- a/Engine/Plugins/Runtime/GeometryScripting/Source/GeometryScriptingCore/Public/GeometryScript/MeshVertexColorFunctions.h
+++ b/Engine/Plugins/Runtime/GeometryScripting/Source/GeometryScriptingCore/Public/GeometryScript/MeshVertexColorFunctions.h
@@ -94,6 +94,14 @@ struct GEOMETRYSCRIPTINGCORE_API FGeometryScriptTransferMeshVertexColorsOptions
 	/** The strength of each post-processing smoothing iteration. */
 	UPROPERTY(EditAnywhere, BlueprintReadWrite, Category = Options, meta = (UIMin = 0, UIMax = 1, ClampMin = 0, ClampMax = 1, EditCondition = "TransferMethod == ETransferVertexColorMethod::Inpaint"))
 	float SmoothingStrength = 0.0f;
+
+	/** Creates vertex instances per triangle to be able to have per-face vertex colors. */
+	UPROPERTY(EditAnywhere, BlueprintReadWrite, Category = Options)
+	bool bHardEdges = false;
+	
+	/** Ratio used to blend a vertex between its position and the center of the face (0 = vertex position, 1 = face centroid) */
+	UPROPERTY(EditAnywhere, BlueprintReadWrite, Category = Options, meta = (UIMin = 0, UIMax = 1, ClampMin = 0, ClampMax = 1))
+	float BiasRatio = UE_KINDA_SMALL_NUMBER;
 };
 
 
