Geometry: fix discarded result of TransformPosition  and ExpandBy
Found with [[nodiscard]]

#rnx

diff --git a/Engine/Plugins/Runtime/GeometryProcessing/Source/DynamicMesh/Private/ShapeApproximation/SimpleShapeSet3.cpp b/Engine/Plugins/Runtime/GeometryProcessing/Source/DynamicMesh/Private/ShapeApproximation/SimpleShapeSet3.cpp
--- a/Engine/Plugins/Runtime/GeometryProcessing/Source/DynamicMesh/Private/ShapeApproximation/SimpleShapeSet3.cpp
+++ b/Engine/Plugins/Runtime/GeometryProcessing/Source/DynamicMesh/Private/ShapeApproximation/SimpleShapeSet3.cpp
@@ -778,7 +778,7 @@ bool FSimpleShapeSet3d::MergeShapes(int32 MergeAboveCount, const FMergeShapesSet
 		{
 			for (int32 Idx = HullIdxStart; Idx < HullVertices.Num(); ++Idx)
 			{
-				ShapeTransform->TransformPosition(HullVertices[Idx]);
+				HullVertices[Idx] = ShapeTransform->TransformPosition(HullVertices[Idx]);
 			}
 		}
 		HullVertexCounts.Add(Vertices.Num());
diff --git a/Engine/Plugins/Runtime/GeometryProcessing/Source/GeometryAlgorithms/Private/Curve/PolygonOffsetUtils.cpp b/Engine/Plugins/Runtime/GeometryProcessing/Source/GeometryAlgorithms/Private/Curve/PolygonOffsetUtils.cpp
--- a/Engine/Plugins/Runtime/GeometryProcessing/Source/GeometryAlgorithms/Private/Curve/PolygonOffsetUtils.cpp
+++ b/Engine/Plugins/Runtime/GeometryProcessing/Source/GeometryAlgorithms/Private/Curve/PolygonOffsetUtils.cpp
@@ -39,7 +39,7 @@ namespace UE::Geometry::Private
 
 		RealType OffsetExpand = bOneSidedOffset ? FMath::Max(0, InOffset) : FMath::Abs(InOffset);
 		RealType MaxMiterExpand = InJoinType == EPolygonOffsetJoinType::Miter ? InMiterLimit : 0;
-		InputBounds.ExpandBy(OffsetExpand + MaxMiterExpand); // pad to allow for offset geometry
+		InputBounds = InputBounds.ExpandBy(OffsetExpand + MaxMiterExpand); // pad to allow for offset geometry
 		RealType InputRange = FMath::Max(TMathUtil<RealType>::ZeroTolerance, InputBounds.GetSize().GetMax());
 		ClipperOffset.SetRoundScaleFactors(static_cast<double>(InputRange) / static_cast<double>(IntRange), DefaultStepsPerRadianScale, MaxStepsPerRadian);
 		
diff --git a/Engine/Plugins/Runtime/GeometryScripting/Source/GeometryScriptingCore/Private/CollisionFunctions.cpp b/Engine/Plugins/Runtime/GeometryScripting/Source/GeometryScriptingCore/Private/CollisionFunctions.cpp
--- a/Engine/Plugins/Runtime/GeometryScripting/Source/GeometryScriptingCore/Private/CollisionFunctions.cpp
+++ b/Engine/Plugins/Runtime/GeometryScripting/Source/GeometryScriptingCore/Private/CollisionFunctions.cpp
@@ -1241,7 +1241,7 @@ FGeometryScriptSimpleCollision UGeometryScriptLibrary_CollisionFunctions::MergeS
 		{
 			for (int32 Idx = HullIdxStart; Idx < HullVertices.Num(); ++Idx)
 			{
-				ShapeTransform->TransformPosition(HullVertices[Idx]);
+				HullVertices[Idx] = ShapeTransform->TransformPosition(HullVertices[Idx]);
 			}
 		}
 		HullVertexCounts.Add(Vertices.Num());
