GeometryProcessing: Fix FSubRegionRemesher having some extra edges in its EdgeROI after initializing from vertices, despite them not being present in the processed Edges. This could potentially affect FRestrictedSubRegionRemesher, but that class is intended to be initialized from a triangle set, which does not have this problem.

diff --git a/Engine/Plugins/Runtime/GeometryProcessing/Source/DynamicMesh/Public/SubRegionRemesher.h b/Engine/Plugins/Runtime/GeometryProcessing/Source/DynamicMesh/Public/SubRegionRemesher.h
--- a/Engine/Plugins/Runtime/GeometryProcessing/Source/DynamicMesh/Public/SubRegionRemesher.h
+++ b/Engine/Plugins/Runtime/GeometryProcessing/Source/DynamicMesh/Public/SubRegionRemesher.h
@@ -50,8 +50,8 @@ protected:
 	// set of triangles removed in last pass. Enable this by calling BeginTrackRemovedTrisInPass()
 	TSet<int> RemovedLastPassTris;
 
-	// temporary data structure used to quickly convert vertex set to edge set. Currenlty it is faster
-	// to construct this and then copy to a properly-sized TSet, than to incrementally build the TSet
+	// When deleted, remove PRAGMA_DISABLE_DEPRECATION_WARNINGS / PRAGMA_ENABLE_DEPRECATION_WARNINGS from ~FSubRegionRemesher()
+	UE_DEPRECATED(5.6, "This member is no longer used")
 	FUniqueIndexSet InitialROIEdgeVertices;
 
 	// temporary data structure used to quickly convert vertex set to triangle set. Currenlty it is faster
@@ -76,8 +76,9 @@ public:
 		};
 	}
 
+	PRAGMA_DISABLE_DEPRECATION_WARNINGS
 	virtual ~FSubRegionRemesher() {}
-
+	PRAGMA_ENABLE_DEPRECATION_WARNINGS
 
 	void Reset()
 	{
@@ -112,7 +113,8 @@ public:
 	}
 
 	/**
-	 * Initialize edge-subregion ROI from the VertexROI member that has been externally initialized
+	 * Initialize edge-subregion ROI from the VertexROI member that has been externally initialized.
+	 * Edge ROI is initialized to fully-contained edges (i.e. both endpoints are in VertexROI).
 	 */
 	void InitializeFromVertexROI()
 	{
@@ -122,29 +124,21 @@ public:
 			TRACE_CPUPROFILER_EVENT_SCOPE(InitializeFromVertexROI_1Init);
 			EdgeROI.Reset();
 			TriangleROI.Reset();
-			InitialROIEdgeVertices.Initialize(UseMesh->MaxVertexID());
 
-			Edges.Reset();
-			LastUpdateROICounter++;
-
-			// convert vertex ROI to set of one-ring edges. Note that this effectively does
-			// UpdateROI(), and if that function changes, this may need to be modified.
-			// New set of ROI vertices is also constructed simultaneously
+			// Convert vertex ROI to set of contained edges.
 			for (int VertIdx : VertexROI)
 			{
-				InitialROIEdgeVertices.Add(VertIdx);
 				UseMesh->EnumerateVertexEdges(VertIdx, [this, UseMesh, VertIdx](int32 eid)
 				{
-					bool bNewEdge = false;
-					EdgeROI.Add(eid, &bNewEdge);
-					if (bNewEdge)
+					FIndex2i EdgeV = UseMesh->GetEdgeV(eid);
+					int32 OtherVid = IndexUtil::FindEdgeOtherVertex(EdgeV, VertIdx);
+					if (VertexROI.Contains(OtherVid))
 					{
-						FIndex2i EdgeV = UseMesh->GetEdgeV(eid);
-						InitialROIEdgeVertices.Add((EdgeV.A == VertIdx) ? EdgeV.B : EdgeV.A);
-						Edges.Add(eid);
+						EdgeROI.Add(eid);
 					}
 				});
 			}
+			UpdateROI();
 		}
 
 		{
@@ -154,7 +148,7 @@ public:
 			// This conservatively guarantees that any triangle affected by a mesh edge op 
 			// (eg split/collapse/flip) is contained within this triangle set
 			InitialROITriangles.Initialize(UseMesh->MaxTriangleID());
-			for (int vid : InitialROIEdgeVertices)
+			for (int vid : VertexROI)
 			{
 				UseMesh->EnumerateVertexEdges(vid, [this](int32 eid)
 				{
@@ -193,9 +187,7 @@ public:
 	 */
 	void UpdateROI()
 	{
-		// note: if you change this function, you also need to change InitializeFromVertexROI() because it
-		// does this internally (or call this afterwards)
-		Edges.Reset();
+		Edges.Reset(EdgeROI.Num());
 		for (int eid : EdgeROI)
 		{
 			Edges.Add(eid);
