add geometry script method to transfer UVs from a source mesh to a target mesh by projection along a direction

diff --git a/Engine/Plugins/Runtime/GeometryProcessing/Source/DynamicMesh/Private/Parameterization/DynamicMeshUVEditor.cpp b/Engine/Plugins/Runtime/GeometryProcessing/Source/DynamicMesh/Private/Parameterization/DynamicMeshUVEditor.cpp
--- a/Engine/Plugins/Runtime/GeometryProcessing/Source/DynamicMesh/Private/Parameterization/DynamicMeshUVEditor.cpp
+++ b/Engine/Plugins/Runtime/GeometryProcessing/Source/DynamicMesh/Private/Parameterization/DynamicMeshUVEditor.cpp
@@ -329,6 +329,159 @@ void FDynamicMeshUVEditor::SetTriangleUVsFromPlanarProjection(
 	}
 }
 
+void FDynamicMeshUVEditor::TransferTriangleUVsFromMeshViaDirectionProjection(TArrayView<const int32> Triangles, TFunctionRef<FVector3d(const FVector3d&)> TransformTargetToSourceSpacePosition,
+	const FVector3d& ProjectionDirection, double ProjectionOffset, TFunctionRef<FVector3d(const FVector3d&)> TransformProjectionToSourceSpaceVector,
+	const FDynamicMeshAABBTree3& SourceMeshSpatial, int32 SourceMeshUVChannel,
+	const FTransferFromMeshViaProjectionSettings& Settings, FUVEditResult* Result)
+{
+	if (!ensure(UVOverlay) || Triangles.IsEmpty())
+	{
+		return;
+	}
+
+	const FDynamicMesh3* SourceMesh = SourceMeshSpatial.GetMesh();
+	check(SourceMesh);
+	const FDynamicMeshUVOverlay* SourceUVs = SourceMesh->HasAttributes() ? SourceMesh->Attributes()->GetUVLayer(SourceMeshUVChannel) : nullptr;
+	if (!SourceUVs)
+	{
+		return;
+	}
+
+	// test for empty range
+	if (Settings.MinDistance >= Settings.MaxDistance)
+	{
+		return;
+	}
+
+	FVector3d SourceProjDirection = TransformProjectionToSourceSpaceVector(ProjectionDirection.GetSafeNormal());
+	double ToSourceDistanceScale = SourceProjDirection.Length();
+	if (ToSourceDistanceScale < UE_DOUBLE_SMALL_NUMBER)
+	{
+		return;
+	}
+	SourceProjDirection /= ToSourceDistanceScale;
+	const FVector3d ConstantProjOffset = SourceProjDirection * ToSourceDistanceScale * ProjectionOffset;
+
+	// Depending on the projection range, we may need to project backwards and forwards, and may need to offset the projections
+	struct FProjectionInfo
+	{
+		double Sign, Offset, Max;
+	};
+	FProjectionInfo ProjInfo[2];
+	int32 NumDirs = 0;
+	if (Settings.MinDistance < 0 && Settings.MaxDistance > 0)
+	{
+		NumDirs = 2;
+		ProjInfo[0].Sign = -1;
+		ProjInfo[1].Sign = 1;
+		ProjInfo[0].Offset = 0;
+		ProjInfo[1].Offset = 0;
+		ProjInfo[0].Max = -Settings.MinDistance * ToSourceDistanceScale;
+		ProjInfo[1].Max = Settings.MaxDistance * ToSourceDistanceScale;
+	}
+	else
+	{
+		NumDirs = 1;
+		if (Settings.MinDistance < 0)
+		{
+			ProjInfo[0].Sign = -1;
+			ProjInfo[0].Offset = -Settings.MaxDistance * ToSourceDistanceScale;
+			ProjInfo[0].Max = -Settings.MinDistance * ToSourceDistanceScale;
+		}
+		else
+		{
+			ProjInfo[0].Sign = 1;
+			ProjInfo[0].Offset = Settings.MinDistance * ToSourceDistanceScale;
+			ProjInfo[0].Max = Settings.MaxDistance * ToSourceDistanceScale;
+		}
+	}
+
+	if (Settings.bResetUVsForUnmatched)
+	{
+		// Reset up-front if we are resetting for un-matched triangles (otherwise we rely on the below UVOverlay->SetTriangle calls to clear unused elements as needed)
+		UVOverlay->ClearElements(Triangles);
+	}
+
+	TMap<int32, int32> BaseToOverlayVIDMap;
+	TArray<int32> NewUVIndices;
+
+	IMeshSpatial::FQueryOptions QueryOptions;
+	QueryOptions.TriangleFilterF = Settings.SourceMeshTriFilter;
+
+
+	for (int32 TID : Triangles)
+	{
+		FIndex3i BaseTri = Mesh->GetTriangle(TID);
+		FIndex3i ElemTri = FIndex3i::Invalid();
+		FVector2f FoundUVs[3];
+		int32 NumFound = 0;
+
+		for (int32 j = 0; j < 3; ++j)
+		{
+			const int32* FoundElementID = BaseToOverlayVIDMap.Find(BaseTri[j]);
+			bool bFound = FoundElementID != nullptr;
+			if (!bFound)
+			{
+				int32 BestTID = INDEX_NONE;
+				double BestT = FMathd::MaxReal;
+				FVector3f BestBary;
+				for (int32 ProjectionIdx = 0; ProjectionIdx < NumDirs; ++ProjectionIdx)
+				{
+					FVector3d Pos = Mesh->GetVertex(BaseTri[j]);
+					FProjectionInfo& Proj = ProjInfo[ProjectionIdx];
+					FVector3d TransformPosStart = TransformTargetToSourceSpacePosition(Pos) + ConstantProjOffset + SourceProjDirection * Proj.Offset * Proj.Sign;
+					FRay3d Ray(TransformPosStart, SourceProjDirection * Proj.Sign);
+					double SearchDist = Proj.Max;
+					
+					double NearestT;
+					int32 NearestTID = INDEX_NONE;
+					FVector3d BaryCoords;
+					QueryOptions.MaxDistance = SearchDist;
+					if (SourceMeshSpatial.FindNearestHitTriangle(Ray, NearestT, NearestTID, BaryCoords, QueryOptions))
+					{
+						if (BestTID == INDEX_NONE || NearestT < BestT)
+						{
+							BestT = NearestT;
+							BestTID = NearestTID;
+							BestBary = (FVector3f)BaryCoords;
+						}
+					}
+				}
+				if (BestTID != INDEX_NONE && SourceUVs->IsSetTriangle(BestTID))
+				{
+					SourceUVs->GetTriBaryInterpolate<float>(BestTID, &BestBary.X, &FoundUVs[j].X);
+					NumFound++;
+				}
+			}
+			else
+			{
+				ElemTri[j] = *FoundElementID;
+				NumFound += bool(ElemTri[j] >= 0);
+			}
+		}
+		// If we have mapped a UV coordinate for every vertex, then append any new UV elements and set the UV overlay triangle
+		if (NumFound == 3)
+		{
+			for (int32 SubIdx = 0; SubIdx < 3; ++SubIdx)
+			{
+				if (ElemTri[SubIdx] == INDEX_NONE)
+				{
+					ElemTri[SubIdx] = UVOverlay->AppendElement(FoundUVs[SubIdx]);
+					NewUVIndices.Add(ElemTri[SubIdx]);
+					BaseToOverlayVIDMap.Add(BaseTri[SubIdx], ElemTri[SubIdx]);
+				}
+			}
+			UVOverlay->SetTriangle(TID, ElemTri);
+		}
+	}
+
+	if (Result != nullptr)
+	{
+		Result->NewUVElements = MoveTemp(NewUVIndices);
+	}
+
+}
+
 
 bool FDynamicMeshUVEditor::EstimateGeodesicCenterFrameVertex(const FDynamicMesh3& Mesh, FFrame3d& FrameOut, int32& VertexIDOut, bool bAlignToUnitAxes)
 {
diff --git a/Engine/Plugins/Runtime/GeometryProcessing/Source/DynamicMesh/Public/Parameterization/DynamicMeshUVEditor.h b/Engine/Plugins/Runtime/GeometryProcessing/Source/DynamicMesh/Public/Parameterization/DynamicMeshUVEditor.h
--- a/Engine/Plugins/Runtime/GeometryProcessing/Source/DynamicMesh/Public/Parameterization/DynamicMeshUVEditor.h
+++ b/Engine/Plugins/Runtime/GeometryProcessing/Source/DynamicMesh/Public/Parameterization/DynamicMeshUVEditor.h
@@ -4,6 +4,7 @@
 
 #include "DynamicMesh/DynamicMesh3.h"
 #include "DynamicMesh/DynamicMeshAttributeSet.h"
+#include "DynamicMesh/DynamicMeshAABBTree3.h"
 
 namespace UE
 {
@@ -257,6 +258,36 @@ public:
 	void SetTriangleUVsFromCylinderProjection(const TArray<int32>& Triangles, TFunctionRef<FVector3d(const FVector3d&)> PointTransform, const FFrame3d& CylFrame, const FVector3d& CylDimensions, float CylinderSplitAngle, FUVEditResult* Result = nullptr);
 
 
+	struct FTransferFromMeshViaProjectionSettings
+	{
+		// Range to search for source triangles
+		float MinDistance = -FMathf::MaxReal, MaxDistance = FMathf::MaxReal;
+
+		// For triangles without matching source triangles, whether reset the UVs -- if false, they will be left as-is
+		bool bResetUVsForUnmatched = false;
+
+		// If a filter is provided, will only transfer UVs from triangles where the filter returns 'true'
+		TFunction<bool(int32)> SourceMeshTriFilter = nullptr;
+	};
+
+	/**
+	 * Transfer UVs from a source mesh by projecting along the specified direction
+	 * @param Triangles Triangles to update with new UVs
+	 * @param TransformTargetToSourceSpacePosition Transform to take target mesh vertices to the space of the source mesh
+	 * @param ProjectionDirection Direction to project, in local space of TargetMesh
+	 * @param ProjectionOffset Distance to offset from TargetMesh vertices before projecting
+	 * @param TransformProjectionToSourceSpaceVector Transform to take the ProjectionDirection to the source mesh space
+	 * @param SourceMeshSpatial AABB Tree of source mesh
+	 * @param SourceMeshUVChannel UV Channel to read from on source mesh
+	 * @param Settings Projection settings
+	 * @param Result If non-null, contains new UV elements IDs created on target mesh
+	 */
+	void TransferTriangleUVsFromMeshViaDirectionProjection(TArrayView<const int32> Triangles, TFunctionRef<FVector3d(const FVector3d&)> TransformTargetToSourceSpacePosition, 
+		const FVector3d& ProjectionDirection, double ProjectionOffset, TFunctionRef<FVector3d(const FVector3d&)> TransformProjectionToSourceSpaceVector,
+		const FDynamicMeshAABBTree3& SourceMeshSpatial, int32 SourceMeshUVChannel,
+		const FTransferFromMeshViaProjectionSettings& Settings, FUVEditResult* Result = nullptr);
+
+
 	/**
 	* Compute the UV-space and 3D area of the given Triangles, and then scale the UV area to be equvalent
 	* to the 3D area. Scaling is around the center of the UV bounding box.
diff --git a/Engine/Plugins/Runtime/GeometryScripting/Source/GeometryScriptingCore/Private/MeshUVFunctions.cpp b/Engine/Plugins/Runtime/GeometryScripting/Source/GeometryScriptingCore/Private/MeshUVFunctions.cpp
--- a/Engine/Plugins/Runtime/GeometryScripting/Source/GeometryScriptingCore/Private/MeshUVFunctions.cpp
+++ b/Engine/Plugins/Runtime/GeometryScripting/Source/GeometryScriptingCore/Private/MeshUVFunctions.cpp
@@ -631,6 +631,105 @@ UDynamicMesh* UGeometryScriptLibrary_MeshUVFunctions::SetMeshUVsFromCylinderProj
 	return TargetMesh;
 }
 
+UDynamicMesh* UGeometryScriptLibrary_MeshUVFunctions::TransferMeshUVsByProjection(
+	UDynamicMesh* TargetMesh,
+	int TargetUVChannel,
+	FGeometryScriptMeshSelection TargetSelection,
+	FTransform TargetTransform,
+	const UDynamicMesh* SourceMesh,
+	FGeometryScriptDynamicMeshBVH SourceMeshOptionalBVH,
+	int SourceUVChannel,
+	FGeometryScriptMeshSelection SourceSelection,
+	FTransform SourceTransform,
+	FGeometryScriptMeshProjectionSettings Settings,
+	FVector ProjectionDirection,
+	double ProjectionOffset,
+	UGeometryScriptDebug* Debug)
+{
+	if (TargetMesh == nullptr)
+	{
+		UE::Geometry::AppendError(Debug, EGeometryScriptErrorType::InvalidInputs, LOCTEXT("SetMeshUVsFromMeshProjection_InvalidInput_Target", "SetMeshUVsFromMeshProjection: TargetMesh is Null"));
+		return TargetMesh;
+	}
+
+	if (SourceMesh == nullptr)
+	{
+		UE::Geometry::AppendError(Debug, EGeometryScriptErrorType::InvalidInputs, LOCTEXT("SetMeshUVsFromMeshProjection_InvalidInput_Source", "SetMeshUVsFromMeshProjection: SourceMesh is Null"));
+		return TargetMesh;
+	}
+
+	const FDynamicMeshUVOverlay* SourceUVLayer = nullptr;
+	SourceMesh->ProcessMesh([&SourceUVLayer, SourceUVChannel](const FDynamicMesh3& SourceMesh)
+	{
+		if (const FDynamicMeshAttributeSet* AttribSet = SourceMesh.Attributes())
+		{
+			SourceUVLayer = AttribSet->GetUVLayer(SourceUVChannel);
+		}
+	});
+	if (!SourceUVLayer)
+	{
+		UE::Geometry::AppendError(Debug, EGeometryScriptErrorType::InvalidInputs, LOCTEXT("SetMeshUVsFromMeshProjection_InvalidSourceUVSet", "SetMeshUVsFromMeshProjection: Source UV Channel does not exist on SourceMesh"));
+		return TargetMesh;
+	}
+
+	if (TargetSelection.IsEmpty() && !Settings.bProcessAllIfEmptySelection)
+	{
+		return TargetMesh;
+	}
+	if (SourceSelection.IsEmpty() && !Settings.bProcessAllIfEmptySelection)
+	{
+		return TargetMesh;
+	}
+
+	bool bHasUVSet = false;
+	ApplyMeshUVEditorOperation(TargetMesh, TargetUVChannel, bHasUVSet, Debug,
+		[SourceUVChannel, &SourceMesh, &SourceMeshOptionalBVH, &TargetSelection, &SourceSelection,
+		ProjectionDirection, ProjectionOffset, &Settings, TargetTransform, SourceTransform]
+		(FDynamicMesh3& EditMesh, FDynamicMeshUVOverlay* UVOverlay, FDynamicMeshUVEditor& UVEditor)
+	{
+		TArray<int32> TriangleROI;
+		TargetSelection.ProcessByTriangleID(EditMesh, [&TriangleROI](int32 TriangleID) { TriangleROI.Add(TriangleID); }, true);
+		FDynamicMeshUVEditor::FTransferFromMeshViaProjectionSettings ProjSettings;
+		ProjSettings.bResetUVsForUnmatched = Settings.bResetUVsForUnmatched;
+		ProjSettings.MaxDistance = Settings.ProjectionRangeMax;
+		ProjSettings.MinDistance = Settings.ProjectionRangeMin;
+		TSet<int32> SourceTriROI;
+		if (!SourceSelection.IsEmpty())
+		{
+			SourceSelection.ProcessByTriangleID(SourceMesh->GetMeshRef(), [&SourceTriROI](int32 TID) { SourceTriROI.Add(TID); });
+			ProjSettings.SourceMeshTriFilter = [&SourceTriROI](int32 TID)
+				{
+					return SourceTriROI.Contains(TID);
+				};
+		}
+		FDynamicMeshAABBTree3 SourceMeshSpatial(SourceMesh->GetMeshPtr(), false);
+		FDynamicMeshAABBTree3* UseSpatial = &SourceMeshSpatial;
+		if (SourceMeshOptionalBVH.Spatial.IsValid() && SourceMeshOptionalBVH.Spatial->GetMesh() == SourceMesh->GetMeshPtr())
+		{
+			UseSpatial = SourceMeshOptionalBVH.Spatial.Get();
+		}
+		else
+		{
+			SourceMeshSpatial.Build();
+		}
+		UVEditor.TransferTriangleUVsFromMeshViaDirectionProjection(TriangleROI, [TargetTransform, SourceTransform](const FVector3d& Pos) -> FVector3d
+			{
+				return SourceTransform.InverseTransformPosition(TargetTransform.TransformPosition(Pos));
+			},
+			ProjectionDirection, ProjectionOffset, [SourceTransform](const FVector3d& Vec) -> FVector3d
+			{
+				return SourceTransform.InverseTransformVector(Vec);
+			}, *UseSpatial, SourceUVChannel,
+			ProjSettings);
+	});
+	if (bHasUVSet == false)
+	{
+		UE::Geometry::AppendError(Debug, EGeometryScriptErrorType::InvalidInputs, LOCTEXT("SetMeshUVsFromMeshProjection_InvalidUVSet", "SetMeshUVsFromMeshProjection: Target UV Channel does not exist on TargetMesh"));
+	}
+
+	return TargetMesh;
+}
+
 UDynamicMesh*
 UGeometryScriptLibrary_MeshUVFunctions::ApplyTexelDensityUVScaling(
 	UDynamicMesh* TargetMesh,
diff --git a/Engine/Plugins/Runtime/GeometryScripting/Source/GeometryScriptingCore/Public/GeometryScript/MeshUVFunctions.h b/Engine/Plugins/Runtime/GeometryScripting/Source/GeometryScriptingCore/Public/GeometryScript/MeshUVFunctions.h
--- a/Engine/Plugins/Runtime/GeometryScripting/Source/GeometryScriptingCore/Public/GeometryScript/MeshUVFunctions.h
+++ b/Engine/Plugins/Runtime/GeometryScripting/Source/GeometryScriptingCore/Public/GeometryScript/MeshUVFunctions.h
@@ -227,6 +227,30 @@ public:
 	TMap<int32, int32> UDIMResolutions;
 };
 
+USTRUCT(BlueprintType)
+struct GEOMETRYSCRIPTINGCORE_API FGeometryScriptMeshProjectionSettings
+{
+	GENERATED_BODY()
+public:
+
+	// If selection contains no triangles, project all triangles (for both source and target)
+	UPROPERTY(EditAnywhere, BlueprintReadWrite, Category = Options)
+	bool bProcessAllIfEmptySelection = true;
+
+	// Maximum projection distance
+	UPROPERTY(EditAnywhere, BlueprintReadWrite, Category = Options)
+	float ProjectionRangeMax = UE_LARGE_WORLD_MAX;
+
+	// Minimum projection distance. If negative, will also consider projection backwards, and take the closest result.
+	UPROPERTY(EditAnywhere, BlueprintReadWrite, Category = Options)
+	float ProjectionRangeMin = -UE_LARGE_WORLD_MAX;
+	
+	// Whether to reset UVs for triangles where projection failed. Otherwise UVs are left as-is where projection failed.
+	UPROPERTY(EditAnywhere, BlueprintReadWrite, Category = Options)
+	bool bResetUVsForUnmatched = false;
+
+};
+
 
 UCLASS(meta = (ScriptName = "GeometryScript_UVs"))
 class GEOMETRYSCRIPTINGCORE_API UGeometryScriptLibrary_MeshUVFunctions : public UBlueprintFunctionLibrary
@@ -450,6 +474,42 @@ public:
 		float SplitAngle = 45.0,
 		UGeometryScriptDebug* Debug = nullptr );
 
+	/**
+	 * Copy UVs from one mesh to another, by projecting along the requested direction.
+	 * Note: This does not transfer UV seams; it assigns a single UV coordinate per vertex in the target mesh selection.
+	 * By default, also searches in -ProjectionDirection and picks the closest source mesh position to copy from --
+	 * set ProjectionRangeMin in Settings to a value >= 0 to only search in +ProjectionDirection.
+	 *
+	 * @param TargetMesh Mesh to assign new UVs
+	 * @param TargetUVChannel UV channel to update on target mesh
+	 * @param TargetSelection Triangles to update on the target mesh
+	 * @param TargetTransform Transform of target mesh
+	 * @param SourceMesh Mesh to transfer UVs from
+	 * @param SourceMeshOptionalBVH Optional BVH for the source mesh
+	 * @param SourceUVChannel UV channel to read from on the source mesh
+	 * @param SourceSelection Triangles to read from on the source mesh
+	 * @param SourceTransform Transform of the source mesh
+	 * @param Settings Additional settings
+	 * @param ProjectionDirection Direction to project (in the space where TargetMesh is transformed by TargetTransform, and SourceMesh is transformed by SourceTransform)
+	 * @param ProjectionOffset Projection will start offset by this amount from the TargetMesh vertices
+	 */
+	UFUNCTION(BlueprintCallable, Category = "GeometryScript|UVs", meta = (ScriptMethod, Keywords = "Set From Mesh"))
+	static UPARAM(DisplayName = "Target Mesh") UDynamicMesh*
+	TransferMeshUVsByProjection(
+		UDynamicMesh* TargetMesh,
+		UPARAM(DisplayName = "Target UV Channel") int TargetUVChannel,
+		FGeometryScriptMeshSelection TargetSelection,
+		FTransform TargetTransform,
+		const UDynamicMesh* SourceMesh,
+		FGeometryScriptDynamicMeshBVH SourceMeshOptionalBVH,
+		UPARAM(DisplayName = "Source UV Channel") int SourceUVChannel,
+		FGeometryScriptMeshSelection SourceSelection,
+		FTransform SourceTransform,
+		FGeometryScriptMeshProjectionSettings Settings,
+		FVector ProjectionDirection = FVector(0,0,-1),
+		double ProjectionOffset = 0,
+		UGeometryScriptDebug* Debug = nullptr);
+
 
 	/**
 	* Rescales UVs in the UV Channel for a Mesh to match a specified texel density, described by the options passed in. Supports
