add validation of symmetry after external mesh changes in vertex sculpt tool

+ add symmetry validation methods to the planar symmetry class, to support this

diff --git a/Engine/Plugins/Runtime/GeometryProcessing/Source/DynamicMesh/Private/Parameterization/MeshPlanarSymmetry.cpp b/Engine/Plugins/Runtime/GeometryProcessing/Source/DynamicMesh/Private/Parameterization/MeshPlanarSymmetry.cpp
--- a/Engine/Plugins/Runtime/GeometryProcessing/Source/DynamicMesh/Private/Parameterization/MeshPlanarSymmetry.cpp
+++ b/Engine/Plugins/Runtime/GeometryProcessing/Source/DynamicMesh/Private/Parameterization/MeshPlanarSymmetry.cpp
@@ -689,6 +689,100 @@ void FMeshPlanarSymmetry::ApplySymmetryPlaneConstraints(const TArray<int>& Verte
 }
 
 
+namespace MeshPlanarSymmetryLocals
+{
+	template<bool bSkipNonSourceVertices>
+	inline bool ValidateVertexSymmetry(const FDynamicMesh3& TestMesh, int32 Idx, 
+		const bool bOnPlane, const bool bIsSourceVertex, const int32 PairedVertex,
+		const FVector3d& Origin, const FVector3d& Normal, double OnPlaneTolerance, double MatchVertexToleranceSq)
+	{
+		FVector3d Pos = TestMesh.GetVertex(Idx);
+		if (bOnPlane)
+		{
+			const double SignedDistance = (Pos - Origin).Dot(Normal);
+			if (FMath::Abs(SignedDistance) > OnPlaneTolerance)
+			{
+				return false;
+			}
+		}
+		else 
+		{
+			if constexpr (bSkipNonSourceVertices)
+			{
+				if (!bIsSourceVertex)
+				{
+					return true;
+				}
+			}
+			const int32 PairedVertexIdx = PairedVertex;
+			if (!TestMesh.IsVertex(PairedVertexIdx))
+			{
+				return false;
+			}
+			FVector3d PairedPos = TestMesh.GetVertex(PairedVertexIdx);
+			const double SignedDistance = (Pos - Origin).Dot(Normal);
+			FVector3d MirrorPos = Pos - (2 * SignedDistance * Normal);
+			if (FVector3d::DistSquared(MirrorPos, PairedPos) > MatchVertexToleranceSq)
+			{
+				return false;
+			}
+		}
+		return true;
+	}
+}
+
+
+
+bool FMeshPlanarSymmetry::ValidateSymmetry(const FDynamicMesh3& TestMesh) const
+{
+	if (TestMesh.MaxVertexID() != Vertices.Num())
+	{
+		return false;
+	}
+
+	const double OnPlaneTolerance = OnPlaneToleranceFactor * ErrorScale;
+	const double MatchVertexToleranceSq = (MatchVertexToleranceFactor * ErrorScale) * (MatchVertexToleranceFactor * ErrorScale);
+
+	for (int32 Idx = 0; Idx < Vertices.Num(); ++Idx)
+	{
+		if (!TestMesh.IsVertex(Idx))
+		{
+			continue;
+		}
+		const FSymmetryVertex& SV = Vertices[Idx];
+		if (!MeshPlanarSymmetryLocals::ValidateVertexSymmetry<true>(TestMesh, Idx, SV.bOnPlane, SV.bIsSourceVertex, SV.PairedVertex, 
+			SymmetryFrame.Origin, CachedSymmetryAxis, OnPlaneTolerance, MatchVertexToleranceSq))
+		{
+			return false;
+		}
+	}
+	return true;
+}
+
+
+bool FMeshPlanarSymmetry::ValidateSymmetry(const FDynamicMesh3& TestMesh, TConstArrayView<int32> VertexROI) const
+{
+	const double OnPlaneTolerance = OnPlaneToleranceFactor * ErrorScale;
+	const double MatchVertexToleranceSq = (MatchVertexToleranceFactor * ErrorScale) * (MatchVertexToleranceFactor * ErrorScale);
+
+	for (int32 Idx : VertexROI)
+	{
+		// All ROI vertices must be valid
+		if (!TestMesh.IsVertex(Idx) || !Vertices.IsValidIndex(Idx))
+		{
+			return false;
+		}
+		const FSymmetryVertex& SV = Vertices[Idx];
+		if (!MeshPlanarSymmetryLocals::ValidateVertexSymmetry<false>(TestMesh, Idx, SV.bOnPlane, SV.bIsSourceVertex, SV.PairedVertex, 
+			SymmetryFrame.Origin, CachedSymmetryAxis, OnPlaneTolerance, MatchVertexToleranceSq))
+		{
+			return false;
+		}
+	}
+	return true;
+}
+
+
 void FMeshPlanarSymmetry::ComputeSymmetryConstrainedPositions(
 	const TArray<int>& SourceVertexROI,
 	const TArray<int>& MirrorVertexROI,
diff --git a/Engine/Plugins/Runtime/GeometryProcessing/Source/DynamicMesh/Public/Parameterization/MeshPlanarSymmetry.h b/Engine/Plugins/Runtime/GeometryProcessing/Source/DynamicMesh/Public/Parameterization/MeshPlanarSymmetry.h
--- a/Engine/Plugins/Runtime/GeometryProcessing/Source/DynamicMesh/Public/Parameterization/MeshPlanarSymmetry.h
+++ b/Engine/Plugins/Runtime/GeometryProcessing/Source/DynamicMesh/Public/Parameterization/MeshPlanarSymmetry.h
@@ -79,6 +79,11 @@ public:
 	void FullSymmetryUpdate();
 
 
+	/** @return true if the TestMesh vertices match the current tracked vertex symmetry correspondence (i.e., the TestMesh vertex IDs match the known correspondence, and have the expected mirror symmetry) */
+	bool ValidateSymmetry(const FDynamicMesh3& TestMesh) const;
+
+	/** @return true if the TestMesh vertices in VertexROI match the current tracked vertex symmetry correspondence (i.e., the vertex IDs match the known correspondence, and have the expected mirror symmetry) */
+	bool ValidateSymmetry(const FDynamicMesh3& TestMesh, TConstArrayView<int32> VertexROI) const;
 
 	//
 	// The set of functions below are intended to be used together in situations like 3D sculpting
diff --git a/Engine/Plugins/Runtime/MeshModelingToolset/Source/MeshModelingTools/Private/MeshVertexSculptTool.cpp b/Engine/Plugins/Runtime/MeshModelingToolset/Source/MeshModelingTools/Private/MeshVertexSculptTool.cpp
--- a/Engine/Plugins/Runtime/MeshModelingToolset/Source/MeshModelingTools/Private/MeshVertexSculptTool.cpp
+++ b/Engine/Plugins/Runtime/MeshModelingToolset/Source/MeshModelingTools/Private/MeshVertexSculptTool.cpp
@@ -458,6 +458,11 @@ void UMeshVertexSculptTool::Setup()
 
 	SymmetryProperties->WatchProperty(SymmetryProperties->bEnableSymmetry,
 		[this](bool bNewValue) { bApplySymmetry = bMeshSymmetryIsValid && bNewValue; });
+	SymmetryProperties->WatchProperty(SymmetryProperties->bSymmetryCanBeEnabled,
+		[this](bool bNewValue) 
+		{
+			bApplySymmetry = bMeshSymmetryIsValid && bNewValue && SymmetryProperties->bEnableSymmetry; 
+		});
 }
 
 void UMeshVertexSculptTool::RegisterBrushes()
@@ -589,7 +594,16 @@ void UMeshVertexSculptTool::UpdateToolMeshes(TFunctionRef<TUniquePtr<FMeshRegion
 			// Note this change should be in the context of a larger transaction, so the text isn't that important
 			GetToolManager()->EmitObjectChange(DynamicMeshComponent, MoveTemp(NewChange), LOCTEXT("UpdateVertexSculptMesh", "Updated Mesh"));
 			
-			// TODO: If bMeshSymmetryIsValid, re-check that symmetry is still valid after sculpt layer change
+			if (bMeshSymmetryIsValid)
+			{
+				// Re-validate that the symmetry still holds after the external mesh change
+				if (!Symmetry->ValidateSymmetry(*GetSculptMesh()))
+				{
+					GetToolManager()->EmitObjectChange(this, MakeUnique<FVertexSculptNonSymmetricChange>(), LOCTEXT("InvalidateSymmetryChange", "Invalidate Symmetry"));
+					bMeshSymmetryIsValid = false;
+					SymmetryProperties->bSymmetryCanBeEnabled = bMeshSymmetryIsValid;
+				}
+			}
 		}
 		// No change is ready to emit, just update component rendering
 		else
