SparseMatrix: guard against integer overflow when computing how many entries to allocate

diff --git a/Engine/Plugins/Runtime/GeometryProcessing/Source/DynamicMesh/Private/Solvers/Internal/SparseMatrix.cpp b/Engine/Plugins/Runtime/GeometryProcessing/Source/DynamicMesh/Private/Solvers/Internal/SparseMatrix.cpp
--- a/Engine/Plugins/Runtime/GeometryProcessing/Source/DynamicMesh/Private/Solvers/Internal/SparseMatrix.cpp
+++ b/Engine/Plugins/Runtime/GeometryProcessing/Source/DynamicMesh/Private/Solvers/Internal/SparseMatrix.cpp
@@ -71,11 +71,16 @@ bool SliceSparseMatrix(const FSparseMatrixD& InMatrix,
 
     // Take a guess at the number of the nonzero elements in the sliced matrix
     std::vector<Eigen::Triplet<FSparseMatrixD::Scalar>> Triplets;
-    const float NonZeroProbabilitytPerEntry = static_cast<float>(NumNonZeros) / static_cast<float>(NumRows) / static_cast<float>(NumCols);	// ~= NumNonZeros / (NumRows*NumCols) but 
+    const float NonZeroProbabilityPerEntry = static_cast<float>(NumNonZeros) / static_cast<float>(NumRows) / static_cast<float>(NumCols);	// ~= NumNonZeros / (NumRows*NumCols) but 
 																																			// won't hit int32 overflow on large systems
 
-    const float SlicedNumEntries = static_cast<float>(NumRowsSliced * NumColsSliced);
-    const int SlicedNonZeroEstimate = static_cast<int>(NonZeroProbabilitytPerEntry * SlicedNumEntries);
+    const float SlicedNumEntries = static_cast<float>(NumRowsSliced) * static_cast<float>(NumColsSliced);
+    const int SlicedNonZeroEstimate = static_cast<int>(NonZeroProbabilityPerEntry * SlicedNumEntries);
+	if (!ensure(SlicedNonZeroEstimate >= 0))
+	{
+		return false;
+	}
+
     Triplets.reserve(SlicedNonZeroEstimate);
 
     // Iterate over all of the non-zero elements and check if they are part of the new sliced matrix.
