ModelingTools: Remove crash in selection system when operating on volumes due to converted volumes not matching the tool view (selection will currently still sometimes be incorrect, but will not crash). Switch other checks in the file to ensures too.

diff --git a/Engine/Plugins/Runtime/GeometryProcessing/Source/DynamicMesh/Private/Selections/GeometrySelectionUtil.cpp b/Engine/Plugins/Runtime/GeometryProcessing/Source/DynamicMesh/Private/Selections/GeometrySelectionUtil.cpp
--- a/Engine/Plugins/Runtime/GeometryProcessing/Source/DynamicMesh/Private/Selections/GeometrySelectionUtil.cpp
+++ b/Engine/Plugins/Runtime/GeometryProcessing/Source/DynamicMesh/Private/Selections/GeometrySelectionUtil.cpp
@@ -268,7 +268,7 @@ void UE::Geometry::UpdateTriangleSelectionViaRaycast(
 	const FGeometrySelectionUpdateConfig& UpdateConfig,
 	FGeometrySelectionUpdateResult& ResultOut )
 {
-	check(Editor->GetTopologyType() == EGeometryTopologyType::Triangle);
+	ensure(Editor->GetTopologyType() == EGeometryTopologyType::Triangle);
 
 	ResultOut.bSelectionMissed = true;
 
@@ -349,7 +349,7 @@ void UE::Geometry::UpdateGroupSelectionViaRaycast(
 	const FGeometrySelectionUpdateConfig& UpdateConfig,
 	FGeometrySelectionUpdateResult& ResultOut)
 {
-	check(Editor->GetTopologyType() == EGeometryTopologyType::Polygroup);
+	ensure(Editor->GetTopologyType() == EGeometryTopologyType::Polygroup);
 
 	ResultOut.bSelectionMissed = true;
 
@@ -743,7 +743,14 @@ bool UE::Geometry::EnumeratePolygroupSelectionTriangles(
 			if (Mesh.IsTriangle(SeedTriangleID))
 			{
 				const int32 GroupID = GroupSet.GetGroup(SeedTriangleID);
-				check(GroupID == (int32)SelectionID.TopologyID);		// sanity-check that we are using the right group
+				// TODO: [TopologyMismatch] We don't have the ensure below because the selection system currently can have a
+				//  different view of a converted target than a tool might after making some edits, and the group ID's can
+				//  differ. For instance, a modified volume dynamic mesh may differ from the dynamic mesh that would result
+				//  from converting to a volume and back to dynamic mesh again. Someday we may fix this, but for now we accept
+				//  that the selection system and tools may disagree, and tools may end up setting incorrect selections after
+				//  doing some edits.
+				// ensure(GroupID == (int32)SelectionID.TopologyID);		// sanity-check that we are using the right group
+
 				if ( GroupID >= 0 && UniqueSeedGroups.Contains(GroupID) == false)
 				{
 					UniqueSeedGroups.Add(GroupID);
@@ -799,7 +806,7 @@ bool UE::Geometry::EnumeratePolygroupSelectionTriangles(
 	const int32 NumGroups = SeedGroups.Num();
 	for (int32 k = 0; k < NumGroups; ++k)
 	{
-		check(GroupSet.GetGroup(SeedTriangles[k]) == SeedGroups[k]);
+		ensure(GroupSet.GetGroup(SeedTriangles[k]) == SeedGroups[k]);
 		const int32 GroupID = SeedGroups[k];
 		FMeshConnectedComponents::GrowToConnectedTriangles(&Mesh, 
 			TArray<int>{SeedTriangles[k]}, TempROI, &QueueBuffer, 
@@ -924,7 +931,10 @@ bool UE::Geometry::EnumeratePolygroupSelectionEdges(
 			if (Mesh.IsTriangle(SeedTriangleID))
 			{
 				const int32 GroupID = GroupSet.GetGroup(SeedTriangleID);
-				check(GroupID == (int32)SelectionID.TopologyID);		// sanity-check that we are using the right group
+
+				// See TODO: [TopologyMismatch] above
+				//ensure(GroupID == (int32)SelectionID.TopologyID);		// sanity-check that we are using the right group
+				
 				if ( GroupID >= 0 && UniqueSeedTriGroups.Contains(GroupID) == false)
 				{
 					UniqueSeedTriGroups.Add(GroupID);
@@ -1034,7 +1044,7 @@ bool UE::Geometry::EnumeratePolygroupSelectionEdges(
 		const int32 NumGroups = SeedTriGroups.Num();
 		for (int32 k = 0; k < NumGroups; ++k)
 		{
-			check(GroupSet.GetGroup(SeedTriangles[k]) == SeedTriGroups[k]);
+			ensure(GroupSet.GetGroup(SeedTriangles[k]) == SeedTriGroups[k]);
 			const int32 GroupID = SeedTriGroups[k];
 			FMeshConnectedComponents::GrowToConnectedTriangles(&Mesh, 
 				TArray<int>{SeedTriangles[k]}, TempROI, &QueueBuffer, 
@@ -1087,7 +1097,10 @@ bool UE::Geometry::EnumeratePolygroupSelectionEdges(
 			if (Mesh.IsTriangle(SeedTriangleID))
 			{
 				const int32 GroupID = GroupTopology.GetGroupID(SeedTriangleID);
-				check(GroupID == (int32)SelectionID.TopologyID);		// sanity-check that we are using the right group
+
+				// See TODO: [TopologyMismatch] above
+				//ensure(GroupID == (int32)SelectionID.TopologyID);		// sanity-check that we are using the right group
+				
 				if ( GroupID >= 0)
 				{
 					SeedGroups.Add(GroupID);
