add initial mesh isocurves algorithm for cutting a mesh by marching triangles on some function at the vertices

+ move collapse degenerate edges algorithm logic to a static function on the local planar simplify class, so it can be called from the mesh isocurves cut method

+ expose isocurves method in geometry script

diff --git a/Engine/Plugins/Runtime/GeometryProcessing/Source/DynamicMesh/Private/Operations/MeshIsoCurves.cpp b/Engine/Plugins/Runtime/GeometryProcessing/Source/DynamicMesh/Private/Operations/MeshIsoCurves.cpp
new file mode 100644
--- /dev/null
+++ b/Engine/Plugins/Runtime/GeometryProcessing/Source/DynamicMesh/Private/Operations/MeshIsoCurves.cpp
@@ -0,0 +1,141 @@
+// Copyright Epic Games, Inc. All Rights Reserved.
+
+#include "Operations/MeshIsoCurves.h"
+
+#include "Async/ParallelFor.h"
+#include "DynamicMesh/DynamicMesh3.h"
+#include "Operations/LocalPlanarSimplify.h"
+
+using namespace UE::Geometry;
+
+void FMeshIsoCurves::Cut(FDynamicMesh3& Mesh, TFunctionRef<float(int32)> VertexFn, TFunctionRef<float(int32 VIDMin, int32 VIDMax, float ValMin, float ValMax)> EdgeCutFn, float IsoValue)
+{
+	int MaxVID = Mesh.MaxVertexID();
+	TArray<float> VertexValues;
+	VertexValues.SetNumUninitialized(MaxVID);
+
+	constexpr bool bNoParallel = false;
+	ParallelFor(MaxVID, [this, &VertexValues, &Mesh, &VertexFn, IsoValue](int32 VID)
+	{
+		if (Mesh.IsVertex(VID))
+		{
+			VertexValues[VID] = VertexFn(VID);
+		}
+		else
+		{
+			// set IsoValue on invalid vertices; any vertex that is later inserted and uses this ID will be on the curve, so should have this value
+			VertexValues[VID] = IsoValue;
+		}
+	}, bNoParallel);
+
+	TSet<int> OnCutEdges;
+	SplitCrossingEdges(Mesh, VertexValues, OnCutEdges, EdgeCutFn, IsoValue);
+
+	// collapse degenerate edges
+	if (Settings.bCollapseDegenerateEdgesOnCut)
+	{
+		FLocalPlanarSimplify::CollapseDegenerateEdges(Mesh, OnCutEdges, false, Settings.DegenerateEdgeTol);
+	}
+
+}
+
+void FMeshIsoCurves::SplitCrossingEdges(FDynamicMesh3& Mesh, const TArray<float>& VertexValues,
+	TSet<int32>& OnCutEdges,
+	TFunctionRef<float(int32 VIDMin, int32 VIDMax, float ValMin, float ValMax)> EdgeCutFn,
+	float IsoValue)
+{
+	OnCutEdges.Reset();
+
+	// have to skip processing of new edges. If edge id
+	// is > max at start, is new. Otherwise if in NewEdges list, also new.
+	int MaxEID = Mesh.MaxEdgeID();
+	TSet<int> NewEdgesBeforeMaxID;
+	auto AddNewEdge = [&NewEdgesBeforeMaxID, MaxEID](int32 NewEID)
+	{
+		if (NewEID < MaxEID)
+		{
+			NewEdgesBeforeMaxID.Add(NewEID);
+		}
+	};
+
+	double UseSnapVertTol = FMath::Max(0, Settings.SnapToExistingVertexTol);
+	const double SnapExistingTolSq = UseSnapVertTol * UseSnapVertTol;
+
+	// split existing edges where the value crosses isovalue
+	for (int32 EID = 0; EID < MaxEID; ++EID)
+	{
+		if (!Mesh.IsEdge(EID) || NewEdgesBeforeMaxID.Contains(EID))
+		{
+			continue;
+		}
+
+		FIndex2i EdgeV = Mesh.GetEdgeV(EID);
+		const float ValueA = VertexValues[EdgeV.A];
+		const float ValueB = VertexValues[EdgeV.B];
+		const float DistA = ValueA - IsoValue;
+		const float DistB = ValueB - IsoValue;
+
+		// If both Signs are 0, this edge is on-contour
+		// If one sign is 0, that vertex is on-contour
+		int AOnCurve = (FMathd::Abs(DistA) <= Settings.CurveIsoValueSnapTolerance) ? 1 : 0;
+		int BOnCurve = (FMathd::Abs(DistB) <= Settings.CurveIsoValueSnapTolerance) ? 1 : 0;
+		if (AOnCurve || BOnCurve)
+		{
+			continue;
+		}
+
+		// no crossing
+		if (DistA * DistB >= 0)
+		{
+			continue;
+		}
+
+		double Param = EdgeCutFn(EdgeV.A, EdgeV.B, ValueA, ValueB);
+		// Cut must be within edge
+		if (Param <= 0 || Param >= 1)
+		{
+			continue;
+		}
+		// Skip the edge split if we're within tolerance of an existing vertex
+		if (SnapExistingTolSq > 0)
+		{
+			FVector3d PosA = Mesh.GetVertex(EdgeV.A);
+			FVector3d PosB = Mesh.GetVertex(EdgeV.B);
+			FVector3d EdgeVec = PosB - PosA;
+			double MinSepSq = (EdgeVec * (Param > .5 ? 1 - Param : Param)).SquaredLength();
+			if (MinSepSq <= SnapExistingTolSq)
+			{
+				continue;
+			}
+		}
+
+		FDynamicMesh3::FEdgeSplitInfo SplitInfo;
+		EMeshResult SplitResult = Mesh.SplitEdge(EID, SplitInfo, Param);
+		if (!ensureMsgf(SplitResult == EMeshResult::Ok, TEXT("FMeshIsoCurves::SplitCrossingEdges: failed to SplitEdge")))
+		{
+			continue; // edge split really shouldn't fail; skip the edge if it somehow does
+		}
+
+		AddNewEdge(SplitInfo.NewEdges.A);
+		AddNewEdge(SplitInfo.NewEdges.B);
+
+		// We need to check whether the other vertices are on curve to decide if the connected edges are on the curve or not
+		int32 OtherVIDA = SplitInfo.OtherVertices.A;
+		// Other vertex is on curve if it's newly created or within curve isovalue tolerance
+		// (Note a newly-created vertex w/ ID < VertexValues.Num() will also have a VertexValue of IsoValue, since we use this as the default value)
+		if (OtherVIDA >= VertexValues.Num() || FMath::Abs(VertexValues[OtherVIDA] - IsoValue) <= Settings.CurveIsoValueSnapTolerance)
+		{
+			OnCutEdges.Add(SplitInfo.NewEdges.B);
+		}
+		
+		if (SplitInfo.NewEdges.C != FDynamicMesh3::InvalidID)
+		{
+			AddNewEdge(SplitInfo.NewEdges.C);
+			int32 OtherVIDB = SplitInfo.OtherVertices.B;
+			if (OtherVIDB >= VertexValues.Num() || FMath::Abs(VertexValues[OtherVIDB] - IsoValue) <= Settings.CurveIsoValueSnapTolerance)
+			{
+				OnCutEdges.Add(SplitInfo.NewEdges.C);
+			}
+		}
+	}
+}
diff --git a/Engine/Plugins/Runtime/GeometryProcessing/Source/DynamicMesh/Private/Operations/MeshPlaneCut.cpp b/Engine/Plugins/Runtime/GeometryProcessing/Source/DynamicMesh/Private/Operations/MeshPlaneCut.cpp
--- a/Engine/Plugins/Runtime/GeometryProcessing/Source/DynamicMesh/Private/Operations/MeshPlaneCut.cpp
+++ b/Engine/Plugins/Runtime/GeometryProcessing/Source/DynamicMesh/Private/Operations/MeshPlaneCut.cpp
@@ -649,68 +649,7 @@ PRAGMA_ENABLE_DEPRECATION_WARNINGS
 
 void FMeshPlaneCut::CollapseDegenerateEdges(TSet<int32>& Edges, bool bRemoveAllDegenerateFromInputSet, TSet<int>* TriangleSelection)
 {
-	double Tol2 = DegenerateEdgeTol * DegenerateEdgeTol;
-	FVector3d A, B;
-	TArray<int32> ToRemove;
-	int Collapsed = 0;
-	do
-	{
-		Collapsed = 0;
-		for (auto EdgesIter = Edges.CreateIterator(); EdgesIter; ++EdgesIter)
-		{
-			int32 EID = *EdgesIter;
-			if (!Mesh->IsEdge(EID))
-			{
-				continue;
-			}
-			Mesh->GetEdgeV(EID, A, B);
-			if (DistanceSquared(A, B) > Tol2)
-			{
-				continue;
-			}
-
-			FIndex2i EV = Mesh->GetEdgeV(EID);
-			// if the vertex we'd remove is on a seam, try removing the other one instead
-			if (Mesh->HasAttributes() && Mesh->Attributes()->IsSeamVertex(EV.B, false))
-			{
-				Swap(EV.A, EV.B);
-				// if they were both on seams, then collapse should not happen?  (& would break OnCollapseEdge assumptions in overlay)
-				if (Mesh->HasAttributes() && Mesh->Attributes()->IsSeamVertex(EV.B, false))
-				{
-					continue;
-				}
-			}
-			FDynamicMesh3::FEdgeCollapseInfo CollapseInfo;
-			EMeshResult Result = Mesh->CollapseEdge(EV.A, EV.B, CollapseInfo);
-			if (Result == EMeshResult::Ok)
-			{
-				Collapsed++;
-				EdgesIter.RemoveCurrent();
-				if (TriangleSelection)
-				{
-					TriangleSelection->Remove(CollapseInfo.RemovedEdges.A);
-					if (CollapseInfo.RemovedEdges.B != FDynamicMesh3::InvalidID)
-					{
-						TriangleSelection->Remove(CollapseInfo.RemovedEdges.B);
-					}
-				}
-				// Add the other edge for later removal, so we don't potentially invalidate our iterator
-				// (but only if we really need to remove invalid edges here; some algorithms may be ok with leaving invalid EIDs in the set)
-				if (bRemoveAllDegenerateFromInputSet && CollapseInfo.RemovedEdges.B != FDynamicMesh3::InvalidID)
-				{
-					ToRemove.Add(CollapseInfo.RemovedEdges.B);
-				}
-			}
-		}
-	} while (Collapsed != 0);
-
-	if (bRemoveAllDegenerateFromInputSet)
-	{
-		for (int32 EID : ToRemove)
-		{
-			Edges.Remove(EID);
-		}
-	}
+	FLocalPlanarSimplify::CollapseDegenerateEdges(*Mesh, Edges, bRemoveAllDegenerateFromInputSet, DegenerateEdgeTol, TriangleSelection);
 }
 
 bool FMeshPlaneCut::SimpleHoleFill(int ConstantGroupID)
diff --git a/Engine/Plugins/Runtime/GeometryProcessing/Source/DynamicMesh/Public/Operations/MeshIsoCurves.h b/Engine/Plugins/Runtime/GeometryProcessing/Source/DynamicMesh/Public/Operations/MeshIsoCurves.h
new file mode 100644
--- /dev/null
+++ b/Engine/Plugins/Runtime/GeometryProcessing/Source/DynamicMesh/Public/Operations/MeshIsoCurves.h
@@ -0,0 +1,77 @@
+// Copyright Epic Games, Inc. All Rights Reserved.
+
+#pragma once
+
+#include "MathUtil.h"
+#include "VectorTypes.h"
+#include "GeometryTypes.h"
+
+namespace UE
+{
+namespace Geometry
+{
+
+class FDynamicMesh3;
+
+struct FMeshIsoCurveSettings
+{
+	// Whether to collapse any degenerate edges created by the curve insertion
+	bool bCollapseDegenerateEdgesOnCut = true;
+
+	// New edges shorter than this will be considered degenerate, and collapsed if bCollapseDegenerateEdgesOnCut is true
+	double DegenerateEdgeTol = FMathd::ZeroTolerance;
+
+	// Distance at which to snap curve vertices to nearby existing vertices
+	double SnapToExistingVertexTol = UE_KINDA_SMALL_NUMBER;
+
+	// Tolerance distance (in function domain) to an existing vertex to be 'on curve'
+	float CurveIsoValueSnapTolerance = 0.f;
+};
+
+/**
+ * Insert edges on a mesh along the isocurve where some scalar value function over the mesh surface crosses a specified value
+ */
+class DYNAMICMESH_API FMeshIsoCurves
+{
+public:
+
+	// Input options
+
+	FMeshIsoCurveSettings Settings;
+
+	/**
+	 * Insert new edges on the given mesh along the curve where a function over the mesh surface crosses a given isovalue
+	 * 
+	 * @param Mesh The mesh to cut
+	 * @param VertexFn Function from vertex ID to values
+	 * @param EdgeCutFn Given the vertices of an edge and their values, return the parameter where the edge should be cut. Only called if IsoValue is crossed between ValueA and ValueB.
+	 * @param IsoValue Value at which to insert a new curve on the mesh
+	 */
+	void Cut(FDynamicMesh3& Mesh, TFunctionRef<float(int32)> VertexFn, TFunctionRef<float(int32 VertA, int32 VertB, float ValueA, float ValueB)> EdgeCutFn, float IsoValue = 0);
+
+	/**
+	 * Insert new edges on the given mesh along the curve where a function over the mesh vertices, and linearly interpolated over edges, crosses a given isovalue
+	 *
+	 * @param Mesh The mesh to cut
+	 * @param VertexFn Function from vertex ID to values
+	 * @param IsoValue Value at which to insert a new curve on the mesh
+	 */
+	void Cut(FDynamicMesh3& Mesh, TFunctionRef<float(int32)> VertexFn, float IsoValue = 0)
+	{
+		Cut(Mesh, VertexFn, [IsoValue](int32, int32, float ValueA, float ValueB)
+			{
+				// Note this is only called on crossing edges, where ValueA != ValueB, so there is no divide-by-zero risk here
+				return (ValueA - IsoValue) / (ValueA - ValueB);
+			}, IsoValue
+		);
+	}
+
+private:
+	void SplitCrossingEdges(FDynamicMesh3& Mesh, const TArray<float>& VertexValues, TSet<int32>& OnCutEdges,
+		TFunctionRef<float(int32 VIDMin, int32 VIDMax, float ValMin, float ValMax)> EdgeCutFn,
+		float IsoValue);
+};
+
+
+} // end namespace UE::Geometry
+} // end namespace UE
\ No newline at end of file
diff --git a/Engine/Plugins/Runtime/GeometryScripting/Source/GeometryScriptingCore/Private/MeshBooleanFunctions.cpp b/Engine/Plugins/Runtime/GeometryScripting/Source/GeometryScriptingCore/Private/MeshBooleanFunctions.cpp
--- a/Engine/Plugins/Runtime/GeometryScripting/Source/GeometryScriptingCore/Private/MeshBooleanFunctions.cpp
+++ b/Engine/Plugins/Runtime/GeometryScripting/Source/GeometryScriptingCore/Private/MeshBooleanFunctions.cpp
@@ -7,9 +7,10 @@
 #include "UDynamicMesh.h"
 
 #include "Operations/MeshBoolean.h"
-#include "Operations/MeshSelfUnion.h"
+#include "Operations/MeshIsoCurves.h"
 #include "Operations/MeshPlaneCut.h"
 #include "Operations/MeshMirror.h"
+#include "Operations/MeshSelfUnion.h"
 #include "MeshSimplification.h"
 #include "MeshBoundaryLoops.h"
 #include "Operations/MinimalHoleFiller.h"
@@ -352,6 +353,43 @@ UDynamicMesh* UGeometryScriptLibrary_MeshBooleanFunctions::ApplyMeshMirror(
 }
 
 
+UDynamicMesh* UGeometryScriptLibrary_MeshBooleanFunctions::ApplyMeshIsoCurves(
+	UDynamicMesh* TargetMesh,
+	FGeometryScriptMeshIsoCurveOptions Options,
+	const FGeometryScriptScalarList& PerVertexValues,
+	float IsoValue,
+	UGeometryScriptDebug* Debug)
+{
+	if (TargetMesh == nullptr)
+	{
+		UE::Geometry::AppendError(Debug, EGeometryScriptErrorType::InvalidInputs, LOCTEXT("ApplyMeshIsoCurves_InvalidInput", "ApplyMeshIsoCurves: TargetMesh is Null"));
+		return TargetMesh;
+	}
+	if (!PerVertexValues.List.IsValid())
+	{
+		UE::Geometry::AppendError(Debug, EGeometryScriptErrorType::InvalidInputs, LOCTEXT("ApplyMeshIsoCurves_InvalidInput2", "ApplyMeshIsoCurves: Per Vertex Values list is Null"));
+	}
+
+	TargetMesh->EditMesh([&Options, &PerVertexValues, IsoValue, Debug](FDynamicMesh3& EditMesh)
+	{
+		if (PerVertexValues.List->Num() != EditMesh.MaxVertexID())
+		{
+			UE::Geometry::AppendError(Debug, EGeometryScriptErrorType::InvalidInputs, LOCTEXT("ApplyMeshIsoCurves_ListWrongSize", "ApplyMeshIsoCurves: Per Vertex Values list has wrong length -- must match vertices"));
+		}
+		FMeshIsoCurves IsoCurves;
+		IsoCurves.Settings.SnapToExistingVertexTol = Options.SnapToExistingVertexTolerance;
+		IsoCurves.Settings.CurveIsoValueSnapTolerance = Options.CurveIsoValueSnapTolerance;
+
+		IsoCurves.Cut(EditMesh, [&PerVertexValues](int32 VID) -> float
+			{
+				return (float)(*PerVertexValues.List)[VID];
+			}, 
+			IsoValue
+		);
+	}, EDynamicMeshChangeType::GeneralEdit, EDynamicMeshAttributeChangeFlags::Unknown, false);
+
+	return TargetMesh;
+}
 
 
 #undef LOCTEXT_NAMESPACE
diff --git a/Engine/Plugins/Runtime/GeometryScripting/Source/GeometryScriptingCore/Public/GeometryScript/MeshBooleanFunctions.h b/Engine/Plugins/Runtime/GeometryScripting/Source/GeometryScriptingCore/Public/GeometryScript/MeshBooleanFunctions.h
--- a/Engine/Plugins/Runtime/GeometryScripting/Source/GeometryScriptingCore/Public/GeometryScript/MeshBooleanFunctions.h
+++ b/Engine/Plugins/Runtime/GeometryScripting/Source/GeometryScriptingCore/Public/GeometryScript/MeshBooleanFunctions.h
@@ -142,7 +142,19 @@ public:
 	bool bWeldAlongPlane = true;
 };
 
+USTRUCT(BlueprintType)
+struct GEOMETRYSCRIPTINGCORE_API FGeometryScriptMeshIsoCurveOptions
+{
+	GENERATED_BODY()
+public:
+
+	UPROPERTY(EditAnywhere, BlueprintReadWrite, Category = Options)
+	float CurveIsoValueSnapTolerance = 0.f;
 
+	UPROPERTY(EditAnywhere, BlueprintReadWrite, Category = Options, meta = (Units = "cm"))
+	float SnapToExistingVertexTolerance = UE_KINDA_SMALL_NUMBER;
+
+};
 
 UCLASS(meta = (ScriptName = "GeometryScript_MeshBooleans"))
 class GEOMETRYSCRIPTINGCORE_API UGeometryScriptLibrary_MeshBooleanFunctions : public UBlueprintFunctionLibrary
@@ -223,4 +235,21 @@ public:
 		FGeometryScriptMeshMirrorOptions Options,
 		UGeometryScriptDebug* Debug = nullptr);
 
+	/**
+	* Insert edges along the curve where the interpolated per-vertex values match a given isovalue
+	*
+	* @param TargetMesh Dynamic Mesh to be updated.
+	* @param Options Settings for processing to perform after inserting the curve edges
+	* @param PerVertexValue Values to interpolate, must be one per vertex
+	* @param IsoValue Value at which to insert the new curve
+	*/
+	UFUNCTION(BlueprintCallable, Category = "GeometryScript|Booleans", meta = (ScriptMethod))
+	static UPARAM(DisplayName = "Target Mesh") UDynamicMesh*
+	ApplyMeshIsoCurves(
+		UDynamicMesh* TargetMesh,
+		FGeometryScriptMeshIsoCurveOptions Options,
+		const FGeometryScriptScalarList& PerVertexValue,
+		float IsoValue = .5,
+		UGeometryScriptDebug* Debug = nullptr);
+
 };
diff --git a/Engine/Source/Runtime/GeometryCore/Private/Operations/LocalPlanarSimplify.cpp b/Engine/Source/Runtime/GeometryCore/Private/Operations/LocalPlanarSimplify.cpp
--- a/Engine/Source/Runtime/GeometryCore/Private/Operations/LocalPlanarSimplify.cpp
+++ b/Engine/Source/Runtime/GeometryCore/Private/Operations/LocalPlanarSimplify.cpp
@@ -512,3 +512,71 @@ void FLocalPlanarSimplify::SimplifyAlongEdges(FDynamicMesh3& Mesh, TSet<int32>&
 		CollapseIters++;
 	}
 }
+
+
+void FLocalPlanarSimplify::CollapseDegenerateEdges(FDynamicMesh3& Mesh, TSet<int32>& InOutEdges, bool bRemoveAllDegenerateFromInputSet, double DegenerateEdgeLenTol, TSet<int>* TriangleSelection)
+{
+	double Tol2 = DegenerateEdgeLenTol * DegenerateEdgeLenTol;
+	FVector3d A, B;
+	TArray<int32> ToRemove;
+	int Collapsed = 0;
+	do
+	{
+		Collapsed = 0;
+		for (auto EdgesIter = InOutEdges.CreateIterator(); EdgesIter; ++EdgesIter)
+		{
+			int32 EID = *EdgesIter;
+			if (!Mesh.IsEdge(EID))
+			{
+				continue;
+			}
+			Mesh.GetEdgeV(EID, A, B);
+			if (DistanceSquared(A, B) > Tol2)
+			{
+				continue;
+			}
+
+			FIndex2i EV = Mesh.GetEdgeV(EID);
+			// if the vertex we'd remove is on a seam, try removing the other one instead
+			if (Mesh.HasAttributes() && Mesh.Attributes()->IsSeamVertex(EV.B, false))
+			{
+				Swap(EV.A, EV.B);
+				// if they were both on seams, then collapse should not happen?  (& would break OnCollapseEdge assumptions in overlay)
+				if (Mesh.HasAttributes() && Mesh.Attributes()->IsSeamVertex(EV.B, false))
+				{
+					continue;
+				}
+			}
+			FDynamicMesh3::FEdgeCollapseInfo CollapseInfo;
+			EMeshResult Result = Mesh.CollapseEdge(EV.A, EV.B, CollapseInfo);
+			if (Result == EMeshResult::Ok)
+			{
+				Collapsed++;
+				EdgesIter.RemoveCurrent();
+				if (TriangleSelection)
+				{
+					TriangleSelection->Remove(CollapseInfo.RemovedEdges.A);
+					if (CollapseInfo.RemovedEdges.B != FDynamicMesh3::InvalidID)
+					{
+						TriangleSelection->Remove(CollapseInfo.RemovedEdges.B);
+					}
+				}
+				// Add the other edge for later removal, so we don't potentially invalidate our iterator
+				// (but only if we really need to remove invalid edges here; some algorithms may be ok with leaving invalid EIDs in the set)
+				if (bRemoveAllDegenerateFromInputSet && CollapseInfo.RemovedEdges.B != FDynamicMesh3::InvalidID)
+				{
+					ToRemove.Add(CollapseInfo.RemovedEdges.B);
+				}
+			}
+		}
+	} while (Collapsed != 0);
+
+	if (bRemoveAllDegenerateFromInputSet)
+	{
+		for (int32 EID : ToRemove)
+		{
+			InOutEdges.Remove(EID);
+		}
+	}
+}
+
diff --git a/Engine/Source/Runtime/GeometryCore/Public/Operations/LocalPlanarSimplify.h b/Engine/Source/Runtime/GeometryCore/Public/Operations/LocalPlanarSimplify.h
--- a/Engine/Source/Runtime/GeometryCore/Public/Operations/LocalPlanarSimplify.h
+++ b/Engine/Source/Runtime/GeometryCore/Public/Operations/LocalPlanarSimplify.h
@@ -57,7 +57,15 @@ public:
 	 */
 	GEOMETRYCORE_API void SimplifyAlongEdges(FDynamicMesh3& Mesh, TSet<int32>& InOutEdges, TUniqueFunction<void(const DynamicMeshInfo::FEdgeCollapseInfo&)> ProcessCollapse = nullptr) const;
 
-
+	/**
+	 * Collapse degenerate edges. Note that SimplifyAlongEdges currently assumes that this (or a similar pass) is run beforehand
+	 * @param Mesh Mesh to repair
+	 * @param InOutEdges Set of edges to consider for collapse; optionally updated by removing edges at they are collapsed
+	 * @param bRemoveAllDegenerateFromInputSet Whether we should also check whether the neighbor edges removed by collapse were also in the set. Note: Often not needed, and has additional processing cost.
+	 * @param DegenerateEdgeLenTol Tolerance length at which we collapse edges
+	 * @param TriangleSelection	Optional set tracking an active selection. Any triangles removed by collapse will also be removed from the set.*/
+	GEOMETRYCORE_API static void CollapseDegenerateEdges(FDynamicMesh3& Mesh, TSet<int32>& InOutEdges, bool bRemoveAllDegenerateFromInputSet = false,
+		double DegenerateEdgeLenTol = UE_DOUBLE_SMALL_NUMBER, TSet<int>* TriangleSelection = nullptr);
 
 
 	//
