add 'material ID' option to the generate polygroups tool + polygroups generator, and also clean up the polygroups generator logic

diff --git a/Engine/Plugins/Runtime/GeometryProcessing/Source/DynamicMesh/Private/Polygroups/PolygroupsGenerator.cpp b/Engine/Plugins/Runtime/GeometryProcessing/Source/DynamicMesh/Private/Polygroups/PolygroupsGenerator.cpp
--- a/Engine/Plugins/Runtime/GeometryProcessing/Source/DynamicMesh/Private/Polygroups/PolygroupsGenerator.cpp
+++ b/Engine/Plugins/Runtime/GeometryProcessing/Source/DynamicMesh/Private/Polygroups/PolygroupsGenerator.cpp
@@ -33,28 +33,10 @@ bool FPolygroupsGenerator::FindPolygroupsFromUVIslands(int32 UVLayer)
 		return false;
 	}
 
-	FMeshConnectedComponents Components(Mesh);
-	Components.FindConnectedTriangles([&UV](int32 TriIdx0, int32 TriIdx1)
+	return FindPolygroupsFromConnectedTris([&UV](int32 TriIdx0, int32 TriIdx1)
 	{
 		return UV->AreTrianglesConnected(TriIdx0, TriIdx1);
 	});
-
-	int32 NumComponents = Components.Num();
-	for (int32 ci = 0; ci < NumComponents; ++ci)
-	{
-		FoundPolygroups.Add(Components.GetComponent(ci).Indices);
-	}
-
-	if (bApplyPostProcessing)
-	{
-		PostProcessPolygroups(false);
-	}
-	if (bCopyToMesh)
-	{
-		CopyPolygroupsToMesh();
-	}
-
-	return (FoundPolygroups.Num() > 0);
 }
 
 
@@ -66,41 +48,30 @@ bool FPolygroupsGenerator::FindPolygroupsFromHardNormalSeams()
 		return false;
 	}
 
-	FMeshConnectedComponents Components(Mesh);
-	Components.FindConnectedTriangles([Normals](int32 TriIdx0, int32 TriIdx1)
+	return FindPolygroupsFromConnectedTris([Normals](int32 TriIdx0, int32 TriIdx1)
 	{
 		return Normals->AreTrianglesConnected(TriIdx0, TriIdx1);
 	});
+}
 
-	int32 NumComponents = Components.Num();
-	for (int32 ci = 0; ci < NumComponents; ++ci)
+bool FPolygroupsGenerator::FindPolygroupsFromMaterialIDs()
+{
+	const FDynamicMeshMaterialAttribute* Materials = Mesh->Attributes()->GetMaterialID();
+	if (!Materials)
 	{
-		FoundPolygroups.Add(Components.GetComponent(ci).Indices);
+		return false;
 	}
 
-	if (bApplyPostProcessing)
-	{
-		PostProcessPolygroups(false);
-	}
-	if (bCopyToMesh)
+	return FindPolygroupsFromConnectedTris([Materials](int32 TriIdx0, int32 TriIdx1)
 	{
-		CopyPolygroupsToMesh();
-	}
-
-	return (FoundPolygroups.Num() > 0);
+		return Materials->GetValue(TriIdx0) == Materials->GetValue(TriIdx1);
+	});
 }
 
-
-bool FPolygroupsGenerator::FindPolygroupsFromConnectedTris()
+bool FPolygroupsGenerator::FindPolygroupsFromConnectedTris(TFunction<bool(int32, int32)> AreTrisConnected)
 {
-
 	FMeshConnectedComponents Components(Mesh);
-	Components.FindConnectedTriangles([this](int32 TriIdx0, int32 TriIdx1)
-	{
-		FIndex3i NbrTris = Mesh->GetTriNeighbourTris(TriIdx0);
-		int NbrIndex = IndexUtil::FindTriIndex(TriIdx1, NbrTris);
-		return (NbrIndex != IndexConstants::InvalidID);
-	});
+	Components.FindConnectedTriangles(AreTrisConnected);
 
 	int32 NumComponents = Components.Num();
 	for (int32 ci = 0; ci < NumComponents; ++ci)
@@ -124,6 +95,7 @@ bool FPolygroupsGenerator::FindPolygroupsFromConnectedTris()
 
 
 
+
 void FPolygroupsGenerator::GetSeamConstraintEdges(bool bUVSeams, bool bNormalSeams, TSet<int32>& ConstraintEdgesOut) const
 {
 	if (bUVSeams && Mesh->HasAttributes() && Mesh->Attributes()->PrimaryUV() != nullptr )
diff --git a/Engine/Plugins/Runtime/GeometryProcessing/Source/DynamicMesh/Public/Polygroups/PolygroupsGenerator.h b/Engine/Plugins/Runtime/GeometryProcessing/Source/DynamicMesh/Public/Polygroups/PolygroupsGenerator.h
--- a/Engine/Plugins/Runtime/GeometryProcessing/Source/DynamicMesh/Public/Polygroups/PolygroupsGenerator.h
+++ b/Engine/Plugins/Runtime/GeometryProcessing/Source/DynamicMesh/Public/Polygroups/PolygroupsGenerator.h
@@ -70,15 +70,22 @@ public:
 	 */
 	bool FindPolygroupsFromUVIslands(int32 UVLayer = 0);
 
+	/**
+	 * Find PolyGroups based on Material IDs -- i.e., connected groups of triangles with the same Material ID become a PolyGroup
+	 */
+	bool FindPolygroupsFromMaterialIDs();
+
 	/**
 	* Find Polygroups based on Seams in UV Overlay
 	*/
 	bool FindPolygroupsFromHardNormalSeams();
 
 	/**
-	 * Find Polygroups based on mesh connectivity, ie each connected-component becomes a Polygroup
+	 * Find PolyGroups based on mesh connectivity and optionally a custom indicator function
+	 * @param AreTrisConnected Optional custom function: Takes neighboring triangle IDs and returns true if they should be in the same PolyGroup
+	 * @return true if at least one polygroup was found
 	 */
-	bool FindPolygroupsFromConnectedTris();
+	bool FindPolygroupsFromConnectedTris(TFunction<bool(int32 TriA, int32 TriB)> AreTrisConnected = nullptr);
 
 	/**
 	* Find Polygroups by trying to invert triangluation of a polygon mesh
diff --git a/Engine/Plugins/Runtime/MeshModelingToolset/Source/MeshModelingTools/Private/ConvertToPolygonsTool.cpp b/Engine/Plugins/Runtime/MeshModelingToolset/Source/MeshModelingTools/Private/ConvertToPolygonsTool.cpp
--- a/Engine/Plugins/Runtime/MeshModelingToolset/Source/MeshModelingTools/Private/ConvertToPolygonsTool.cpp
+++ b/Engine/Plugins/Runtime/MeshModelingToolset/Source/MeshModelingTools/Private/ConvertToPolygonsTool.cpp
@@ -125,6 +125,11 @@ public:
 			Generator.FindPolygroupsFromHardNormalSeams();
 			break;
 		}
+		case EConvertToPolygonsMode::FromMaterialIDs:
+		{
+			Generator.FindPolygroupsFromMaterialIDs();
+			break;
+		}
 		case EConvertToPolygonsMode::FromConnectedTris:
 		{
 			Generator.FindPolygroupsFromConnectedTris();
diff --git a/Engine/Plugins/Runtime/MeshModelingToolset/Source/MeshModelingTools/Public/ConvertToPolygonsTool.h b/Engine/Plugins/Runtime/MeshModelingToolset/Source/MeshModelingTools/Public/ConvertToPolygonsTool.h
--- a/Engine/Plugins/Runtime/MeshModelingToolset/Source/MeshModelingTools/Public/ConvertToPolygonsTool.h
+++ b/Engine/Plugins/Runtime/MeshModelingToolset/Source/MeshModelingTools/Public/ConvertToPolygonsTool.h
@@ -41,19 +41,21 @@ UENUM()
 enum class EConvertToPolygonsMode
 {
 	/** Convert based on Angle Tolerance between Face Normals */
-	FaceNormalDeviation UMETA(DisplayName = "Face Normal Deviation"),
+	FaceNormalDeviation = 0 UMETA(DisplayName = "Face Normal Deviation"),
 	/** Create Polygroups by merging triangle pairs into Quads */
-	FindPolygons UMETA(DisplayName = "Find Quads"),
+	FindPolygons = 1 UMETA(DisplayName = "Find Quads"),
+	/** Create PolyGroups based on Material IDs */
+	FromMaterialIDs = 7 UMETA(DisplayName = "From Material IDs"),
 	/** Create PolyGroups based on UV Islands */
-	FromUVIslands  UMETA(DisplayName = "From UV Islands"),
+	FromUVIslands = 2 UMETA(DisplayName = "From UV Islands"),
 	/** Create PolyGroups based on Hard Normal Seams */
-	FromNormalSeams  UMETA(DisplayName = "From Hard Normal Seams"),
+	FromNormalSeams = 3 UMETA(DisplayName = "From Hard Normal Seams"),
 	/** Create Polygroups based on Connected Triangles */
-	FromConnectedTris UMETA(DisplayName = "From Connected Tris"),
+	FromConnectedTris = 4 UMETA(DisplayName = "From Connected Tris"),
 	/** Create Polygroups centered on well-spaced sample points, approximating a surface Voronoi diagram */
-	FromFurthestPointSampling UMETA(DisplayName = "Furthest Point Sampling"),
+	FromFurthestPointSampling = 5 UMETA(DisplayName = "Furthest Point Sampling"),
 	/** Copy from existing Polygroup Layer */
-	CopyFromLayer UMETA(DisplayName = "Copy From Layer"),
+	CopyFromLayer = 6 UMETA(DisplayName = "Copy From Layer")
 };
 
 
