[GeometryProcessing]
* Moved dll storage to methods/variables instead of on types. This is needed to get merged modules working

diff --git a/Engine/Plugins/Runtime/GeometryProcessing/Source/DynamicMesh/Public/MeshConstraints.h b/Engine/Plugins/Runtime/GeometryProcessing/Source/DynamicMesh/Public/MeshConstraints.h
--- a/Engine/Plugins/Runtime/GeometryProcessing/Source/DynamicMesh/Public/MeshConstraints.h
+++ b/Engine/Plugins/Runtime/GeometryProcessing/Source/DynamicMesh/Public/MeshConstraints.h
@@ -39,7 +39,7 @@ enum class EEdgeRefineFlags
 /**
  * FEdgeConstraint is a constraint on a triangle mesh edge
  */
-struct DYNAMICMESH_API FEdgeConstraint
+struct FEdgeConstraint
 {
 public:
 	/** Constraint flags on this edge */
@@ -143,7 +143,7 @@ public:
 /**
  * FVertexConstraint is a constraint on a triangle mesh vertex
  */
-struct DYNAMICMESH_API FVertexConstraint
+struct FVertexConstraint
 {
 public:
 
@@ -223,7 +223,7 @@ public:
 	}
 };
 
-struct DYNAMICMESH_API FConstraintProjectionData
+struct FConstraintProjectionData
 {
 	TArray<TSharedPtr<FMeshConstraintCurve>> ProjectionCurves;
 };
@@ -232,7 +232,7 @@ struct DYNAMICMESH_API FConstraintProjectionData
 /**
  * FMeshConstraints is a set of Edge and Vertex constraints for a Triangle Mesh
  */
-class DYNAMICMESH_API FMeshConstraints
+class FMeshConstraints
 {
 protected:
 
diff --git a/Engine/Plugins/Runtime/GeometryProcessing/Source/DynamicMesh/Public/MeshRefinerBase.h b/Engine/Plugins/Runtime/GeometryProcessing/Source/DynamicMesh/Public/MeshRefinerBase.h
--- a/Engine/Plugins/Runtime/GeometryProcessing/Source/DynamicMesh/Public/MeshRefinerBase.h
+++ b/Engine/Plugins/Runtime/GeometryProcessing/Source/DynamicMesh/Public/MeshRefinerBase.h
@@ -26,7 +26,7 @@ public:
  * This is a base class that implements common functionality for various triangle mesh resampling strategies
  * (ie FRemesher and FReducer). You probably should not use this class directly.
  */
-class DYNAMICMESH_API FMeshRefinerBase : public FMeshRefinerBaseFixLayout
+class FMeshRefinerBase : public FMeshRefinerBaseFixLayout
 {
 protected:
 	/** Mesh that will be refined */
@@ -178,7 +178,7 @@ protected:
 	 * @param tc triangle on one side of edge
 	 * @param td triangle on other side of edge
 	 */
-	bool CheckIfCollapseCreatesFlipOrInvalid(int vid, int vother, const FVector3d& newv, int tc, int td) const;
+	DYNAMICMESH_API bool CheckIfCollapseCreatesFlipOrInvalid(int vid, int vother, const FVector3d& newv, int tc, int td) const;
 
 	/** Avoid creation of triangles smaller than this value. 
 	 * We compare the triangle cross product's norm squared to this value since it's cheaper to compute than the triangle area. So the actual
@@ -196,7 +196,7 @@ protected:
 	 * @param tc triangle on one side of edge
 	 * @param td triangle on other side of edge
 	 */
-	bool CheckIfCollapseCreatesTinyTriangle(int vid, int vother, const FVector3d& newv, int tc, int td) const;
+	DYNAMICMESH_API bool CheckIfCollapseCreatesTinyTriangle(int vid, int vother, const FVector3d& newv, int tc, int td) const;
 
 	/**
 	 * Check if edge flip might reverse normal direction.
@@ -207,7 +207,7 @@ protected:
 	 * @param d opposing vertex 2
 	 * @param t0 index of triangle containing [a,b,c]
 	 */
-	bool CheckIfFlipInvertsNormals(int a, int b, int c, int d, int t0) const;
+	DYNAMICMESH_API bool CheckIfFlipInvertsNormals(int a, int b, int c, int d, int t0) const;
 
 	/**
 	 * Check if edge flip might create a triangle with small area (either all vertices are close together, or in a sliver configuration)
@@ -217,7 +217,7 @@ protected:
 	 * @param d opposing vertex 2
 	 * @param t0 index of triangle containing [a,b,c]
 	 */
-	bool CheckIfFlipCreatesTinyTriangle(int OriginalEdgeVertexA, int OriginalEdgeVertexB, int OppositeEdgeVertexC, int OppositeEdgeVertexD, int OriginalTriangleIndex) const;
+	DYNAMICMESH_API bool CheckIfFlipCreatesTinyTriangle(int OriginalEdgeVertexA, int OriginalEdgeVertexB, int OppositeEdgeVertexC, int OppositeEdgeVertexD, int OriginalTriangleIndex) const;
 
 	/**
 	 * Figure out if we can collapse edge eid=[a,b] under current constraint set.
@@ -233,7 +233,7 @@ protected:
 	 * @param td index of triangle [a,b,d]
 	 * @param collapse_to either a or b if we should collapse to one of those, or -1 if either is acceptable
 	 */
-	bool CanCollapseEdge(int eid, int a, int b, int c, int d, int tc, int td, int& collapse_to) const;
+	DYNAMICMESH_API bool CanCollapseEdge(int eid, int a, int b, int c, int d, int tc, int td, int& collapse_to) const;
 
 	/**
 	 * Resolve vertex constraints for collapsing edge eid=[a,b]. Generally we would
@@ -246,7 +246,7 @@ protected:
 	 * @param b second vertex of edge*
 	 * @param collapse_to either a or b if we should collapse to one of those, or -1 if either is acceptable
 	 */
-	bool CanCollapseVertex(int eid, int a, int b, int& collapse_to) const;
+	DYNAMICMESH_API bool CanCollapseVertex(int eid, int a, int b, int& collapse_to) const;
 
 
 	/**
@@ -290,22 +290,22 @@ public:
 	// Mesh Change Tracking support
 	//
 
-	void SetMeshChangeTracker(FDynamicMeshChangeTracker* Tracker);
+	DYNAMICMESH_API void SetMeshChangeTracker(FDynamicMeshChangeTracker* Tracker);
 
 protected:
-	virtual void SaveTriangleBeforeModify(int32 TriangleID);
-	virtual void SaveEdgeBeforeModify(int32 EdgeID);
-	virtual void SaveVertexTrianglesBeforeModify(int32 VertexID);
+	DYNAMICMESH_API virtual void SaveTriangleBeforeModify(int32 TriangleID);
+	DYNAMICMESH_API virtual void SaveEdgeBeforeModify(int32 EdgeID);
+	DYNAMICMESH_API virtual void SaveVertexTrianglesBeforeModify(int32 VertexID);
 
 
 protected:
 	/*
 	 * testing/debug/profiling stuff
 	 */
-	void RuntimeDebugCheck(int EdgeID);
-	virtual void DoDebugChecks(bool bEndOfPass = false);
-	void DebugCheckUVSeamConstraints();
-	void DebugCheckVertexConstraints();
+	DYNAMICMESH_API void RuntimeDebugCheck(int EdgeID);
+	DYNAMICMESH_API virtual void DoDebugChecks(bool bEndOfPass = false);
+	DYNAMICMESH_API void DebugCheckUVSeamConstraints();
+	DYNAMICMESH_API void DebugCheckVertexConstraints();
 
 };
 
diff --git a/Engine/Plugins/Runtime/GeometryProcessing/Source/DynamicMesh/Public/Parameterization/PatchBasedMeshUVGenerator.h b/Engine/Plugins/Runtime/GeometryProcessing/Source/DynamicMesh/Public/Parameterization/PatchBasedMeshUVGenerator.h
--- a/Engine/Plugins/Runtime/GeometryProcessing/Source/DynamicMesh/Public/Parameterization/PatchBasedMeshUVGenerator.h
+++ b/Engine/Plugins/Runtime/GeometryProcessing/Source/DynamicMesh/Public/Parameterization/PatchBasedMeshUVGenerator.h
@@ -23,7 +23,7 @@ namespace Geometry
  * then the initial steps can be skipped
  * 
  */
-class DYNAMICMESH_API FPatchBasedMeshUVGenerator
+class FPatchBasedMeshUVGenerator
 {
 public:
 
@@ -70,7 +70,7 @@ public:
 	 * given TargetUVOverlay, using the steps below, based on the parameters above
 	 * @return result information, which may include warnings/etc
 	 */
-	FGeometryResult AutoComputeUVs(
+	DYNAMICMESH_API FGeometryResult AutoComputeUVs(
 		FDynamicMesh3& TargetMesh,
 		FDynamicMeshUVOverlay& TargetUVOverlay,
 		FProgressCancel* Progress = nullptr);
@@ -80,7 +80,7 @@ public:
 	 * @param InitialComponentsOut returned decomposition into triangle sets
 	 * @return true on success
 	 */
-	bool ComputeInitialMeshPatches(
+	DYNAMICMESH_API bool ComputeInitialMeshPatches(
 		FDynamicMesh3& TargetMesh,
 		FMeshConnectedComponents& InitialComponentsOut,
 		FProgressCancel* Progress = nullptr);
@@ -94,7 +94,7 @@ public:
 	 * @param IslandsOut the set of output mesh patches / triangle lists
 	 * @return false if catastrophic failure ocurred (ie IslandsOut is invalid) or if Progress was cancelled
 	 */
-	bool ComputeIslandsByRegionMerging(
+	DYNAMICMESH_API bool ComputeIslandsByRegionMerging(
 		FDynamicMesh3& TargetMesh,
 		FDynamicMeshUVOverlay& TargetUVOverlay,
 		const FMeshConnectedComponents& InitialComponents, 
@@ -110,7 +110,7 @@ public:
 	 * @param IslandsOut the set of output mesh patches / triangle lists
 	 * @return false if catastrophic failure ocurred (ie IslandsOut is invalid) or if Progress was cancelled
 	 */
-	bool ComputeIslandsByRegionMerging(
+	DYNAMICMESH_API bool ComputeIslandsByRegionMerging(
 		FDynamicMesh3& TargetMesh,
 		FDynamicMeshUVOverlay& TargetUVOverlay,
 		const FMeshConnectedComponents& InitialComponents,
@@ -124,7 +124,7 @@ public:
 	 * @param bValidUVsComputedOut output list of booleans the same length as TriangleSets, set to false if UV computation failed for that island
 	 * @return number of failed UV islands, 0 if all were successful, and -1 if Progress was cancelled
 	 */
-	int32 ComputeUVsFromTriangleSets(
+	DYNAMICMESH_API int32 ComputeUVsFromTriangleSets(
 		FDynamicMesh3& TargetMesh,
 		FDynamicMeshUVOverlay& TargetUVOverlay,
 		const TArray<TArray<int32>>& TriangleSets,
diff --git a/Engine/Plugins/Runtime/GeometryProcessing/Source/DynamicMesh/Public/ProjectionTargets.h b/Engine/Plugins/Runtime/GeometryProcessing/Source/DynamicMesh/Public/ProjectionTargets.h
--- a/Engine/Plugins/Runtime/GeometryProcessing/Source/DynamicMesh/Public/ProjectionTargets.h
+++ b/Engine/Plugins/Runtime/GeometryProcessing/Source/DynamicMesh/Public/ProjectionTargets.h
@@ -17,7 +17,7 @@ namespace Geometry
  * FMeshProjectionTarget provides an IProjectionTarget interface to a FDynamicMesh + FDynamicMeshAABBTree3
  * Use to project points to mesh surface.
  */
-class DYNAMICMESH_API FMeshProjectionTarget : public IOrientedProjectionTarget
+class FMeshProjectionTarget : public IOrientedProjectionTarget
 {
 public:
 	/** The mesh to project onto */
@@ -38,18 +38,18 @@ public:
 	/**
 	 * @return Projection of Point onto this target
 	 */
-	virtual FVector3d Project(const FVector3d& Point, int Identifier = -1) override;
+	DYNAMICMESH_API virtual FVector3d Project(const FVector3d& Point, int Identifier = -1) override;
 
 	/**
 	 * @return Projection of Point onto this target, and set ProjectNormalOut to the triangle normal at the returned point (*not* interpolated vertex normal)
 	 */
-	virtual FVector3d Project(const FVector3d& Point, FVector3d& ProjectNormalOut, int Identifier = -1) override;
+	DYNAMICMESH_API virtual FVector3d Project(const FVector3d& Point, FVector3d& ProjectNormalOut, int Identifier = -1) override;
 
 
 };
 
 
-class DYNAMICMESH_API FWorldSpaceProjectionTarget : public FMeshProjectionTarget
+class FWorldSpaceProjectionTarget : public FMeshProjectionTarget
 {
 public:
 
diff --git a/Engine/Plugins/Runtime/GeometryProcessing/Source/DynamicMesh/Public/QueueRemesher.h b/Engine/Plugins/Runtime/GeometryProcessing/Source/DynamicMesh/Public/QueueRemesher.h
--- a/Engine/Plugins/Runtime/GeometryProcessing/Source/DynamicMesh/Public/QueueRemesher.h
+++ b/Engine/Plugins/Runtime/GeometryProcessing/Source/DynamicMesh/Public/QueueRemesher.h
@@ -21,7 +21,7 @@ namespace Geometry
  *  - TrackedFullProjectionPass() projects all vertices but only adds to queue if edge changes enough
  * 
  */
-class DYNAMICMESH_API FQueueRemesher : public FRemesher
+class FQueueRemesher : public FRemesher
 {
 
 public:
@@ -38,7 +38,7 @@ public:
 	double MinActiveEdgeFraction = 0.01;
 
 	/// Converge on remeshed result as quickly as possible
-	void FastestRemesh();
+	DYNAMICMESH_API void FastestRemesh();
 
 	/// "Outer loop" for all remeshing operations
 	void BasicRemeshPass() override
@@ -49,7 +49,7 @@ public:
 protected:
 
 	// Note: FRemesher has this: virtual void FullProjectionPass();
-	virtual void TrackedFullProjectionPass(bool bParallel);
+	DYNAMICMESH_API virtual void TrackedFullProjectionPass(bool bParallel);
 
 	// Set of edges that have been modified during a given iteration.
 	// Before the first iteration, this Optional is unset, signaling that the first iteration should process all edges
@@ -139,26 +139,26 @@ protected:
 	
 	/// This pass only does edge splits and enqueues modified vertex neighborhoods.
 	/// @return Number of split edges.
-	int FastSplitIteration();
+	DYNAMICMESH_API int FastSplitIteration();
 
 	/// This pass does all edge operations plug smoothing and reprojection.
 	/// EarlyTerminationCheck will be called after remesh operations, before smoothing/reprojection. Caller can return true to abort further computation.
-	void RemeshIteration(
+	DYNAMICMESH_API void RemeshIteration(
 		TFunctionRef<bool(void)> EarlyTerminationCheck
 	);
 
 	// Whether an edge should be added to ModifiedEdges (used in async tasks)
 	TArray<bool> EdgeShouldBeQueuedBuffer;	
 
-	virtual void InitializeVertexBufferForPass();
+	DYNAMICMESH_API virtual void InitializeVertexBufferForPass();
 
 	// Move all vertices in parallel. Update the timestamps. Add edges whose endpoints change to ModifiedEdges.
 	// The NewVertexPosition function accepts a vertex index and returns a new vertex position, as well as setting the 
 	// bool out-param to indicate that the vertex has moved.
-	void TrackedMoveVerticesParallel(TFunction<FVector3d(int, bool&)> NewVertexPosition);
+	DYNAMICMESH_API void TrackedMoveVerticesParallel(TFunction<FVector3d(int, bool&)> NewVertexPosition);
 
 	// Note: superclass has this: virtual void FullSmoothPass_Buffer(bool bParallel);
-	void TrackedFullSmoothPass_Buffer(bool bParallel);
+	DYNAMICMESH_API void TrackedFullSmoothPass_Buffer(bool bParallel);
 
 
 };
diff --git a/Engine/Plugins/Runtime/GeometryProcessing/Source/DynamicMesh/Public/Sampling/MeshCurvatureMapBaker.h b/Engine/Plugins/Runtime/GeometryProcessing/Source/DynamicMesh/Public/Sampling/MeshCurvatureMapBaker.h
--- a/Engine/Plugins/Runtime/GeometryProcessing/Source/DynamicMesh/Public/Sampling/MeshCurvatureMapBaker.h
+++ b/Engine/Plugins/Runtime/GeometryProcessing/Source/DynamicMesh/Public/Sampling/MeshCurvatureMapBaker.h
@@ -13,7 +13,7 @@ namespace Geometry
 
 class FMeshVertexCurvatureCache;
 
-class DYNAMICMESH_API FMeshCurvatureMapBaker : public FMeshImageBaker
+class FMeshCurvatureMapBaker : public FMeshImageBaker
 {
 public:
 	virtual ~FMeshCurvatureMapBaker() {}
@@ -69,10 +69,10 @@ public:
 	//
 
 	/** Calculate bake result */
-	virtual void Bake() override;
+	DYNAMICMESH_API virtual void Bake() override;
 
 	/** populate Curvatures member if valid data has not been provided */
-	void CacheDetailCurvatures(const FDynamicMesh3* DetailMesh);
+	DYNAMICMESH_API void CacheDetailCurvatures(const FDynamicMesh3* DetailMesh);
 
 	//
 	// Output
diff --git a/Engine/Plugins/Runtime/GeometryProcessing/Source/DynamicMesh/Public/Sampling/MeshImageBaker.h b/Engine/Plugins/Runtime/GeometryProcessing/Source/DynamicMesh/Public/Sampling/MeshImageBaker.h
--- a/Engine/Plugins/Runtime/GeometryProcessing/Source/DynamicMesh/Public/Sampling/MeshImageBaker.h
+++ b/Engine/Plugins/Runtime/GeometryProcessing/Source/DynamicMesh/Public/Sampling/MeshImageBaker.h
@@ -9,7 +9,7 @@ namespace UE
 namespace Geometry
 {
 
-class DYNAMICMESH_API FMeshImageBaker
+class FMeshImageBaker
 {
 public:
 	virtual ~FMeshImageBaker() {}
diff --git a/Engine/Plugins/Runtime/GeometryProcessing/Source/DynamicMesh/Public/Sampling/MeshImageBakingCache.h b/Engine/Plugins/Runtime/GeometryProcessing/Source/DynamicMesh/Public/Sampling/MeshImageBakingCache.h
--- a/Engine/Plugins/Runtime/GeometryProcessing/Source/DynamicMesh/Public/Sampling/MeshImageBakingCache.h
+++ b/Engine/Plugins/Runtime/GeometryProcessing/Source/DynamicMesh/Public/Sampling/MeshImageBakingCache.h
@@ -10,12 +10,14 @@
 #include "Image/ImageDimensions.h"
 #include "Image/ImageOccupancyMap.h"
 
+#define UE_API DYNAMICMESH_API
+
 namespace UE
 {
 namespace Geometry
 {
 
-class DYNAMICMESH_API FMeshImageBakingCache
+class FMeshImageBakingCache
 {
 public:
 
@@ -36,14 +38,14 @@ public:
 		Identity
 	};
 
-	void SetDetailMesh(const FDynamicMesh3* Mesh, const FDynamicMeshAABBTree3* Spatial);
-	void SetBakeTargetMesh(const FDynamicMesh3* Mesh);
+	UE_API void SetDetailMesh(const FDynamicMesh3* Mesh, const FDynamicMeshAABBTree3* Spatial);
+	UE_API void SetBakeTargetMesh(const FDynamicMesh3* Mesh);
 
-	void SetDimensions(FImageDimensions Dimensions);
-	void SetUVLayer(int32 UVLayer);
-	void SetThickness(double Thickness);
-	void SetCorrespondenceStrategy(ECorrespondenceStrategy Strategy);
-	void SetGutterSize(int32 GutterSize);
+	UE_API void SetDimensions(FImageDimensions Dimensions);
+	UE_API void SetUVLayer(int32 UVLayer);
+	UE_API void SetThickness(double Thickness);
+	UE_API void SetCorrespondenceStrategy(ECorrespondenceStrategy Strategy);
+	UE_API void SetGutterSize(int32 GutterSize);
 
 	FImageDimensions GetDimensions() const { return Dimensions; }
 	int32 GetUVLayer() const { return UVLayer; }
@@ -51,17 +53,17 @@ public:
 	ECorrespondenceStrategy GetCorrespondenceStrategy() const { return CorrespondenceStrategy; }
 
 	const FDynamicMesh3* GetBakeTargetMesh() const { return TargetMesh; }
-	const FDynamicMeshUVOverlay* GetBakeTargetUVs() const;
-	const FDynamicMeshNormalOverlay* GetBakeTargetNormals() const;
+	UE_API const FDynamicMeshUVOverlay* GetBakeTargetUVs() const;
+	UE_API const FDynamicMeshNormalOverlay* GetBakeTargetNormals() const;
 
 	const FDynamicMesh3* GetDetailMesh() const { return DetailMesh; }
 	const FDynamicMeshAABBTree3* GetDetailSpatial() const { return DetailSpatial; }
-	const FDynamicMeshNormalOverlay* GetDetailNormals() const;
+	UE_API const FDynamicMeshNormalOverlay* GetDetailNormals() const;
 
 
 	bool IsCacheValid() const { return bOccupancyValid && bSamplesValid; }
 
-	bool ValidateCache();
+	UE_API bool ValidateCache();
 
 
 	struct FCorrespondenceSample
@@ -74,10 +76,10 @@ public:
 	};
 
 
-	void EvaluateSamples(TFunctionRef<void(const FVector2i&, const FCorrespondenceSample&)> SampleFunction,
+	UE_API void EvaluateSamples(TFunctionRef<void(const FVector2i&, const FCorrespondenceSample&)> SampleFunction,
 		bool bParallel = true) const;
 
-	const FImageOccupancyMap* GetOccupancyMap() const;
+	UE_API const FImageOccupancyMap* GetOccupancyMap() const;
 
 
 	/**
@@ -88,7 +90,7 @@ public:
 	 * better results than (eg) using a "worse" sampling strategy to guarantee some
 	 * non-garbage result (eg Nearest-point sampling)
 	 */
-	void FindSamplingHoles(TFunctionRef<bool(const FVector2i&)> IsValidSampleFunction,
+	UE_API void FindSamplingHoles(TFunctionRef<bool(const FVector2i&)> IsValidSampleFunction,
 		TArray<FVector2i>& HolePixelsOut,
 		bool bParallel = true) const;
 
@@ -109,12 +111,12 @@ protected:
 
 	TDenseGrid2<FCorrespondenceSample> SampleMap;
 	bool bSamplesValid = false;
-	void InvalidateSamples();
+	UE_API void InvalidateSamples();
 
 
 	TUniquePtr<FImageOccupancyMap> OccupancyMap;
 	bool bOccupancyValid = false;
-	void InvalidateOccupancy();
+	UE_API void InvalidateOccupancy();
 
 };
 
@@ -122,3 +124,5 @@ protected:
 
 } // end namespace UE::Geometry
 } // end namespace UE
+
+#undef UE_API
\ No newline at end of file
diff --git a/Engine/Plugins/Runtime/GeometryProcessing/Source/DynamicMesh/Public/Sampling/MeshOcclusionMapBaker.h b/Engine/Plugins/Runtime/GeometryProcessing/Source/DynamicMesh/Public/Sampling/MeshOcclusionMapBaker.h
--- a/Engine/Plugins/Runtime/GeometryProcessing/Source/DynamicMesh/Public/Sampling/MeshOcclusionMapBaker.h
+++ b/Engine/Plugins/Runtime/GeometryProcessing/Source/DynamicMesh/Public/Sampling/MeshOcclusionMapBaker.h
@@ -22,7 +22,7 @@ enum class EOcclusionMapType : uint8
 };
 ENUM_CLASS_FLAGS(EOcclusionMapType);
 
-class DYNAMICMESH_API FMeshOcclusionMapBaker : public FMeshImageBaker
+class FMeshOcclusionMapBaker : public FMeshImageBaker
 {
 public:
 	virtual ~FMeshOcclusionMapBaker() {}
@@ -67,7 +67,7 @@ public:
 	// Compute functions
 	//
 
-	virtual void Bake() override;
+	DYNAMICMESH_API virtual void Bake() override;
 
 	//
 	// Output
