[GeometryProcessing]
* Moved dll storage to methods/variables instead of on types. This is needed to get merged modules working

diff --git a/Engine/Plugins/Runtime/GeometryProcessing/Source/DynamicMesh/Private/Solvers/Internal/ConstrainedMeshDeformationSolver.h b/Engine/Plugins/Runtime/GeometryProcessing/Source/DynamicMesh/Private/Solvers/Internal/ConstrainedMeshDeformationSolver.h
--- a/Engine/Plugins/Runtime/GeometryProcessing/Source/DynamicMesh/Private/Solvers/Internal/ConstrainedMeshDeformationSolver.h
+++ b/Engine/Plugins/Runtime/GeometryProcessing/Source/DynamicMesh/Private/Solvers/Internal/ConstrainedMeshDeformationSolver.h
@@ -13,6 +13,8 @@
 
 #include "MatrixBase.h"
 
+#define UE_API DYNAMICMESH_API
+
 
 namespace UE
 {
@@ -29,26 +31,26 @@ namespace Geometry
  *
  * Both the Laplacian weighting scheme/type and The Sparse Matrix Solver type are configurable.
  */
-class DYNAMICMESH_API FConstrainedMeshDeformationSolver : public UE::Solvers::IConstrainedMeshSolver
+class FConstrainedMeshDeformationSolver : public UE::Solvers::IConstrainedMeshSolver
 {
 public:
 	typedef UE::Solvers::FPositionConstraint  FConstraintPosition;
 
-	FConstrainedMeshDeformationSolver(const FDynamicMesh3& DynamicMesh, const ELaplacianWeightScheme Scheme, const EMatrixSolverType MatrixSolverType);
-	virtual ~FConstrainedMeshDeformationSolver();
+	UE_API FConstrainedMeshDeformationSolver(const FDynamicMesh3& DynamicMesh, const ELaplacianWeightScheme Scheme, const EMatrixSolverType MatrixSolverType);
+	UE_API virtual ~FConstrainedMeshDeformationSolver();
 
 	// Create the solver using a generic mesh type. Assume uniform weight scheme.
 	template<typename MeshType>
 	FConstrainedMeshDeformationSolver(const MeshType& Mesh, const EMatrixSolverType MatrixSolverType);
 
 	// Add constraint associated with given vertex id.  Boundary vertices will be ignored
-	void AddConstraint(const int32 VtxId, const double Weight, const FVector3d& Pos, const bool bPostFix) override;
+	UE_API void AddConstraint(const int32 VtxId, const double Weight, const FVector3d& Pos, const bool bPostFix) override;
 
 	// Update the position of an existing constraint.  Bool return if a corresponding constraint weight exists. Boundary vertices will be ignored (and return false).  
-	bool UpdateConstraintPosition(const int32 VtxId, const FVector3d& Position, const bool bPostFix) override;
+	UE_API bool UpdateConstraintPosition(const int32 VtxId, const FVector3d& Position, const bool bPostFix) override;
 
 	// The underlying solver will have to refactor the matrix if this is done. Bool return if a corresponding constraint position exists. Boundary vertices will be ignored (and return false).  
-	bool UpdateConstraintWeight(const int32 VtxId, const double Weight) override;
+	UE_API bool UpdateConstraintWeight(const int32 VtxId, const double Weight) override;
 
 	// Clear all constraints associated with this smoother
 	void ClearConstraints() override { FConstrainedMeshDeformationSolver::ClearConstraintPositions(); FConstrainedMeshDeformationSolver::ClearConstraintWeights(); }
@@ -59,21 +61,21 @@ public:
 
 
 	// Test if for constraint associated with given vertex id. Will return false for any boundary vert.
-	bool IsConstrained(const int32 VtxId) const override;
+	UE_API bool IsConstrained(const int32 VtxId) const override;
 
 	virtual bool Deform(TArray<FVector3d>& PositionBuffer) override { return false; }
 
 	// Sync constraints with internal solver.  If in the process any internal matrix factoring is dirty, it will be rebuilt.
 	// Note: this is called from within the Deform() method.   Only call this method if you want to trigger the matrix refactor yourself.
-	void UpdateSolverConstraints();
+	UE_API void UpdateSolverConstraints();
 
 
 
 public:
 	// these only apply to iterative solvers
 
-	bool SetMaxIterations(int32 MaxIterations);
-	bool SetTolerance(double Tol);
+	UE_API bool SetMaxIterations(int32 MaxIterations);
+	UE_API bool SetTolerance(double Tol);
 
 
 protected:
@@ -143,15 +145,15 @@ protected:
 
 	//void UpdateMeshWithConstraints();
 	// Respect any bPostFix constraints by moving those vertices to position defined by said constraint.
-	void UpdateWithPostFixConstraints(FSOAPositions& PositionalVector) const;
+	UE_API void UpdateWithPostFixConstraints(FSOAPositions& PositionalVector) const;
 
 
 	// converts the positionalvector to a TArray<FVector3d> where the offset in the array is implicitly the
 	// VtxId in the mesh, and not ness the matrix row id
 	// NB: the resulting array is treated as sparse and may have un-initialized elements.
-	bool CopyInternalPositions(const FSOAPositions& PositionalVector, TArray<FVector3d>& LinearArray) const;
+	UE_API bool CopyInternalPositions(const FSOAPositions& PositionalVector, TArray<FVector3d>& LinearArray) const;
 
-	bool CopyBoundaryPositions(TArray<FVector3d>& LinearArray) const;
+	UE_API bool CopyBoundaryPositions(TArray<FVector3d>& LinearArray) const;
 
 
 protected:
@@ -210,36 +212,36 @@ FConstrainedMeshDeformationSolver::FConstrainedMeshDeformationSolver(const MeshT
  *
  * Voronoi-Area Clamped Cotangent weights are used for the Laplacian, and an LU Sover
  */
-class DYNAMICMESH_API FSoftMeshDeformationSolver : public UE::Solvers::IConstrainedLaplacianMeshSolver
+class FSoftMeshDeformationSolver : public UE::Solvers::IConstrainedLaplacianMeshSolver
 {
 public:
 	typedef UE::Solvers::FPositionConstraint  FPositionConstraint;
 
-	FSoftMeshDeformationSolver(const FDynamicMesh3& DynamicMesh);
-	virtual ~FSoftMeshDeformationSolver();
+	UE_API FSoftMeshDeformationSolver(const FDynamicMesh3& DynamicMesh);
+	UE_API virtual ~FSoftMeshDeformationSolver();
 
 	//
 	// IConstrainedMeshSolver API
 	//
 
 	// Add constraint associated with given vertex id.  Boundary vertices will be ignored
-	void AddConstraint(const int32 VtxId, const double Weight, const FVector3d& Position, const bool bPostFix) override;
+	UE_API void AddConstraint(const int32 VtxId, const double Weight, const FVector3d& Position, const bool bPostFix) override;
 
 	// Update the position of an existing constraint.  Bool return if a corresponding constraint weight exists. Boundary vertices will be ignored (and return false).  
-	bool UpdateConstraintPosition(const int32 VtxId, const FVector3d& Position, const bool bPostFix) override;
+	UE_API bool UpdateConstraintPosition(const int32 VtxId, const FVector3d& Position, const bool bPostFix) override;
 
 	// The underlying solver will have to refactor the matrix if this is done. Bool return if a corresponding constraint position exists. Boundary vertices will be ignored (and return false).  
-	bool UpdateConstraintWeight(const int32 VtxId, const double Weight) override;
+	UE_API bool UpdateConstraintWeight(const int32 VtxId, const double Weight) override;
 
 	// Clear all constraints associated with this smoother
-	void ClearConstraints() override;
+	UE_API void ClearConstraints() override;
 
 	// do not support these
 	void ClearConstraintWeights() override { check(false); }
 	void ClearConstraintPositions() override { check(false); }
 
 	// Test if for constraint associated with given vertex id. Will return false for any boundary vert.
-	bool IsConstrained(const int32 VtxId) const override;
+	UE_API bool IsConstrained(const int32 VtxId) const override;
 
 	virtual bool Deform(TArray<FVector3d>& PositionBuffer) override { return false;	}
 
@@ -250,14 +252,14 @@ public:
 
 
 	// Update global scale applied to Laplacian vectors before solve.
-	virtual void UpdateLaplacianScale(double UniformScale);
+	UE_API virtual void UpdateLaplacianScale(double UniformScale);
 
 
 public:
 
 	// Sync constraints with internal solver.  If in the process any internal matrix factoring is dirty, it will be rebuilt.
 	// Note: this is called from within the Deform() method.   Only call this method if you want to trigger the matrix refactor yourself.
-	void UpdateSolverConstraints();
+	UE_API void UpdateSolverConstraints();
 
 
 protected:
@@ -276,10 +278,12 @@ protected:
 	TUniquePtr<FConstrainedSolver> ConstrainedSolver;
 
 	double LaplacianScale = 1.0;
-	bool HasLaplacianScale() const;
-	double GetLaplacianScale(int32 LinearVtxIndex) const;
+	UE_API bool HasLaplacianScale() const;
+	UE_API double GetLaplacianScale(int32 LinearVtxIndex) const;
 };
 
 
 } // end namespace UE::Geometry
 } // end namespace UE
+
+#undef UE_API
diff --git a/Engine/Plugins/Runtime/GeometryProcessing/Source/DynamicMesh/Private/Solvers/Internal/ConstrainedMeshDeformers.h b/Engine/Plugins/Runtime/GeometryProcessing/Source/DynamicMesh/Private/Solvers/Internal/ConstrainedMeshDeformers.h
--- a/Engine/Plugins/Runtime/GeometryProcessing/Source/DynamicMesh/Private/Solvers/Internal/ConstrainedMeshDeformers.h
+++ b/Engine/Plugins/Runtime/GeometryProcessing/Source/DynamicMesh/Private/Solvers/Internal/ConstrainedMeshDeformers.h
@@ -6,6 +6,8 @@
 #include "DynamicMesh/DynamicMesh3.h"
 #include "ConstrainedMeshDeformationSolver.h"
 
+#define UE_API DYNAMICMESH_API
+
 namespace UE
 {
 namespace Geometry
@@ -19,17 +21,17 @@ namespace Geometry
  *
  * Boundary vertices are fixed to their input positions.
  */
-class DYNAMICMESH_API FConstrainedMeshDeformer : public FConstrainedMeshDeformationSolver
+class FConstrainedMeshDeformer : public FConstrainedMeshDeformationSolver
 {
 public:
-	FConstrainedMeshDeformer(const FDynamicMesh3& DynamicMesh, const ELaplacianWeightScheme LaplacianType);
+	UE_API FConstrainedMeshDeformer(const FDynamicMesh3& DynamicMesh, const ELaplacianWeightScheme LaplacianType);
 	~FConstrainedMeshDeformer() override {}
 
 	// Generic mesh type, assume uniform weight scheme
 	template<typename MeshT>
 	FConstrainedMeshDeformer(const MeshT& Mesh);
 	
-	bool Deform(TArray<FVector3d>& PositionBuffer) override;
+	UE_API bool Deform(TArray<FVector3d>& PositionBuffer) override;
 
 private:
 
@@ -73,13 +75,13 @@ FConstrainedMeshDeformer::FConstrainedMeshDeformer(const MeshT& Mesh) :
  * Clamped Cotangent weights with Voronoi area are used.
  * Boundary Vertices are *not* fixed, they are included in the system and so should have soft constraints set similar to any other vertex
  */
-class DYNAMICMESH_API FSoftMeshDeformer : public FSoftMeshDeformationSolver
+class FSoftMeshDeformer : public FSoftMeshDeformationSolver
 {
 public:
-	FSoftMeshDeformer(const FDynamicMesh3& DynamicMesh);
+	UE_API FSoftMeshDeformer(const FDynamicMesh3& DynamicMesh);
 	~FSoftMeshDeformer() override {}
 
-	bool Deform(TArray<FVector3d>& PositionBuffer) override;
+	UE_API bool Deform(TArray<FVector3d>& PositionBuffer) override;
 
 private:
 
@@ -89,4 +91,6 @@ private:
 
 
 } // end namespace UE::Geometry
-} // end namespace UE
\ No newline at end of file
+} // end namespace UE
+
+#undef UE_API
diff --git a/Engine/Plugins/Runtime/GeometryProcessing/Source/DynamicMesh/Private/Solvers/Internal/MeshUVSolver.h b/Engine/Plugins/Runtime/GeometryProcessing/Source/DynamicMesh/Private/Solvers/Internal/MeshUVSolver.h
--- a/Engine/Plugins/Runtime/GeometryProcessing/Source/DynamicMesh/Private/Solvers/Internal/MeshUVSolver.h
+++ b/Engine/Plugins/Runtime/GeometryProcessing/Source/DynamicMesh/Private/Solvers/Internal/MeshUVSolver.h
@@ -12,6 +12,8 @@
 #include "DenseMatrix.h"
 #include "SparseMatrix.h"
 
+#define UE_API DYNAMICMESH_API
+
 namespace UE
 {
 namespace Geometry
@@ -20,31 +22,31 @@ namespace Geometry
 /**
  * FConstrainedMeshUVSolver is an abstract class that implements methods for setting vertex constraints and their weights.
  */
-class DYNAMICMESH_API FConstrainedMeshUVSolver : public UE::Solvers::IConstrainedMeshUVSolver
+class FConstrainedMeshUVSolver : public UE::Solvers::IConstrainedMeshUVSolver
 {
 public:
 	typedef UE::Solvers::FUVConstraint  FUVConstraint;
 	
-	FConstrainedMeshUVSolver(const FDynamicMesh3& DynamicMesh);
+	UE_API FConstrainedMeshUVSolver(const FDynamicMesh3& DynamicMesh);
 
 	virtual ~FConstrainedMeshUVSolver() 
 	{
 	};
 
 	// Add constraint associated with given vertex id.
-	void AddConstraint(const int32 VtxId, const double Weight, const FVector2d& Pos, const bool bPostFix) override;
+	UE_API void AddConstraint(const int32 VtxId, const double Weight, const FVector2d& Pos, const bool bPostFix) override;
 
 	// Update the position of an existing constraint. Bool return if a corresponding constraint weight exists.
-	bool UpdateConstraintPosition(const int32 VtxId, const FVector2d& Position, const bool bPostFix) override;
+	UE_API bool UpdateConstraintPosition(const int32 VtxId, const FVector2d& Position, const bool bPostFix) override;
 
 	// The underlying solver will have to refactor the matrix if this is done. Bool return if a corresponding constraint position exists.
-	bool UpdateConstraintWeight(const int32 VtxId, const double Weight) override;
+	UE_API bool UpdateConstraintWeight(const int32 VtxId, const double Weight) override;
 
 	// Clear all constraints associated with this uv solver.
-	void ClearConstraints() override;
+	UE_API void ClearConstraints() override;
 
 	// Test if for constraint associated with given vertex id. 
-	bool IsConstrained(const int32 VtxId) const override;
+	UE_API bool IsConstrained(const int32 VtxId) const override;
 
 protected:
 
@@ -67,7 +69,7 @@ protected:
 /**
  * FConformalMeshUVSolver is an abstract class that implements shared methods for the constrained conformal UV solvers.
  */
-class DYNAMICMESH_API FConformalMeshUVSolver : public FConstrainedMeshUVSolver
+class FConformalMeshUVSolver : public FConstrainedMeshUVSolver
 {
 public:
 	
@@ -94,7 +96,7 @@ public:
 	 * @param OutCotangentMatrix 	(optional) If not nullptr, then store the cotangent Laplacian matrix.
 	 * @param OutAreaMatrix 		(optional) If not nullptr, then store the vector area matrix.
 	 */
-	void ConstructConformalEnergyMatrix(FSparseMatrixD& OutConfromalEnergy,
+	UE_API void ConstructConformalEnergyMatrix(FSparseMatrixD& OutConfromalEnergy,
 										const bool bInAreaWeighted = false, 
 										const double InSmallTriangleArea = SMALL_NUMBER,
 										FSparseMatrixD* OutCotangentMatrix = nullptr,
@@ -116,7 +118,7 @@ protected:
 	 * 
 	 * @param OutAreaMatrix The signed vector area matrix.
 	 */
-	void ConstructVectorAreaMatrix(FSparseMatrixD& OutAreaMatrix);
+	UE_API void ConstructVectorAreaMatrix(FSparseMatrixD& OutAreaMatrix);
 
 	
 	/**
@@ -137,7 +139,7 @@ protected:
 	 * @param InSmallTriangleArea  The "floor" for the triangle area. Triangles with areas below this value are  
 	 * 						   	   considered to be triangles with SmallTriangleArea area.
 	 */
-	void ConstructWeightedVectorAreaMatrix(FSparseMatrixD& OutAreaMatrix, 
+	UE_API void ConstructWeightedVectorAreaMatrix(FSparseMatrixD& OutAreaMatrix, 
 										   double& OutMaxArea,
 										   const double InSmallTriangleArea = SMALL_NUMBER); 
 
@@ -150,7 +152,7 @@ protected:
  * 
  * However, the energy is constructed using the formulation presented in the Spectral Conformal parameterization [Mullen 2008].
  */
-class DYNAMICMESH_API FLeastSquaresConformalMeshUVSolver : public FConformalMeshUVSolver
+class FLeastSquaresConformalMeshUVSolver : public FConformalMeshUVSolver
 {
 public:
 		
@@ -164,7 +166,7 @@ public:
 	{
 	};
 
-	virtual bool SolveUVs(const FDynamicMesh3* InMesh, TArray<FVector2d>& OutUVBuffer) override;
+	UE_API virtual bool SolveUVs(const FDynamicMesh3* InMesh, TArray<FVector2d>& OutUVBuffer) override;
 
 protected:
 	
@@ -173,7 +175,7 @@ protected:
 	 * matrix (created with ConstructConformalEnergyMatrix()) but with the rows corresponding to the constraints 
 	 * (defined by ConstraintMap) to only have one non-zero value along the diagonal (set to one).
 	 */
-	void ConstructSystemMatrices(FSparseMatrixD& OutSystemMatrix);
+	UE_API void ConstructSystemMatrices(FSparseMatrixD& OutSystemMatrix);
 	
 	/**
 	 * Solve the natural/free-boundary conformal parameterization problem defined by the given system matrix
@@ -186,7 +188,7 @@ protected:
 	 * 
 	 * @return true if the solve succeded, false otherwise.
 	 */
-	bool SolveParameterization(const FSparseMatrixD& InSystemMatrix, 
+	UE_API bool SolveParameterization(const FSparseMatrixD& InSystemMatrix, 
 							   const EMatrixSolverType InMatrixSolverType, 
 							   FColumnVectorD& OutSolution);
 };
@@ -196,7 +198,7 @@ protected:
  * FConformalMeshUVSolver computes the conformal UV map using the Spectral Conformal Parameterization [Mullen 2008] 
  * method including enforcing insensitivity to sampling irregularity extension.
  */
-class DYNAMICMESH_API FSpectralConformalMeshUVSolver : public FConformalMeshUVSolver
+class FSpectralConformalMeshUVSolver : public FConformalMeshUVSolver
 {
 public:
 
@@ -210,7 +212,7 @@ public:
 	{
 	};
 
-	virtual bool SolveUVs(const FDynamicMesh3* InMesh, TArray<FVector2d>& OutUVBuffer) override;
+	UE_API virtual bool SolveUVs(const FDynamicMesh3* InMesh, TArray<FVector2d>& OutUVBuffer) override;
 
 	
 	//
@@ -231,7 +233,7 @@ public:
 	 * 
 	 * @note resets RandomStream optional variable to avoid confusion as to which one to use.
 	 */ 
-	void SetInitialSolution(const TArray<FVector2d>& InInitialUVs);
+	UE_API void SetInitialSolution(const TArray<FVector2d>& InInitialUVs);
  
 	/**
 	 * @param InRandomStream  The stream will be used to sample random values for the initial UV solution. 
@@ -239,7 +241,7 @@ public:
 	 *
 	 * @note resets InitialUVs optional variable to avoid confusion as to which one to use.
 	 */
-	void SetInitialSolution(const FRandomStream& InRandomStream);
+	UE_API void SetInitialSolution(const FRandomStream& InRandomStream);
 
 protected: 
 
@@ -253,7 +255,7 @@ protected:
 	 * 					 		 (resp.,v-coordinate) of the boundary vertices (not including any of the internal 
 	 * 							 boundaries) and 0 otherwise. V_b is the number of the boundary vertices.
 	 */
-	void ConstructSystemMatrices(FSparseMatrixD& OutConformalEnergy, 
+	UE_API void ConstructSystemMatrices(FSparseMatrixD& OutConformalEnergy, 
 								 FSparseMatrixD& OutMatrixB, 
 								 FSparseMatrixD& OutMatrixE);
 
@@ -261,7 +263,7 @@ protected:
 	 * Solve the general eigenvalue problem to compute the UV mapping using the matrices computed in 
 	 * ConstructSystemMatrices().
 	 */
-	bool SolveParameterization(const FSparseMatrixD& InConformalEnergy,
+	UE_API bool SolveParameterization(const FSparseMatrixD& InConformalEnergy,
 							   const FSparseMatrixD& InMatrixB,
 							   const FSparseMatrixD& InMatrixE,
 							   FColumnVectorD& OutSolution);
@@ -280,4 +282,6 @@ protected:
 
 
 } // end namespace UE::Geometry
-} // end namespace UE
\ No newline at end of file
+} // end namespace UE
+
+#undef UE_API
diff --git a/Engine/Plugins/Runtime/GeometryProcessing/Source/DynamicMesh/Private/Solvers/Internal/PowerMethodSolver.h b/Engine/Plugins/Runtime/GeometryProcessing/Source/DynamicMesh/Private/Solvers/Internal/PowerMethodSolver.h
--- a/Engine/Plugins/Runtime/GeometryProcessing/Source/DynamicMesh/Private/Solvers/Internal/PowerMethodSolver.h
+++ b/Engine/Plugins/Runtime/GeometryProcessing/Source/DynamicMesh/Private/Solvers/Internal/PowerMethodSolver.h
@@ -5,6 +5,8 @@
 #include "CoreMinimal.h"
 #include "MatrixSolver.h"
 
+#define UE_API DYNAMICMESH_API
+
 namespace UE
 {
 namespace Geometry
@@ -22,7 +24,7 @@ namespace Geometry
  * 
  * For sparse matrices see FSparsePowerMethod.
  */
-class DYNAMICMESH_API FPowerMethod
+class FPowerMethod
 {
 public:
 	using ScalarType = typename IMatrixSolverBase::ScalarType;
@@ -97,7 +99,7 @@ public:
 	 * @param OutEigenVector Computed eigenvector corresponding to the largest (smallest) eigen value.
 	 * @param bComputeLargest If true (false) will compute the eigen pair corresponding to the largest (smallest) eigenvalue. 
 	 */
-	virtual ESolveStatus Solve(ScalarType& OutEigenValue, RealVectorType& OutEigenVector, bool bComputeLargest = false);
+	UE_API virtual ESolveStatus Solve(ScalarType& OutEigenValue, RealVectorType& OutEigenVector, bool bComputeLargest = false);
 
 	/** The tolerance value upon reaching convergence or upon hitting the max iteration. */
 	inline ScalarType GetConvergedTolerance() 
@@ -113,10 +115,10 @@ public:
 
 protected:
 	/** Checks if we reached the convergence condition ||A*x - lambda*B*x||_Inf < Tolerance */
-	virtual bool Converged(ScalarType Lambda, const RealVectorType& EVector);
+	UE_API virtual bool Converged(ScalarType Lambda, const RealVectorType& EVector);
 	
 	/** Compute the eigenvalue corresponding to the given eigenvector using the Rayleigh Quotient: lambda = x'*A*x/x'*B*x */
-	virtual ScalarType RayleighQuotient(const RealVectorType& NormalizedEVector) const;
+	UE_API virtual ScalarType RayleighQuotient(const RealVectorType& NormalizedEVector) const;
 	
 	/** Normalize the input vector. Default is L2 normalization. */
 	virtual void Normalize(RealVectorType& EVector) const 
@@ -125,16 +127,16 @@ protected:
 	}
 
 	/** Single power iteration that computes the next approximation for the eigenvector.*/
-	virtual bool Iteration(RealVectorType& EVector, bool bComputeLargest);
+	UE_API virtual bool Iteration(RealVectorType& EVector, bool bComputeLargest);
 
 	/** 
 	 * In some cases we need to solve a linear system on every iteration. This method is called outside of the loop 
 	 * to pre-factorize any matrices used in the solve.
 	 */
-	virtual bool SetupSolver(bool bComputeLargest);
+	UE_API virtual bool SetupSolver(bool bComputeLargest);
 
 	/** Check that the input user parameters are valid and matrix operators are set */
-	virtual bool VerifyUserParameters(bool bComputeLargest) const;
+	UE_API virtual bool VerifyUserParameters(bool bComputeLargest) const;
 
 	/** Is matrix B not an identity matrix*/
 	inline bool IsGeneralProblem() const 
@@ -182,7 +184,7 @@ protected:
  * @note The MatrixOperator TFunctions capture the passed-in matrices by reference, so the lifetime of these matrices 
  * must extend through the use of the FSparsePowerMethod class methods (specifically the Solve() method).
  */
-class DYNAMICMESH_API FSparsePowerMethod : public FPowerMethod
+class FSparsePowerMethod : public FPowerMethod
 {
 public:
 
@@ -239,10 +241,12 @@ public:
      * @note InMatrix is captured by reference by all TFunctions and hence we assume that the lifetime of the
      * referenced matrix extends past the use of the OutMatrixOp.
      */
-	static void CreateSparseMatrixOperator(const FSparseMatrixD& InMatrix, 
+	static UE_API void CreateSparseMatrixOperator(const FSparseMatrixD& InMatrix, 
 										   const FMatrixHints& InHints, 
 										   MatrixOperator& OutMatrixOp);
 };
 
 } // end namespace UE::Geometry
-} // end namespace UE
\ No newline at end of file
+} // end namespace UE
+
+#undef UE_API
diff --git a/Engine/Plugins/Runtime/GeometryProcessing/Source/DynamicMesh/Private/Solvers/Internal/QuadraticProgramming.h b/Engine/Plugins/Runtime/GeometryProcessing/Source/DynamicMesh/Private/Solvers/Internal/QuadraticProgramming.h
--- a/Engine/Plugins/Runtime/GeometryProcessing/Source/DynamicMesh/Private/Solvers/Internal/QuadraticProgramming.h
+++ b/Engine/Plugins/Runtime/GeometryProcessing/Source/DynamicMesh/Private/Solvers/Internal/QuadraticProgramming.h
@@ -8,6 +8,8 @@
 #include "DenseMatrix.h"
 #include "MatrixBase.h"
 
+#define UE_API DYNAMICMESH_API
+
 namespace UE
 {
 namespace Geometry
@@ -61,42 +63,42 @@ namespace Geometry
  *          Linear Inequality Constraints:
  */
 
-class DYNAMICMESH_API FQuadraticProgramming
+class FQuadraticProgramming
 {
 public:
 
-    FQuadraticProgramming(const FSparseMatrixD* InMatrixQ, const FColumnVectorD* InVectorF = nullptr);
+    UE_API FQuadraticProgramming(const FSparseMatrixD* InMatrixQ, const FColumnVectorD* InVectorF = nullptr);
 
     /** If true, will solve for each column x \in X in parallel. By default, parallelization is enabled. */
-    void SetEnableParallelization(const bool bEnable);
+    UE_API void SetEnableParallelization(const bool bEnable);
 
     /** Set the fixed constraints either via a sparse or a dense matrix. Calls are mutually exclusive. */
-    bool SetFixedConstraints(const TArray<int>* InFixedRowIndices, const FSparseMatrixD* InFixedValues);
-    bool SetFixedConstraints(const TArray<int>* InFixedRowIndices, const FDenseMatrixD* InFixedValues);
+    UE_API bool SetFixedConstraints(const TArray<int>* InFixedRowIndices, const FSparseMatrixD* InFixedValues);
+    UE_API bool SetFixedConstraints(const TArray<int>* InFixedRowIndices, const FDenseMatrixD* InFixedValues);
 
     /** Pre-factorizes the matrices and sets up the solver. */
-    bool PreFactorize();
+    UE_API bool PreFactorize();
     
     /**
      * @param Solution The result of the solve as a dense matrix.
 	 * @param bVariablesOnly If true, the Solution matrix only contains the rows for non-fixed (variable) rows (set via SetFixedConstraints function).
      */
-    bool Solve(FDenseMatrixD& Solution, const bool bVariablesOnly = false);
+    UE_API bool Solve(FDenseMatrixD& Solution, const bool bVariablesOnly = false);
 
     /**
      * @param Solution The result of the solve as a sparse matrix.
 	 * @param bVariablesOnly If true, the Solution matrix only contains the rows for non-fixed (variable) rows (set via SetFixedConstraints function).
      * @param Tolerance Convert dense solution to sparse by pruning any values with their absolute value below the Tolerance.
      */
-    bool Solve(FSparseMatrixD& Solution, const bool bVariablesOnly = false, const double Tolerance = SMALL_NUMBER);
+    UE_API bool Solve(FSparseMatrixD& Solution, const bool bVariablesOnly = false, const double Tolerance = SMALL_NUMBER);
 
 
     // 
     // Helper "one-function call" methods for setting up and solving the QP problems
     //
 
-    static bool SolveWithFixedConstraints(const FSparseMatrixD& MatrixQ, const FColumnVectorD* VectorF, const TArray<int>& FixedRowIndices, const FSparseMatrixD& FixedValues, FDenseMatrixD& Solution, const bool bVariablesOnly = false, TArray<int>* VariableRowIndices = nullptr);
-    static bool SolveWithFixedConstraints(const FSparseMatrixD& MatrixQ, const FColumnVectorD* VectorF, const TArray<int>& FixedRowIndices, const FSparseMatrixD& FixedValues, FSparseMatrixD& Solution, const bool bVariablesOnly = false, const double Tolerance = SMALL_NUMBER, TArray<int>* VariableRowIndices = nullptr);
+    static UE_API bool SolveWithFixedConstraints(const FSparseMatrixD& MatrixQ, const FColumnVectorD* VectorF, const TArray<int>& FixedRowIndices, const FSparseMatrixD& FixedValues, FDenseMatrixD& Solution, const bool bVariablesOnly = false, TArray<int>* VariableRowIndices = nullptr);
+    static UE_API bool SolveWithFixedConstraints(const FSparseMatrixD& MatrixQ, const FColumnVectorD* VectorF, const TArray<int>& FixedRowIndices, const FSparseMatrixD& FixedValues, FSparseMatrixD& Solution, const bool bVariablesOnly = false, const double Tolerance = SMALL_NUMBER, TArray<int>* VariableRowIndices = nullptr);
 
 
     //
@@ -104,10 +106,10 @@ public:
     // 
 
     /** How long did the last call to Solve() take in seconds. */
-    double GetSolveTimeElapsedInSec() const;
+    UE_API double GetSolveTimeElapsedInSec() const;
 
     /** Array of row indices into MatrixQ representing variables (non-fixed). */
-    TArray<int> GetVariableRowIndices() const;
+    UE_API TArray<int> GetVariableRowIndices() const;
 
 protected:
 
@@ -136,3 +138,5 @@ protected:
 
 }
 }
+
+#undef UE_API
diff --git a/Engine/Plugins/Runtime/GeometryProcessing/Source/DynamicMesh/Public/FaceGroupUtil.h b/Engine/Plugins/Runtime/GeometryProcessing/Source/DynamicMesh/Public/FaceGroupUtil.h
--- a/Engine/Plugins/Runtime/GeometryProcessing/Source/DynamicMesh/Public/FaceGroupUtil.h
+++ b/Engine/Plugins/Runtime/GeometryProcessing/Source/DynamicMesh/Public/FaceGroupUtil.h
@@ -9,6 +9,8 @@
 #include "DynamicMeshEditor.h"
 #include "Polygroups/PolygroupSet.h"
 
+#define UE_API DYNAMICMESH_API
+
 /**
  * Utility functions for dealing with Mesh FaceGroups
  */
@@ -82,7 +84,7 @@ namespace Geometry
  * FGroupVisualizationCache is intended to be used as a cache for per-group information that 
  * would be used to visualize PolyGroups, in particular drawing the Polygroup ID for each group.
  */
-class DYNAMICMESH_API FGroupVisualizationCache
+class FGroupVisualizationCache
 {
 public:
 
@@ -113,7 +115,7 @@ public:
 	 * than 2 triangles it will be "eroded" down a single triangle to determine the group "center", also somewhat expensive.
 	 * This step can be computed in parallel if bParallel is true.
 	 */
-	void UpdateGroupInfo_ConnectedComponents(
+	UE_API void UpdateGroupInfo_ConnectedComponents(
 		const FDynamicMesh3& SourceMesh,
 		const FPolygroupSet& GroupSet,
 		bool bParallel = true);
@@ -132,3 +134,5 @@ public:
 
 } // end namespace UE::Geometry
 } // end namespace UE
+
+#undef UE_API
diff --git a/Engine/Plugins/Runtime/GeometryProcessing/Source/DynamicMesh/Public/GroupTopology.h b/Engine/Plugins/Runtime/GeometryProcessing/Source/DynamicMesh/Public/GroupTopology.h
--- a/Engine/Plugins/Runtime/GeometryProcessing/Source/DynamicMesh/Public/GroupTopology.h
+++ b/Engine/Plugins/Runtime/GeometryProcessing/Source/DynamicMesh/Public/GroupTopology.h
@@ -8,6 +8,8 @@
 #include "Util/IndexUtil.h"
 #include "Containers/BitArray.h"
 
+#define UE_API DYNAMICMESH_API
+
 namespace UE
 {
 namespace Geometry
@@ -18,7 +20,7 @@ class FGroupTopology;
 /**
  * FGroupTopologySelection represents a set of selected elements of a FGroupTopology
  */
-struct DYNAMICMESH_API FGroupTopologySelection
+struct FGroupTopologySelection
 {
 	TSet<int32> SelectedGroupIDs;
 	TSet<int32> SelectedCornerIDs;
@@ -146,12 +148,12 @@ private:
  * indices into the internal Corners and Edges arrays. However references to "GroupID"
  * refer to the actual group IDs in FDynamicMesh (not an index into Groups).
  */
-class DYNAMICMESH_API FGroupTopology
+class FGroupTopology
 {
 public:
 	FGroupTopology() {}
-	FGroupTopology(const FDynamicMesh3* Mesh, bool bAutoBuild);
-	FGroupTopology(const FDynamicMesh3* Mesh, const FDynamicMeshPolygroupAttribute* GroupLayer, bool bAutoBuild);
+	UE_API FGroupTopology(const FDynamicMesh3* Mesh, bool bAutoBuild);
+	UE_API FGroupTopology(const FDynamicMesh3* Mesh, const FDynamicMeshPolygroupAttribute* GroupLayer, bool bAutoBuild);
 
 	//~ If we allow deprecation warnings, we'll get a warning in destructor and assignment operator 
 	//~ about CornerVerticesFlags (which we deprecated here) whenever it's copied or destroyed.
@@ -166,7 +168,7 @@ public:
 	 * a new cloned mesh. If the given mesh is not cloned, the topology will no
 	 * longer be consistent.
 	 */
-	void RetargetOnClonedMesh(const FDynamicMesh3* Mesh);
+	UE_API void RetargetOnClonedMesh(const FDynamicMesh3* Mesh);
 
 	const FDynamicMesh3* GetMesh() const 
 	{ 
@@ -177,7 +179,7 @@ public:
 	 * Build the group topology graph.
 	 * @return true on success, false if there was an error
 	 */
-	virtual bool RebuildTopology();
+	UE_API virtual bool RebuildTopology();
 
 	/**
 	 * Optional function that will add extra corners (in addition to the ones that are always placed at junctures
@@ -200,7 +202,7 @@ public:
 	 * the topology with a particular group of extra corners during client undo/redo operations where the behavior
 	 * of ShouldAddExtraCornerAtVert may have changed in the interim.
 	 */
-	virtual bool RebuildTopologyWithSpecificExtraCorners(const TSet<int32>& ExtraCornerVids);
+	UE_API virtual bool RebuildTopologyWithSpecificExtraCorners(const TSet<int32>& ExtraCornerVids);
 
 	/** 
 	 * Adjacency of Per-Triangle integers are what define the triangle groups.
@@ -217,7 +219,7 @@ public:
 	/**
 	 * FCorner is a "corner" in the group topology, IE a vertex where 3+ groups meet.
 	 */
-	struct DYNAMICMESH_API FCorner
+	struct FCorner
 	{
 		/** Mesh Vertex ID */
 		int VertexID;		
@@ -229,7 +231,7 @@ public:
 	TArray<FCorner> Corners;
 
 	/** A Group is bounded by closed loops of FGroupEdge elements. A FGroupBoundary is one such loop. */
-	struct DYNAMICMESH_API FGroupBoundary
+	struct FGroupBoundary
 	{
 		/** Ordered list of edges forming this boundary */
 		TArray<int> GroupEdges;
@@ -240,7 +242,7 @@ public:
 	};
 
 	/** FGroup is a set of connected triangles with the same GroupID */
-	struct DYNAMICMESH_API FGroup
+	struct FGroup
 	{
 		/** GroupID for this group */
 		int GroupID;
@@ -263,7 +265,7 @@ public:
 	 * 
 	 * FGroupEdge instances are *shared* between the two FGroups on either side, via FGroupBoundary.
 	 */
-	struct DYNAMICMESH_API FGroupEdge
+	struct FGroupEdge
 	{
 		/** Groups adjoinging this edge, sorted by increasing value if both present, otherwise the second element may be InvalidID (if boundary edge). */
 		FIndex2i Groups;
@@ -281,7 +283,7 @@ public:
 		}
 
 		/** @return true if any vertex in the Span is in the Vertices set*/
-		bool IsConnectedToVertices(const TSet<int>& Vertices) const;
+		UE_API bool IsConnectedToVertices(const TSet<int>& Vertices) const;
 	};
 
 	/** List of Edges in the topology graph (each edge connects two corners/nodes) */
@@ -290,77 +292,77 @@ public:
 
 
 	/** @return the mesh vertex ID for the given Corner ID */
-	int GetCornerVertexID(int CornerID) const;
+	UE_API int GetCornerVertexID(int CornerID) const;
 
 	/** @return the Group ID for a given Vertex ID, or InvalidID if vertex is not a group corner */
-	int32 GetCornerIDFromVertexID(int32 VertexID) const;
+	UE_API int32 GetCornerIDFromVertexID(int32 VertexID) const;
 
 	/** @return the FGroup for the given GroupID, or nullptr if not found */
-	const FGroup* FindGroupByID(int GroupID) const;
+	UE_API const FGroup* FindGroupByID(int GroupID) const;
 	/** @return the list of triangles in the given GroupID, or empty list if not found */
 	const TArray<int>& GetGroupFaces(int GroupID) const { return GetGroupTriangles(GroupID); }
 	/** @return the list of triangles in the given GroupID, or empty list if not found */
-	const TArray<int>& GetGroupTriangles(int GroupID) const;
+	UE_API const TArray<int>& GetGroupTriangles(int GroupID) const;
 	/** @return the list of neighbour GroupIDs for the given GroupID, or empty list if not found */
-	const TArray<int>& GetGroupNbrGroups(int GroupID) const;
+	UE_API const TArray<int>& GetGroupNbrGroups(int GroupID) const;
 
 	/** @return true if the mesh edge identified by FMeshTriEdgeID is a group boundary edge */
-	bool IsGroupEdge(FMeshTriEdgeID TriEdgeID, bool bIncludeMeshBoundary = false) const; 
+	UE_API bool IsGroupEdge(FMeshTriEdgeID TriEdgeID, bool bIncludeMeshBoundary = false) const; 
 	/** @return the ID of the FGroupEdge that contains the given Mesh Edge ID*/
-	int FindGroupEdgeID(int MeshEdgeID) const;
+	UE_API int FindGroupEdgeID(int MeshEdgeID) const;
 	/** @return the ID of the FGroupEdge that contains the given Mesh Edge ID */
-	int FindGroupEdgeID(FMeshTriEdgeID FMeshTriEdgeID) const;
+	UE_API int FindGroupEdgeID(FMeshTriEdgeID FMeshTriEdgeID) const;
 	/** @return the list of vertices of a FGroupEdge identified by the GroupEdgeID */
-	const TArray<int>& GetGroupEdgeVertices(int GroupEdgeID) const;
+	UE_API const TArray<int>& GetGroupEdgeVertices(int GroupEdgeID) const;
 	/** @return the list of edges of a FGroupEdge identified by the GroupEdgeID */
-	const TArray<int>& GetGroupEdgeEdges(int GroupEdgeID) const;
+	UE_API const TArray<int>& GetGroupEdgeEdges(int GroupEdgeID) const;
 
 	/** Add the groups connected to the given GroupEdgeID to the GroupsOut list. This is not the either-side pair, but the set of groups on the one-ring of each connected corner. */
-	void FindEdgeNbrGroups(int GroupEdgeID, TArray<int>& GroupsOut) const;
+	UE_API void FindEdgeNbrGroups(int GroupEdgeID, TArray<int>& GroupsOut) const;
 	/** Add the groups connected to all the GroupEdgeIDs to the GroupsOut list. This is not the either-side pair, but the set of groups on the one-ring of each connected corner. */
-	void FindEdgeNbrGroups(const TArray<int>& GroupEdgeIDs, TArray<int>& GroupsOut) const;
+	UE_API void FindEdgeNbrGroups(const TArray<int>& GroupEdgeIDs, TArray<int>& GroupsOut) const;
 
 	/** Add the edges connected to the given GroupEdgeID to the EdgesOut list. This is somewhat expensive. */
-	void FindEdgeNbrEdges(int GroupEdgeID, TArray<int>& EdgesOut) const;
+	UE_API void FindEdgeNbrEdges(int GroupEdgeID, TArray<int>& EdgesOut) const;
 
 	/** @return true if the FGroupEdge identified by the GroupEdgeID is on the mesh boundary */
-	bool IsBoundaryEdge(int32 GroupEdgeID) const;
+	UE_API bool IsBoundaryEdge(int32 GroupEdgeID) const;
 	/** @return true if group edge is a "simple" edge, meaning it only has one mesh edge */
-	bool IsSimpleGroupEdge(int32 GroupEdgeID) const;
+	UE_API bool IsSimpleGroupEdge(int32 GroupEdgeID) const;
 
 	/** @return true if the FGroupEdge identified by the GroupEdgeID is an isolated loop (ie not connected to any other edges) */
-	bool IsIsolatedLoop(int32 GroupEdgeID) const;
+	UE_API bool IsIsolatedLoop(int32 GroupEdgeID) const;
 
 	/** @return arc length of edge, and optionally accumulated arclength distances for each edge vertex */
-	double GetEdgeArcLength(int32 GroupEdgeID, TArray<double>* PerVertexLengthsOut = nullptr) const;
+	UE_API double GetEdgeArcLength(int32 GroupEdgeID, TArray<double>* PerVertexLengthsOut = nullptr) const;
 
 	/** @return arclength midpoint of edge, and optionally arclength info */
-	FVector3d GetEdgeMidpoint(int32 GroupEdgeID, double* ArcLengthOut = nullptr, TArray<double>* PerVertexLengthsOut = nullptr) const;
+	UE_API FVector3d GetEdgeMidpoint(int32 GroupEdgeID, double* ArcLengthOut = nullptr, TArray<double>* PerVertexLengthsOut = nullptr) const;
 
 	/** @return average position of edge (may not be on edge for curved edges) */
-	FVector3d GetEdgeAveragePosition(int32 GroupEdgeID) const;
+	UE_API FVector3d GetEdgeAveragePosition(int32 GroupEdgeID) const;
 
 	/** Add all the groups connected to the given Corner to the GroupsOut list */
-	void FindCornerNbrGroups(int CornerID, TArray<int>& GroupsOut) const;
+	UE_API void FindCornerNbrGroups(int CornerID, TArray<int>& GroupsOut) const;
 	/** Add all the groups connected to the given Corners to the GroupsOut list */
-	void FindCornerNbrGroups(const TArray<int>& CornerIDs, TArray<int>& GroupsOut) const;
+	UE_API void FindCornerNbrGroups(const TArray<int>& CornerIDs, TArray<int>& GroupsOut) const;
 
 	/**
 	 * Helper that iterates over the group edges neighboring a particular corner.
 	 * @param ReturnTrueToContinue Function to call on all the group edges, which should return
 	 *  true if it wants iteration to continue, and false if it wants to exit early.
 	 */
-	void ForCornerNbrEdges(int CornerID, TFunctionRef<bool(int32 EdgeID)> ReturnTrueToContinue) const;
+	UE_API void ForCornerNbrEdges(int CornerID, TFunctionRef<bool(int32 EdgeID)> ReturnTrueToContinue) const;
 
 	/** Add all the Edges connected to the given Corner to the EdgesOut list. This is somewhat expensive as there is no connection from Corners to Edges, must iterate over adjacent faces */
-	void FindCornerNbrEdges(int CornerID, TArray<int>& EdgesOut) const;
+	UE_API void FindCornerNbrEdges(int CornerID, TArray<int>& EdgesOut) const;
 	/** Add all the Corners connected to the given Corner to the EdgesOut list. This is somewhat expensive as there is no connection from Corners to Corners, must iterate over adjacent faces/edges */
-	void FindCornerNbrCorners(int CornerID, TArray<int>& CornersOut) const;
+	UE_API void FindCornerNbrCorners(int CornerID, TArray<int>& CornersOut) const;
 
 	/** Add all the groups connected to the given Mesh Vertex to the GroupsOut list */
-	void FindVertexNbrGroups(int VertexID, TArray<int>& GroupsOut ) const;
+	UE_API void FindVertexNbrGroups(int VertexID, TArray<int>& GroupsOut ) const;
 	/** Add all the groups connected to the given Mesh Vertices to the GroupsOut list */
-	void FindVertexNbrGroups(const TArray<int>& VertexIDs, TArray<int>& GroupsOut) const;
+	UE_API void FindVertexNbrGroups(const TArray<int>& VertexIDs, TArray<int>& GroupsOut) const;
 
 	/**
 	 * Call EdgeFunc for each boundary edge of the given Group  (no order defined). 
@@ -417,36 +419,36 @@ public:
 	}
 
 	/** Add all the vertices of the given GroupID to the Vertices set */
-	void CollectGroupVertices(int GroupID, TSet<int>& Vertices) const;
+	UE_API void CollectGroupVertices(int GroupID, TSet<int>& Vertices) const;
 	/** Add all the group boundary vertices of the given GroupID to the Vertices set */
-	void CollectGroupBoundaryVertices(int GroupID, TSet<int>& Vertices) const;
+	UE_API void CollectGroupBoundaryVertices(int GroupID, TSet<int>& Vertices) const;
 
 	/** 
 	 * Calculate tangent of group edge, as direction from start to end endpoints
 	 * @return false if edge is degenerate or start == end 
 	 */
-	bool GetGroupEdgeTangent(int GroupEdgeID, FVector3d& TangentOut) const;
+	UE_API bool GetGroupEdgeTangent(int GroupEdgeID, FVector3d& TangentOut) const;
 
 	/**
 	 * @return a 3D frame for the given group. Based on centroid of triangle centroids, so does not necessarily lie on group surface.
 	 */
-	FFrame3d GetGroupFrame(int32 GroupID) const;
+	UE_API FFrame3d GetGroupFrame(int32 GroupID) const;
 
 	/**
 	 * @return a 3D frame for the given selection
 	 */
-	FFrame3d GetSelectionFrame(const FGroupTopologySelection& Selection, FFrame3d* InitialLocalFrame = nullptr) const;
+	UE_API FFrame3d GetSelectionFrame(const FGroupTopologySelection& Selection, FFrame3d* InitialLocalFrame = nullptr) const;
 
 	/**
 	 * @return A Bounding Box for the vertices of the current selection, with TransformFunc applied
 	 */
-	FAxisAlignedBox3d GetSelectionBounds(const FGroupTopologySelection& Selection, TFunctionRef<FVector3d(const FVector3d&)> TransformFunc) const;
+	UE_API FAxisAlignedBox3d GetSelectionBounds(const FGroupTopologySelection& Selection, TFunctionRef<FVector3d(const FVector3d&)> TransformFunc) const;
 
 
 	/**
 	 * Get the set of selected triangles for a given GroupTopologySelection
 	 */
-	void GetSelectedTriangles(const FGroupTopologySelection& Selection, TArray<int32>& Triangles) const;
+	UE_API void GetSelectedTriangles(const FGroupTopologySelection& Selection, TArray<int32>& Triangles) const;
 
 	/**
 	 * Helper function for use with ShouldAddExtraCornerAtVert. Checks that the dot product of the two mesh edge
@@ -456,7 +458,7 @@ public:
 	 * @param DotThreshold Threshold usually set equal to FMathd::Cos(ThresholdDegrees * FMathd::DegToRad). The
 	 *   parameter is not degrees just to avoid that calculation on each invocation.
 	 */ 
-	static bool IsEdgeAngleSharp(const FDynamicMesh3* Mesh, int32 SharedVid, const FIndex2i& AttachedGroupEdgeEids, double DotThreshold);
+	static UE_API bool IsEdgeAngleSharp(const FDynamicMesh3* Mesh, int32 SharedVid, const FIndex2i& AttachedGroupEdgeEids, double DotThreshold);
 
 protected:
 	const FDynamicMesh3* Mesh = nullptr;
@@ -477,7 +479,7 @@ protected:
 	 *  use the topology structures.
 	 */
 	UE_DEPRECATED(5.1, "Use ShouldVertBeCorner instead.")
-	virtual bool IsCornerVertex(int VertexID) const;
+	UE_API virtual bool IsCornerVertex(int VertexID) const;
 
 	/**
 	 * Return true if the given vertex should be a corner based on the number of adjoining group edges or based on
@@ -486,14 +488,14 @@ protected:
 	 * @param bIsExtraCornerOut Set to true if the vertex didn't have 3 or more group adjoining edges, but was forced
 	 *  to be a corner by ShouldAddExtraCornerAtVert.
 	 */
-	virtual bool ShouldVertBeCorner(int VertexID, bool* bIsExtraCornerOut = nullptr) const;
+	UE_API virtual bool ShouldVertBeCorner(int VertexID, bool* bIsExtraCornerOut = nullptr) const;
 
 	/** Return false if we had a problem finding group boundaries */
 	UE_DEPRECATED(5.1, "Use GenerateBoundaryAndGroupEdges instead.")
 	bool ExtractGroupEdges(FGroup& Group) { return false; }
 
 	// Generates group edges and boundary structs for given group. Returns false if failed to find group boundaries.
-	bool GenerateBoundaryAndGroupEdges(FGroup& Group, TMap<int32, int32>& GroupEdgeMinEidToGroupEdgeID, TBitArray<>& VertCheckedForCorner);
+	UE_API bool GenerateBoundaryAndGroupEdges(FGroup& Group, TMap<int32, int32>& GroupEdgeMinEidToGroupEdgeID, TBitArray<>& VertCheckedForCorner);
 
 	UE_DEPRECATED(5.1, "Use MakeEdgeGroupsPair instead. The returned pair cannot uniquely identify a group edge.")
 	inline FIndex2i MakeEdgeID(int MeshEdgeID) const { return MakeEdgeGroupsPair(MeshEdgeID); }
@@ -506,7 +508,7 @@ protected:
 	UE_DEPRECATED(5.1, "Group edges cannot be uniquely identified by a group pair.")
 	int FindExistingGroupEdge(int GroupID, int OtherGroupID, int FirstVertexID, int SecondVertexID) { return 0; }
 
-	void GetAllVertexGroups(int32 VertexID, TArray<int32>& GroupsOut) const;
+	UE_API void GetAllVertexGroups(int32 VertexID, TArray<int32>& GroupsOut) const;
 };
 
 FIndex2i FGroupTopology::MakeEdgeGroupsPair(int MeshEdgeID) const
@@ -545,13 +547,13 @@ bool FGroupTopologySelection::IsSelectedTriangle(const FDynamicMesh3* Mesh, cons
  * This allows algorithms to be written against FGroupTopology that will also work per-triangle.
  * (However, there is enormous overhead to doing it this way!)
  */
-class DYNAMICMESH_API FTriangleGroupTopology : public FGroupTopology
+class FTriangleGroupTopology : public FGroupTopology
 {
 public:
 	FTriangleGroupTopology() {}
-	FTriangleGroupTopology(const FDynamicMesh3* Mesh, bool bAutoBuild);
+	UE_API FTriangleGroupTopology(const FDynamicMesh3* Mesh, bool bAutoBuild);
 
-	virtual bool RebuildTopology() override;
+	UE_API virtual bool RebuildTopology() override;
 
 	virtual int GetGroupID(int TriangleID) const override
 	{
@@ -563,3 +565,4 @@ public:
 } // end namespace UE::Geometry
 } // end namespace UE
 
+#undef UE_API
diff --git a/Engine/Plugins/Runtime/GeometryProcessing/Source/DynamicMesh/Public/MeshConstraintsUtil.h b/Engine/Plugins/Runtime/GeometryProcessing/Source/DynamicMesh/Public/MeshConstraintsUtil.h
--- a/Engine/Plugins/Runtime/GeometryProcessing/Source/DynamicMesh/Public/MeshConstraintsUtil.h
+++ b/Engine/Plugins/Runtime/GeometryProcessing/Source/DynamicMesh/Public/MeshConstraintsUtil.h
@@ -8,6 +8,8 @@
 #include "MeshConstraints.h"
 #include "DynamicMesh/DynamicMeshAttributeSet.h"
 
+#define UE_API DYNAMICMESH_API
+
 
 namespace UE
 {
@@ -17,7 +19,7 @@ namespace Geometry
 /**
  * Utility functions for configuring a FMeshConstraints instance
  */
-class DYNAMICMESH_API FMeshConstraintsUtil
+class FMeshConstraintsUtil
 {
 public:
 
@@ -30,7 +32,7 @@ public:
 	 * @param bAllowSmoothing should we allow constrained vertices to be smoothed
 	 * @param bParallel should we run the algo in parallel
 	 */
-	static void ConstrainAllSeams(FMeshConstraints& Constraints, const FDynamicMesh3& Mesh, bool bAllowSplits, bool bAllowSmoothing, bool bParallel = true);
+	static UE_API void ConstrainAllSeams(FMeshConstraints& Constraints, const FDynamicMesh3& Mesh, bool bAllowSplits, bool bAllowSmoothing, bool bParallel = true);
 
 	/**
 	 * Constrain all attribute seams for all overlays of a mesh and different types of boundaries on the mesh
@@ -44,7 +46,7 @@ public:
 	 * @param bAllowSeamCollapse should we allow seam edges to collapse
 	 * @param bParallel should we run the algo in parallel
 	 */
-	static void ConstrainAllBoundariesAndSeams(FMeshConstraints& Constraints,
+	static UE_API void ConstrainAllBoundariesAndSeams(FMeshConstraints& Constraints,
 											   const FDynamicMesh3& Mesh,
 											   EEdgeRefineFlags MeshBoundaryConstraint,
 											   EEdgeRefineFlags GroupBoundaryConstraint,
@@ -61,7 +63,7 @@ s	 * EdgeROI that are seams, as well as vertices incident to edges in EdgeROI th
 	 * @param bAllowSplits should we allow constrained edges to be split
 	 * @param bAllowSmoothing should we allow constrained vertices to be smoothed
 	 */
-	static void ConstrainSeamsInEdgeROI(FMeshConstraints& Constraints, const FDynamicMesh3& Mesh, const TArray<int>& EdgeROI, bool bAllowSplits, bool bAllowSmoothing, bool bParallel = true);
+	static UE_API void ConstrainSeamsInEdgeROI(FMeshConstraints& Constraints, const FDynamicMesh3& Mesh, const TArray<int>& EdgeROI, bool bAllowSplits, bool bAllowSmoothing, bool bParallel = true);
 
 
 	/**
@@ -73,7 +75,7 @@ s	 * EdgeROI that are seams, as well as vertices incident to edges in EdgeROI th
 	 * @param bAllowSplits should we allow constrained edges to be split
 	 * @param bAllowSmoothing should we allow constrained vertices to be smoothed
 	 */
-	static void ConstrainROIBoundariesInEdgeROI(FMeshConstraints& Constraints, 
+	static UE_API void ConstrainROIBoundariesInEdgeROI(FMeshConstraints& Constraints, 
 		const FDynamicMesh3& Mesh, 
 		const TSet<int>& EdgeROI, 
 		const TSet<int>& TriangleROI,
@@ -161,7 +163,7 @@ s	 * EdgeROI that are seams, as well as vertices incident to edges in EdgeROI th
 	 * @param VertexConstraintB [out] vertex constraint to be updated on return
 	 * @return false if no constraint is needed for this edge
 	 */
-	static bool ConstrainEdgeBoundariesAndSeams(const int EdgeID,
+	static UE_API bool ConstrainEdgeBoundariesAndSeams(const int EdgeID,
 											    const FDynamicMesh3& Mesh,
 			                                    const EEdgeRefineFlags MeshBoundaryConstraint,
 			                                    const EEdgeRefineFlags GroupBoundaryConstraint,
@@ -188,7 +190,7 @@ s	 * EdgeROI that are seams, as well as vertices incident to edges in EdgeROI th
 	 * @param Mesh the mesh to constrain
 	 * @param CornerAngleThreshold angle threshold (in degrees) above which adjacent edges are considered a corner
 	 */
-	static void SetBoundaryConstraintsWithProjection(
+	static UE_API void SetBoundaryConstraintsWithProjection(
 		FMeshConstraints& Constraints,
 		const EBoundaryType BoundaryToConstrain,
 		const FDynamicMesh3& Mesh,
@@ -202,4 +204,6 @@ private:
 
 
 } // end namespace UE::Geometry
-} // end namespace UE
\ No newline at end of file
+} // end namespace UE
+
+#undef UE_API
diff --git a/Engine/Plugins/Runtime/GeometryProcessing/Source/DynamicMesh/Public/MeshCurvature.h b/Engine/Plugins/Runtime/GeometryProcessing/Source/DynamicMesh/Public/MeshCurvature.h
--- a/Engine/Plugins/Runtime/GeometryProcessing/Source/DynamicMesh/Public/MeshCurvature.h
+++ b/Engine/Plugins/Runtime/GeometryProcessing/Source/DynamicMesh/Public/MeshCurvature.h
@@ -5,6 +5,8 @@
 #include "DynamicMesh/DynamicMesh3.h"
 #include "Sampling/SampleSetStatistics.h"
 
+#define UE_API DYNAMICMESH_API
+
 namespace UE
 {
 	/**
@@ -59,7 +61,7 @@ namespace Geometry
  * FMeshVertexCurvatureCache calculates and stores various per-vertex Curvature types
  * for a Mesh, as well as some statistics for those values.
  */
-class DYNAMICMESH_API FMeshVertexCurvatureCache
+class FMeshVertexCurvatureCache
 {
 public:
 
@@ -81,8 +83,10 @@ public:
 	int32 Num() const { return Curvatures.Num(); }
 	const FVertexCurvature& operator[](int32 Index) const { return Curvatures[Index]; }
 
-	void BuildAll(const FDynamicMesh3& Mesh);
+	UE_API void BuildAll(const FDynamicMesh3& Mesh);
 };
 
 } // end namespace UE::Geometry
-} // end namespace UE
\ No newline at end of file
+} // end namespace UE
+
+#undef UE_API
diff --git a/Engine/Plugins/Runtime/GeometryProcessing/Source/DynamicMesh/Public/MeshPaths.h b/Engine/Plugins/Runtime/GeometryProcessing/Source/DynamicMesh/Public/MeshPaths.h
--- a/Engine/Plugins/Runtime/GeometryProcessing/Source/DynamicMesh/Public/MeshPaths.h
+++ b/Engine/Plugins/Runtime/GeometryProcessing/Source/DynamicMesh/Public/MeshPaths.h
@@ -5,6 +5,8 @@
 #include "DynamicMesh/DynamicMesh3.h"
 #include "EdgeSpan.h"
 
+#define UE_API DYNAMICMESH_API
+
 namespace UE
 {
 namespace Geometry
@@ -15,7 +17,7 @@ namespace Geometry
   * We do not distinguish betweens spans and loops -- if the input set of edges forms a loop, it will create an FEdgeSpan with those edges.
   * If the input edges do not form loops, then paths will begin/end at vertices incident on 1 or >2 of the given edges (e.g. "dead ends" or "intersections" of paths.)
   */
-class DYNAMICMESH_API FMeshPaths
+class FMeshPaths
 {
 public:
 
@@ -32,7 +34,7 @@ public:
 	TArray<FEdgeSpan> Spans;
 
 	// Find connected simple paths made up of edges in the input Edges structure
-	bool Compute();
+	UE_API bool Compute();
 
 private:
 
@@ -45,9 +47,11 @@ private:
 	 * Consider only edges incident on CrossedVertex, which is the vertex belonging to CurrentEdge that is not FromVertex.
 	 * Do not consider an adjacent edge if CrossedVertex is incident on more than 2 input Edges (i.e. we've hit a corner).
 	 */
-	int GetNextEdge(int CurrentEdge, int FromVertex, int& CrossedVertex) const;
+	UE_API int GetNextEdge(int CurrentEdge, int FromVertex, int& CrossedVertex) const;
 
 };
 
 } // end namespace UE::Geometry
 } // end namespace UE
+
+#undef UE_API
diff --git a/Engine/Plugins/Runtime/GeometryProcessing/Source/DynamicMesh/Public/MeshWeights.h b/Engine/Plugins/Runtime/GeometryProcessing/Source/DynamicMesh/Public/MeshWeights.h
--- a/Engine/Plugins/Runtime/GeometryProcessing/Source/DynamicMesh/Public/MeshWeights.h
+++ b/Engine/Plugins/Runtime/GeometryProcessing/Source/DynamicMesh/Public/MeshWeights.h
@@ -6,6 +6,8 @@
 
 #include "DynamicMesh/DynamicMesh3.h"
 
+#define UE_API DYNAMICMESH_API
+
 namespace UE
 {
 namespace Geometry
@@ -15,7 +17,7 @@ namespace Geometry
  * FMeshWeights implements various techniques for computing local weights of a mesh,
  * for example one-ring weights like Cotangent or Mean-Value.
  */
-class DYNAMICMESH_API FMeshWeights
+class FMeshWeights
 {
 public:
 
@@ -23,45 +25,45 @@ public:
 	 * Compute uniform centroid of a vertex one-ring.
 	 * These weights are strictly positive and all equal to 1 / valence
 	 */
-	static FVector3d UniformCentroid(const FDynamicMesh3& Mesh, int32 VertexIndex);
+	static UE_API FVector3d UniformCentroid(const FDynamicMesh3& Mesh, int32 VertexIndex);
 
 	/**
 	 * Compute uniform centroid of a vertex one-ring.
 	 * These weights are strictly positive and all equal to 1 / valence
 	 */
-	static FVector3d UniformCentroid(const FDynamicMesh3& Mesh, int32 VertexIndex, TFunctionRef<FVector3d(int32)> VertexPositionFunc );
+	static UE_API FVector3d UniformCentroid(const FDynamicMesh3& Mesh, int32 VertexIndex, TFunctionRef<FVector3d(int32)> VertexPositionFunc );
 
 	/**
 	 * Compute uniform centroid of a subset of vertex one-ring (eg boundary vertices)
 	 * These weights are strictly positive and all equal to 1 / valence
 	 */
-	static FVector3d FilteredUniformCentroid(const FDynamicMesh3& Mesh, int32 VertexIndex, TFunctionRef<FVector3d(int32)> VertexPositionFunc, TFunctionRef<bool(int32)> VertexFilterFunc);
+	static UE_API FVector3d FilteredUniformCentroid(const FDynamicMesh3& Mesh, int32 VertexIndex, TFunctionRef<FVector3d(int32)> VertexPositionFunc, TFunctionRef<bool(int32)> VertexFilterFunc);
 
 	/**
 	 * Compute mean-value centroid of a vertex one-ring.
 	 * These weights are strictly positive.
 	 */
-	static FVector3d MeanValueCentroid(const FDynamicMesh3& Mesh, int32 VertexIndex, double WeightClamp = FMathf::MaxReal);
+	static UE_API FVector3d MeanValueCentroid(const FDynamicMesh3& Mesh, int32 VertexIndex, double WeightClamp = FMathf::MaxReal);
 
 	/**
 	 * Compute mean-value centroid of a vertex one-ring.
 	 * These weights are strictly positive.
 	 */
-	static FVector3d MeanValueCentroid(const FDynamicMesh3& Mesh, int32 VertexIndex, TFunctionRef<FVector3d(int32)> VertexPositionFunc, double WeightClamp = FMathf::MaxReal);
+	static UE_API FVector3d MeanValueCentroid(const FDynamicMesh3& Mesh, int32 VertexIndex, TFunctionRef<FVector3d(int32)> VertexPositionFunc, double WeightClamp = FMathf::MaxReal);
 
 	/**
 	 * Compute cotan-weighted centroid of a vertex one-ring.
 	 * These weights are numerically unstable if any of the triangles are degenerate.
 	 * We catch these problems and return input vertex as centroid
 	 */
-	static FVector3d CotanCentroid(const FDynamicMesh3& Mesh, int32 VertexIndex);
+	static UE_API FVector3d CotanCentroid(const FDynamicMesh3& Mesh, int32 VertexIndex);
 
 	/**
 	 * Compute cotan-weighted centroid of a vertex one-ring.
 	 * These weights are numerically unstable if any of the triangles are degenerate.
 	 * We catch these problems and return input vertex as centroid
 	 */
-	static FVector3d CotanCentroid(const FDynamicMesh3& Mesh, int32 VertexIndex, TFunctionRef<FVector3d(int32)> VertexPositionFunc);
+	static UE_API FVector3d CotanCentroid(const FDynamicMesh3& Mesh, int32 VertexIndex, TFunctionRef<FVector3d(int32)> VertexPositionFunc);
 
 
 
@@ -71,7 +73,7 @@ public:
 	 * @param DegenerateTol if any weights are larger than this value, return uniform weights instead. Should be > 1.
 	 * @param bFailedToUniform will be set to true if non-null and result was clamped
 	 */
-	static FVector3d CotanCentroidSafe(const FDynamicMesh3& Mesh, int32 VertexIndex, double DegenerateTol = 100.0, bool* bFailedToUniform = nullptr);
+	static UE_API FVector3d CotanCentroidSafe(const FDynamicMesh3& Mesh, int32 VertexIndex, double DegenerateTol = 100.0, bool* bFailedToUniform = nullptr);
 
 	/**
 	 * Compute cotan-weighted centroid of a vertex one-ring, with some weight analysis/clamping to avoid vertices getting "stuck"
@@ -79,7 +81,7 @@ public:
 	 * @param DegenerateTol if any weights are larger than this value, return uniform weights instead. Should be > 1.
 	 * @param bFailedToUniform will be set to true if non-null and result was clamped
 	 */
-	static FVector3d CotanCentroidSafe(const FDynamicMesh3& Mesh, int32 VertexIndex, TFunctionRef<FVector3d(int32)> VertexPositionFunc, double DegenerateTol = 100.0, bool* bFailedToUniform = nullptr);
+	static UE_API FVector3d CotanCentroidSafe(const FDynamicMesh3& Mesh, int32 VertexIndex, TFunctionRef<FVector3d(int32)> VertexPositionFunc, double DegenerateTol = 100.0, bool* bFailedToUniform = nullptr);
 
 
 	/**
@@ -90,7 +92,7 @@ public:
 	 * @param DegenerateTol if any weights are larger than this value, return uniform weights instead. Should be > 1.
 	 * @param bFailedToUniform will be set to true if non-null and result was clamped
 	 */
-	static void CotanWeightsBlendSafe(const FDynamicMesh3& Mesh, int32 VertexIndex, TFunctionRef<void(int32,double)> BlendingFunc, double DegenerateTol = 100.0, bool* bFailedToUniform = nullptr);
+	static UE_API void CotanWeightsBlendSafe(const FDynamicMesh3& Mesh, int32 VertexIndex, TFunctionRef<void(int32,double)> BlendingFunc, double DegenerateTol = 100.0, bool* bFailedToUniform = nullptr);
 
 
 
@@ -98,13 +100,13 @@ public:
 	 * Compute the Mixed Voronoi Area associated with a vertex.
 	 * Based on Fig 4 from "Discrete Differential-Geometry Operators for Triangulated 2-Manifolds", Meyer et al 2002,
 	 */
-	static double VoronoiArea(const FDynamicMesh3& Mesh, int32 VertexIndex);
+	static UE_API double VoronoiArea(const FDynamicMesh3& Mesh, int32 VertexIndex);
 
 	/**
 	 * Compute the Mixed Voronoi Area associated with a vertex.
 	 * Based on Fig 4 from "Discrete Differential-Geometry Operators for Triangulated 2-Manifolds", Meyer et al 2002,
 	 */
-	static double VoronoiArea(const FDynamicMesh3& Mesh, int32 VertexIndex, TFunctionRef<FVector3d(int32)> VertexPositionFunc);
+	static UE_API double VoronoiArea(const FDynamicMesh3& Mesh, int32 VertexIndex, TFunctionRef<FVector3d(int32)> VertexPositionFunc);
 
 
 protected:
@@ -113,4 +115,6 @@ protected:
 
 
 } // end namespace UE::Geometry
-} // end namespace UE
\ No newline at end of file
+} // end namespace UE
+
+#undef UE_API
diff --git a/Engine/Plugins/Runtime/GeometryProcessing/Source/DynamicMesh/Public/NormalFlowRemesher.h b/Engine/Plugins/Runtime/GeometryProcessing/Source/DynamicMesh/Public/NormalFlowRemesher.h
--- a/Engine/Plugins/Runtime/GeometryProcessing/Source/DynamicMesh/Public/NormalFlowRemesher.h
+++ b/Engine/Plugins/Runtime/GeometryProcessing/Source/DynamicMesh/Public/NormalFlowRemesher.h
@@ -4,6 +4,8 @@
 
 #include "QueueRemesher.h"
 
+#define UE_API DYNAMICMESH_API
+
 namespace UE
 {
 namespace Geometry
@@ -15,7 +17,7 @@ namespace Geometry
  * This class also performs a pass of edge flips aimed at further aligning triangle normals to their originals. This
  * can be helpful in fixing the occasional single "bad edge" along a sequence of feature edges.
  */
-class DYNAMICMESH_API FNormalFlowRemesher : public FQueueRemesher
+class FNormalFlowRemesher : public FQueueRemesher
 {
 public:
 
@@ -60,11 +62,11 @@ protected:
 	// Do remeshing with face-aligned projection. During the Projection step of RemeshIteration, perform N passes of 
 	// face-aligned projection (N = FaceProjectionPassesPerRemeshIteration).
 	// Additionally, do M passes of face-aligned projection after remeshing has finished (M = NumExtraProjectionIterations).
-	void RemeshWithFaceProjection();
+	UE_API void RemeshWithFaceProjection();
 
 	// Perform face-aligned projection onto the target mesh. Queue edges whose lengths change because of it.
-	void TrackedFaceProjectionPass(double& MaxDistanceMoved, bool bIsTuningIteration);
-	void TrackedFaceProjectionPass_Serial(double& MaxDistanceMoved, bool bIsTuningIteration);
+	UE_API void TrackedFaceProjectionPass(double& MaxDistanceMoved, bool bIsTuningIteration);
+	UE_API void TrackedFaceProjectionPass_Serial(double& MaxDistanceMoved, bool bIsTuningIteration);
 
 	// This is called during RemeshIteration 
 	void TrackedFullProjectionPass(bool bParallel) override
@@ -134,12 +136,14 @@ protected:
 
 	@return True iff (CurrentEdgeError > BadEdgeErrorThreshold) and (NewEdgeError < ImprovementRatioThreshold * CurrentEdgeError)
 	*/
-	bool EdgeFlipWouldReduceNormalError(int EdgeID, double BadEdgeErrorThreshold = 0.01, double ImprovementRatioThreshold = 0.75) const;
+	UE_API bool EdgeFlipWouldReduceNormalError(int EdgeID, double BadEdgeErrorThreshold = 0.01, double ImprovementRatioThreshold = 0.75) const;
 
 	/** Flip all edges for which EdgeFlipWouldReduceNormalError returns true */
-	void TrackedEdgeFlipPass();
+	UE_API void TrackedEdgeFlipPass();
 };
 
 
 } // end namespace UE::Geometry
-} // end namespace UE
\ No newline at end of file
+} // end namespace UE
+
+#undef UE_API
diff --git a/Engine/Plugins/Runtime/GeometryProcessing/Source/DynamicMesh/Public/Operations/AverageOverlayToVertices.h b/Engine/Plugins/Runtime/GeometryProcessing/Source/DynamicMesh/Public/Operations/AverageOverlayToVertices.h
--- a/Engine/Plugins/Runtime/GeometryProcessing/Source/DynamicMesh/Public/Operations/AverageOverlayToVertices.h
+++ b/Engine/Plugins/Runtime/GeometryProcessing/Source/DynamicMesh/Public/Operations/AverageOverlayToVertices.h
@@ -5,6 +5,8 @@
 #include "CoreTypes.h"
 #include "Containers/Array.h"
 
+#define UE_API DYNAMICMESH_API
+
 // Forward declarations
 namespace UE::Geometry 
 { 
@@ -21,7 +23,7 @@ namespace Geometry
 /** 
  * Compute the per-vertex values which are the weighted average of all vertex element values stored in the overlay. 
  */
-class DYNAMICMESH_API FAverageOverlayToVertices
+class FAverageOverlayToVertices
 {
 public:
 
@@ -51,7 +53,7 @@ protected:
 
 public:
 
-	FAverageOverlayToVertices(const FDynamicMesh3& InMesh);
+	UE_API FAverageOverlayToVertices(const FDynamicMesh3& InMesh);
 
 	virtual ~FAverageOverlayToVertices() = default;
 
@@ -73,4 +75,6 @@ public:
 };
 
 }
-}
\ No newline at end of file
+}
+
+#undef UE_API
diff --git a/Engine/Plugins/Runtime/GeometryProcessing/Source/DynamicMesh/Public/Operations/ExtrudeBoundaryEdges.h b/Engine/Plugins/Runtime/GeometryProcessing/Source/DynamicMesh/Public/Operations/ExtrudeBoundaryEdges.h
--- a/Engine/Plugins/Runtime/GeometryProcessing/Source/DynamicMesh/Public/Operations/ExtrudeBoundaryEdges.h
+++ b/Engine/Plugins/Runtime/GeometryProcessing/Source/DynamicMesh/Public/Operations/ExtrudeBoundaryEdges.h
@@ -8,6 +8,8 @@
 #include "Misc/Optional.h"
 #include "VectorTypes.h"
 
+#define UE_API DYNAMICMESH_API
+
 class FProgressCancel;
 
 namespace UE::Geometry
@@ -22,7 +24,7 @@ public:
 	virtual ~FExtrudeBoundaryEdgesFixLayout() = default;
 };
 
-class DYNAMICMESH_API FExtrudeBoundaryEdges : public FExtrudeBoundaryEdgesFixLayout
+class FExtrudeBoundaryEdges : public FExtrudeBoundaryEdgesFixLayout
 {
 public:
 	using FFrame3d = UE::Geometry::FFrame3d;
@@ -31,7 +33,7 @@ public:
 	// Allows vertices to adjust the extrusion distance along one of their extrusion frame
 	//  axes when trying to keep edges parallel.
 	// When using extrusion frames, vertices will be moved in the XZ plane, usually along X.
-	struct DYNAMICMESH_API FExtrudeFrame
+	struct FExtrudeFrame
 	{
 		FExtrudeFrame() {}
 
@@ -44,9 +46,9 @@ public:
 			, Scaling(ScalingIn)
 		{}
 
-		FVector3d FromFramePoint(FVector3d FramePoint) const;
+		UE_API FVector3d FromFramePoint(FVector3d FramePoint) const;
 
-		FVector3d ToFramePoint(FVector3d WorldPoint) const;
+		UE_API FVector3d ToFramePoint(FVector3d WorldPoint) const;
 
 		FFrame3d Frame;
 
@@ -111,7 +113,7 @@ public:
 	TArray<int32> NewExtrudedEids;
 
 public:
-	FExtrudeBoundaryEdges(FDynamicMesh3* mesh);
+	UE_API FExtrudeBoundaryEdges(FDynamicMesh3* mesh);
 
 	virtual ~FExtrudeBoundaryEdges() {}
 
@@ -119,7 +121,7 @@ public:
 	 * Apply the operation to the input mesh.
 	 * @return true if the algorithm succeeds
 	 */
-	virtual bool Apply(FProgressCancel* Progress);
+	UE_API virtual bool Apply(FProgressCancel* Progress);
 
 	/** 
 	 * Pairs up edges across vertices to help in extrusion frame calculation. Public because it is used by
@@ -127,7 +129,7 @@ public:
 	 * 
 	 * @return false if there is an error.
 	 */
-	static bool GetInputEdgePairings(
+	static UE_API bool GetInputEdgePairings(
 		const FDynamicMesh3& Mesh, TArray<int32>& InputEids, bool bAssignAnyBoundaryNeighborToUnmatched,
 		TArray<FNewVertSourceData>& NewVertDataOut, TMap<int32, FIndex2i>& EidToIndicesIntoNewVertsOut);
 
@@ -137,9 +139,11 @@ public:
 	 * 
 	 @return false if there is an error.
 	 */
-	static bool GetExtrudeFrame(const FDynamicMesh3& Mesh, int32 Vid,
+	static UE_API bool GetExtrudeFrame(const FDynamicMesh3& Mesh, int32 Vid,
 		int32 IncomingEid, int32 OutgoingEid, FExtrudeFrame& ExtrudeFrameOut,
 		double ScalingLimit);
 };
 
-} // end namespace UE::Geometry
\ No newline at end of file
+} // end namespace UE::Geometry
+
+#undef UE_API
diff --git a/Engine/Plugins/Runtime/GeometryProcessing/Source/DynamicMesh/Public/Operations/ExtrudeMesh.h b/Engine/Plugins/Runtime/GeometryProcessing/Source/DynamicMesh/Public/Operations/ExtrudeMesh.h
--- a/Engine/Plugins/Runtime/GeometryProcessing/Source/DynamicMesh/Public/Operations/ExtrudeMesh.h
+++ b/Engine/Plugins/Runtime/GeometryProcessing/Source/DynamicMesh/Public/Operations/ExtrudeMesh.h
@@ -9,6 +9,8 @@
 #include "GeometryTypes.h"
 #include "MeshRegionBoundaryLoops.h"
 
+#define UE_API DYNAMICMESH_API
+
 
 namespace UE
 {
@@ -36,7 +38,7 @@ public:
  * 
  * Each quad of the border loop is assigned it's own normal and UVs (ie each is a separate UV-island)
  */
-class DYNAMICMESH_API FExtrudeMesh : public FExtrudeMeshFixLayout
+class FExtrudeMesh : public FExtrudeMeshFixLayout
 {
 public:
 
@@ -99,7 +101,7 @@ public:
 	TArray<FExtrusionInfo> Extrusions;
 
 public:
-	FExtrudeMesh(FDynamicMesh3* mesh);
+	UE_API FExtrudeMesh(FDynamicMesh3* mesh);
 
 	virtual ~FExtrudeMesh() {}
 
@@ -121,14 +123,16 @@ public:
 	 * Apply the Extrude operation to the input mesh.
 	 * @return true if the algorithm succeeds
 	 */
-	virtual bool Apply();
+	UE_API virtual bool Apply();
 
 
 
 protected:
 
-	bool ApplyExtrude(FExtrusionInfo& Region, FMeshNormals* UseNormals = nullptr);
+	UE_API bool ApplyExtrude(FExtrusionInfo& Region, FMeshNormals* UseNormals = nullptr);
 };
 
 } // end namespace UE::Geometry
 } // end namespace UE
+
+#undef UE_API
diff --git a/Engine/Plugins/Runtime/GeometryProcessing/Source/DynamicMesh/Public/Operations/FFDLattice.h b/Engine/Plugins/Runtime/GeometryProcessing/Source/DynamicMesh/Public/Operations/FFDLattice.h
--- a/Engine/Plugins/Runtime/GeometryProcessing/Source/DynamicMesh/Public/Operations/FFDLattice.h
+++ b/Engine/Plugins/Runtime/GeometryProcessing/Source/DynamicMesh/Public/Operations/FFDLattice.h
@@ -8,6 +8,8 @@
 #include "Containers/StaticArray.h"
 #include "DynamicMesh/DynamicMeshAttributeSet.h"
 
+#define UE_API DYNAMICMESH_API
+
 class FProgressCancel;
 
 namespace UE
@@ -33,7 +35,7 @@ struct FLatticeExecutionInfo
  * Free-form deformation lattice. Initialize it with a mesh and desired resolution, get the initial lattice points out.
  * Then pass in deformed lattice points and it will compute deformed mesh vertex positions.
  */
-class DYNAMICMESH_API FFFDLattice
+class FFFDLattice
 {
 public:
 
@@ -46,11 +48,11 @@ public:
 
 	/// Create a lattice that fits the given mesh and has the given resolution along each dimension. Also precompute
 	/// the weight of each mesh vertex inside the appropriate lattice cell.
-	FFFDLattice(const FVector3i& InDims, const FDynamicMesh3& InMesh, float Padding);
+	UE_API FFFDLattice(const FVector3i& InDims, const FDynamicMesh3& InMesh, float Padding);
 
 	/// Using the Lattice's current control point positions and the original embedding information, compute the
 	/// deformed mesh vertex positions.
-	void GetDeformedMeshVertexPositions(const TArray<FVector3d>& LatticeControlPoints,
+	UE_API void GetDeformedMeshVertexPositions(const TArray<FVector3d>& LatticeControlPoints,
 										TArray<FVector3d>& OutVertexPositions, 
 										ELatticeInterpolation Interpolation = ELatticeInterpolation::Linear,
 										FLatticeExecutionInfo ExecutionInfo = FLatticeExecutionInfo(),
@@ -58,7 +60,7 @@ public:
 
 	/// Using the Lattice's current control point positions, the original embedding information, and the original 
 	/// normals contained in an overlay, compute the new rotated normals
-	void GetRotatedOverlayNormals(const TArray<FVector3d>& LatticeControlPoints,
+	UE_API void GetRotatedOverlayNormals(const TArray<FVector3d>& LatticeControlPoints,
 								  const FDynamicMeshNormalOverlay* NormalOverlay,
 								  TArray<FVector3f>& OutNormals,
 								  ELatticeInterpolation Interpolation = ELatticeInterpolation::Linear,
@@ -67,7 +69,7 @@ public:
 
 	/// Using the Lattice's current control point positions, the original embedding information, and the original 
 	/// normals, compute the new rotated normals
-	void GetRotatedMeshVertexNormals(const TArray<FVector3d>& LatticeControlPoints,
+	UE_API void GetRotatedMeshVertexNormals(const TArray<FVector3d>& LatticeControlPoints,
 									 const TArray<FVector3f>& OriginalNormals,
 									 TArray<FVector3f>& OutNormals,
 									 ELatticeInterpolation Interpolation,
@@ -75,10 +77,10 @@ public:
 									 FProgressCancel* Progress) const;
 
 	/// Return the set of lattice corner positions in undeformed state
-	void GenerateInitialLatticePositions(TArray<FVector3d>& OutLatticePositions) const;
+	UE_API void GenerateInitialLatticePositions(TArray<FVector3d>& OutLatticePositions) const;
 
 	/// Return the lattice edges by index. The indices refer to the array of lattice positions returned by GenerateInitialLatticePositions
-	void GenerateLatticeEdges(TArray<FVector2i>& OutLatticeEdges) const;
+	UE_API void GenerateLatticeEdges(TArray<FVector2i>& OutLatticeEdges) const;
 
 	/// Number of lattice points in each dimension
 	const FVector3i& GetDimensions() const
@@ -126,29 +128,31 @@ protected:
 	TArray<FEmbedding> VertexEmbeddings;
 
 	/// For each vertex in Mesh, compute the lattice cell it resides in and its weighting
-	void ComputeInitialEmbedding(const FDynamicMesh3& Mesh, FLatticeExecutionInfo ExecutionInfo = FLatticeExecutionInfo());
+	UE_API void ComputeInitialEmbedding(const FDynamicMesh3& Mesh, FLatticeExecutionInfo ExecutionInfo = FLatticeExecutionInfo());
 
 	/// Compute cell index and linear interpolation weights for a given point
-	FVector3d ComputeTrilinearWeights(const FVector3d& Pt, FVector3i& GridCoordinates) const;
+	UE_API FVector3d ComputeTrilinearWeights(const FVector3d& Pt, FVector3i& GridCoordinates) const;
 
 	/// Helper for linear interpolation. Returns lattice points at (I,J,K) and (I+1,J,K).
-	void GetValuePair(int I, int J, int K, FVector3d& A, FVector3d& B, const TArray<FVector3d>& LatticeControlPoints) const;
+	UE_API void GetValuePair(int I, int J, int K, FVector3d& A, FVector3d& B, const TArray<FVector3d>& LatticeControlPoints) const;
 
 	/// Use the given cell index and weights to compute the interpolated position
-	FVector3d InterpolatedPosition(const FEmbedding& VertexEmbedding, const TArray<FVector3d>& LatticeControlPoints) const;
-	FVector3d InterpolatedPositionCubic(const FEmbedding& VertexEmbedding, const TArray<FVector3d>& LatticeControlPoints) const;
+	UE_API FVector3d InterpolatedPosition(const FEmbedding& VertexEmbedding, const TArray<FVector3d>& LatticeControlPoints) const;
+	UE_API FVector3d InterpolatedPositionCubic(const FEmbedding& VertexEmbedding, const TArray<FVector3d>& LatticeControlPoints) const;
 
-	FMatrix3d LinearInterpolationJacobian(const FEmbedding& VertexEmbedding, const TArray<FVector3d>& LatticeControlPoints) const;
-	FMatrix3d CubicInterpolationJacobian(const FEmbedding& VertexEmbedding, const TArray<FVector3d>& LatticeControlPoints) const;
+	UE_API FMatrix3d LinearInterpolationJacobian(const FEmbedding& VertexEmbedding, const TArray<FVector3d>& LatticeControlPoints) const;
+	UE_API FMatrix3d CubicInterpolationJacobian(const FEmbedding& VertexEmbedding, const TArray<FVector3d>& LatticeControlPoints) const;
 
 	/// Clamp the given index to [0, Dims] and return the current lattice point position at the clamped index
-	FVector3d ClosestLatticePosition(const FVector3i& VirtualControlPointIndex, const TArray<FVector3d>& LatticeControlPoints) const;
+	UE_API FVector3d ClosestLatticePosition(const FVector3i& VirtualControlPointIndex, const TArray<FVector3d>& LatticeControlPoints) const;
 
 	/// Compute the extrapolated position for a "virtual" lattice control point outside of the actual lattice
-	FVector3d ExtrapolatedLatticePosition(const FVector3i& VirtualPointIndex, const TArray<FVector3d>& LatticeControlPoints) const;
+	UE_API FVector3d ExtrapolatedLatticePosition(const FVector3i& VirtualPointIndex, const TArray<FVector3d>& LatticeControlPoints) const;
 
 };
 
 
 } // end namespace UE::Geometry
-} // end namespace UE
\ No newline at end of file
+} // end namespace UE
+
+#undef UE_API
diff --git a/Engine/Plugins/Runtime/GeometryProcessing/Source/DynamicMesh/Public/Operations/GeodesicPath.h b/Engine/Plugins/Runtime/GeometryProcessing/Source/DynamicMesh/Public/Operations/GeodesicPath.h
--- a/Engine/Plugins/Runtime/GeometryProcessing/Source/DynamicMesh/Public/Operations/GeodesicPath.h
+++ b/Engine/Plugins/Runtime/GeometryProcessing/Source/DynamicMesh/Public/Operations/GeodesicPath.h
@@ -15,6 +15,8 @@
 #include "Util/RefCountVector.h"
 #include "MathUtil.h"
 
+#define UE_API DYNAMICMESH_API
+
 
 namespace UE
 {
@@ -30,7 +32,7 @@ namespace Geometry
 *       no checking is done to ensure the path is continuous, in fact it could 
 *       be used to represent multiple disconnected paths.
 */
-class DYNAMICMESH_API FEdgePath
+class FEdgePath
 {
 public:
 
@@ -55,24 +57,24 @@ public:
 	inline bool IsSegment(int32 SID) const;
 	
 	/** @return requested segment if it is part of the path, otherwise return invalid segment (with EID = InvalidID) */
-	FDirectedSegment GetSegment(int32 SID) const;
+	UE_API FDirectedSegment GetSegment(int32 SID) const;
 
 	/** @return the segment ID of the {first, last} segment in the path, or InvalidID if the path is empty */
-	int32 GetHeadSegmentID() const;
-	int32 GetTailSegmentID() const;
+	UE_API int32 GetHeadSegmentID() const;
+	UE_API int32 GetTailSegmentID() const;
 
 	/** @return segment ID of {previous, next} segment in the path, or InvalidID if if no {previous, next} exists or if SID itself is not in the path  */
-	int32 GetPrevSegmentID(int32 SID) const;
-	int32 GetNextSegmentID(int32 SID) const;
+	UE_API int32 GetPrevSegmentID(int32 SID) const;
+	UE_API int32 GetNextSegmentID(int32 SID) const;
 
 	/** @return segment ID of a new segment appended to the tail of the path */
-	int32 AppendSegment(const FDirectedSegment& Segment);
+	UE_API int32 AppendSegment(const FDirectedSegment& Segment);
 
 	/** @return segment ID of a new segment inserted before the specified existing segment. If SegmentIDToInsertBefore is InvalidID, then insert at head  */
-	int32 InsertSegmentBefore(const FDirectedSegment& DirectedSegment, int32 SegmentIDToInsertBefore = InvalidID);
+	UE_API int32 InsertSegmentBefore(const FDirectedSegment& DirectedSegment, int32 SegmentIDToInsertBefore = InvalidID);
 	
 	/** remove specified segment from path. Note SID is added to an internal free list and maybe reused  */
-	void RemoveSegment(int32 SID);
+	UE_API void RemoveSegment(int32 SID);
 
 
 protected:
@@ -105,7 +107,7 @@ protected:
 * these edges are encoded with path direction as FDirectedSegments.
 * Note: The current implementation does not support closed paths or multiple paths.
 */
-class DYNAMICMESH_API FDeformableEdgePath
+class FDeformableEdgePath
 {
 public:
 	
@@ -116,7 +118,7 @@ public:
 	/**
 	* Constructor assumes the directed segments are ordered from tail (at index 0) to head
 	*/
-	FDeformableEdgePath(const FDynamicMesh3& SurfaceMeshIn, const TArray<FEdgePath::FDirectedSegment>& OriginalPathAsDirectedSegments);
+	UE_API FDeformableEdgePath(const FDynamicMesh3& SurfaceMeshIn, const TArray<FEdgePath::FDirectedSegment>& OriginalPathAsDirectedSegments);
 	
 	virtual ~FDeformableEdgePath(){}
 
@@ -140,7 +142,7 @@ public:
 	* @param DeformedPathInfo  - on return this info struct is populated
 	* @param MaxNumIterations  - the maximum number of two-edge joints that may be replaced by shorter paths.
 	*/ 
-	virtual void Minimize(FEdgePathDeformationInfo&  DeformedPathInfo, const int32 MaxNumIterations = TMathUtilConstants<int32>::MaxReal);
+	UE_API virtual void Minimize(FEdgePathDeformationInfo&  DeformedPathInfo, const int32 MaxNumIterations = TMathUtilConstants<int32>::MaxReal);
 	
 	/**
 	* @return a const reference to the edge path.
@@ -172,7 +174,7 @@ public:
 	*                            as a sequence of vertex adjacent edge-crossings very close to the vertex.. the threshold is applied in
 	*                            a post-process coalesce those crossings into a single vertex crossing.
 	*/
-	TArray<FSurfacePoint> AsSurfacePoints(double CoalesceThreshold = 0.) const;
+	UE_API TArray<FSurfacePoint> AsSurfacePoints(double CoalesceThreshold = 0.) const;
 
 protected:
 	
@@ -190,7 +192,7 @@ protected:
 	* that connects the same end points. This new path is formed by flipping the adjacent edges to the 
 	* left (or right) of the joint vertex and may have one, two, or more segments
 	*/
-	void DeformJoint(FPathJoint& PathJoint);
+	UE_API void DeformJoint(FPathJoint& PathJoint);
 
 	/**
 	* Edge flip the adjacent edges to the JointVID and collect the DeformedPath that can be used to replace the local incoming and outgoing edge path
@@ -202,7 +204,7 @@ protected:
 	* @param bClockwise   - When true, the deformation and new path are on the "left" of the original path, otherwise on the "right"
 	*                       where the path is locally viewed as traveling from bottom to top
 	*/ 
-	bool DeformJoint( const int32 IncomingEID, const int32 OutgoingEID, int32 JointVID, 
+	UE_API bool DeformJoint( const int32 IncomingEID, const int32 OutgoingEID, int32 JointVID, 
 	                  TArray<FEdgePath::FDirectedSegment>& DeformedPath, const bool bClockwise );
 	
 	/**
@@ -218,25 +220,25 @@ protected:
 	/**
 	* @return true if no part of the global path lies in the wedge (partial one-ring) on PathJoint.Interior side
 	*/ 
-	bool IsJointFlexible(const FPathJoint& PathJoint) const;
+	UE_API bool IsJointFlexible(const FPathJoint& PathJoint) const;
 	
 	/**
 	* Compute data associated with the joint formed by the specified segment (outgoing), and the preceding segment (incoming) in the path.
 	* Also update the JointAngleQueue that tracks straightness of each joint.
 	*/ 
-	void UpdateJointAndQueue(int32 OutgoingSegmentID);
+	UE_API void UpdateJointAndQueue(int32 OutgoingSegmentID);
 	
 	/**
 	* remove specified segment from the linked-list path
 	* and clean up joint-related data for the joint that ends with this segment.
 	* note: the joint data (e.g. joint angle queue) for the joint that started with this segment will be dirty.
 	*/ 
-	void RemoveSegment(int32 SID);
+	UE_API void RemoveSegment(int32 SID);
 
 	/**
 	* Splice in a path to replace the joint and update the joint angle queue.
 	*/ 
-	void ReplaceJointWithPath(const FPathJoint PathJoint, const TArray<FEdgePath::FDirectedSegment>& PathAsDirectedEdges);
+	UE_API void ReplaceJointWithPath(const FPathJoint PathJoint, const TArray<FEdgePath::FDirectedSegment>& PathAsDirectedEdges);
 
 	/**
 	* Traveling either CCW or CW around vertex CenterVID, from edge StartEID to edge EndEID, flip edges adjacent to CVID.
@@ -248,7 +250,7 @@ protected:
 	*       - no attempt is made to restore the mesh.
 	*       
 	*/
-	bool OuterArcFlipEdges(int32 StartEID, int32 EndEID, int32 CenterVID, bool bClockwise = false);
+	UE_API bool OuterArcFlipEdges(int32 StartEID, int32 EndEID, int32 CenterVID, bool bClockwise = false);
 	
 	/**
 	* Compute the two angles formed at the vertex CenterVID by the two wedges formed by splitting the local one-ring by local path
@@ -256,7 +258,7 @@ protected:
 	* 
 	* It is assumed the specified edges do in fact have vertex CVID in common, but this is not checked.
 	*/
-	void ComputeWedgeAngles(int32 IncomingEID, int32 OutgoingEID, int32 CenterVID, 
+	UE_API void ComputeWedgeAngles(int32 IncomingEID, int32 OutgoingEID, int32 CenterVID, 
 	                        double& LeftSideAngle, double& RightSideAngle) const ;
 	
 protected:
@@ -327,3 +329,5 @@ bool FEdgePath::IsSegment(int32 SID) const
 
 }; // end namespace Geometry
 }; // end namespace UE
+
+#undef UE_API
diff --git a/Engine/Plugins/Runtime/GeometryProcessing/Source/DynamicMesh/Public/Operations/GroupEdgeInserter.h b/Engine/Plugins/Runtime/GeometryProcessing/Source/DynamicMesh/Public/Operations/GroupEdgeInserter.h
--- a/Engine/Plugins/Runtime/GeometryProcessing/Source/DynamicMesh/Public/Operations/GroupEdgeInserter.h
+++ b/Engine/Plugins/Runtime/GeometryProcessing/Source/DynamicMesh/Public/Operations/GroupEdgeInserter.h
@@ -6,6 +6,8 @@
 
 #include "GroupTopology.h"
 
+#define UE_API DYNAMICMESH_API
+
 class FProgressCancel;
 
 namespace UE
@@ -19,7 +21,7 @@ class FGroupTopology;
 /**
  * Used to insert group edges and group edge loops.
  */
-class DYNAMICMESH_API FGroupEdgeInserter
+class FGroupEdgeInserter
 {
 public:
 
@@ -97,7 +99,7 @@ public:
 		TSet<int32>* ProblemGroupEdgeIDsOut = nullptr;
 	};
 
-	bool InsertEdgeLoops(const FEdgeLoopInsertionParams& Params, FOptionalOutputParams OptionalOut = FOptionalOutputParams(), FProgressCancel* Progress = nullptr);
+	UE_API bool InsertEdgeLoops(const FEdgeLoopInsertionParams& Params, FOptionalOutputParams OptionalOut = FOptionalOutputParams(), FProgressCancel* Progress = nullptr);
 
 
 	/** Point along a group edge that is used as a start/endpoint for an inserted group edge. */
@@ -155,10 +157,12 @@ public:
 		EInsertionMode Mode = EInsertionMode::Retriangulate;
 	};
 
-	bool InsertGroupEdge(FGroupEdgeInsertionParams& Params, FOptionalOutputParams OptionalOut = FOptionalOutputParams(), FProgressCancel* Progress = nullptr);
+	UE_API bool InsertGroupEdge(FGroupEdgeInsertionParams& Params, FOptionalOutputParams OptionalOut = FOptionalOutputParams(), FProgressCancel* Progress = nullptr);
 
 };
 
 
 } // end namespace UE::Geometry
-} // end namespace UE
\ No newline at end of file
+} // end namespace UE
+
+#undef UE_API
diff --git a/Engine/Plugins/Runtime/GeometryProcessing/Source/DynamicMesh/Public/Operations/GroupTopologyDeformer.h b/Engine/Plugins/Runtime/GeometryProcessing/Source/DynamicMesh/Public/Operations/GroupTopologyDeformer.h
--- a/Engine/Plugins/Runtime/GeometryProcessing/Source/DynamicMesh/Public/Operations/GroupTopologyDeformer.h
+++ b/Engine/Plugins/Runtime/GeometryProcessing/Source/DynamicMesh/Public/Operations/GroupTopologyDeformer.h
@@ -4,6 +4,8 @@
 
 #include "EdgeSpan.h"
 
+#define UE_API DYNAMICMESH_API
+
 
 namespace UE
 {
@@ -18,7 +20,7 @@ struct FGeometrySelection;
 /**
  * Basic cache of vertex positions.
  */
-class DYNAMICMESH_API FVertexPositionCache
+class FVertexPositionCache
 {
 public:
 	TArray<int> Vertices;
@@ -65,13 +67,13 @@ public:
  * Various functions can be overrided to customize behavior.
  * 
  */
-class DYNAMICMESH_API FGroupTopologyDeformer
+class FGroupTopologyDeformer
 {
 public:
 	virtual ~FGroupTopologyDeformer() {}
 
 	/** Set the Mesh and Topology to use for the deformation */
-	void Initialize(const FDynamicMesh3* Mesh, const FGroupTopology* Topology);
+	UE_API void Initialize(const FDynamicMesh3* Mesh, const FGroupTopology* Topology);
 	
 	const FDynamicMesh3* GetMesh() const { return Mesh; }
 	const FGroupTopology* GetTopology() const { return Topology; }
@@ -83,23 +85,23 @@ public:
 	/**
 	 * Set the active handle to the given Faces
 	 */
-	virtual void SetActiveHandleFaces(const TArray<int>& FaceGroupIDs);
+	UE_API virtual void SetActiveHandleFaces(const TArray<int>& FaceGroupIDs);
 
 	/**
 	 * Set the active handle to the given Corners
 	 */
-	virtual void SetActiveHandleCorners(const TArray<int>& TopologyCornerIDs);
+	UE_API virtual void SetActiveHandleCorners(const TArray<int>& TopologyCornerIDs);
 
 	/**
 	 * Set the active handle to the given Edges
 	 */
-	virtual void SetActiveHandleEdges(const TArray<int>& TopologyEdgeIDs);
+	UE_API virtual void SetActiveHandleEdges(const TArray<int>& TopologyEdgeIDs);
 
 
 	/**
 	 * Set the active handle from GeometrySelection
 	 */
-	virtual void SetActiveHandleFromSelection(const FGeometrySelection& Selection);
+	UE_API virtual void SetActiveHandleFromSelection(const FGeometrySelection& Selection);
 
 	//
 	// Solving
@@ -109,12 +111,12 @@ public:
 	 * Update TargetMesh by first calling HandleVertexDeformFunc() to get new
 	 * handle vertex positions, then solving for new ROI vertex positions.
 	 */
-	virtual void UpdateSolution(FDynamicMesh3* TargetMesh, const TFunction<FVector3d(FDynamicMesh3* Mesh, int)>& HandleVertexDeformFunc);
+	UE_API virtual void UpdateSolution(FDynamicMesh3* TargetMesh, const TFunction<FVector3d(FDynamicMesh3* Mesh, int)>& HandleVertexDeformFunc);
 
 	/**
 	 * Restore the Handle and ROI vertex positions to their initial state
 	 */
-	virtual void ClearSolution(FDynamicMesh3* TargetMesh);
+	UE_API virtual void ClearSolution(FDynamicMesh3* TargetMesh);
 
 
 
@@ -140,23 +142,23 @@ protected:
 	// 
 
 	/** reset all internal data structures, eg when changing handle */
-	virtual void Reset();
+	UE_API virtual void Reset();
 
 	/**
 	 * Populate the internal ROI data structures based on the given HandleGroups and ROIGroups.
 	 * HandleGroups will be empty if the Handle is a set of Vertices or Edges.
 	 */
-	virtual void CalculateROI(const TArray<int>& HandleGroups, const TArray<int>& ROIGroups);
+	UE_API virtual void CalculateROI(const TArray<int>& HandleGroups, const TArray<int>& ROIGroups);
 
 	/**
 	 * Save the positions of all vertices that will be modified
 	 */
-	virtual void SaveInitialPositions();
+	UE_API virtual void SaveInitialPositions();
 
 	/**
 	 * Precompute the representation of the ROI vertices at the initial positions.
 	 */
-	virtual void ComputeEncoding();
+	UE_API virtual void ComputeEncoding();
 
 protected:
 	const FDynamicMesh3* Mesh = nullptr;
@@ -221,4 +223,6 @@ protected:
 
 
 } // end namespace UE::Geometry
-} // end namespace UE
\ No newline at end of file
+} // end namespace UE
+
+#undef UE_API
diff --git a/Engine/Plugins/Runtime/GeometryProcessing/Source/DynamicMesh/Public/Operations/InsetMeshRegion.h b/Engine/Plugins/Runtime/GeometryProcessing/Source/DynamicMesh/Public/Operations/InsetMeshRegion.h
--- a/Engine/Plugins/Runtime/GeometryProcessing/Source/DynamicMesh/Public/Operations/InsetMeshRegion.h
+++ b/Engine/Plugins/Runtime/GeometryProcessing/Source/DynamicMesh/Public/Operations/InsetMeshRegion.h
@@ -8,6 +8,8 @@
 #include "MeshRegionBoundaryLoops.h"
 #include "ProjectionTargets.h"
 
+#define UE_API DYNAMICMESH_API
+
 
 namespace UE
 {
@@ -29,7 +31,7 @@ class FDynamicMeshChangeTracker;
  *
  * Each quad of the border loop is assigned it's own normal and UVs (ie each is a separate UV-island)
  */
-class DYNAMICMESH_API FInsetMeshRegion
+class FInsetMeshRegion
 {
 public:
 
@@ -108,9 +110,9 @@ public:
 
 
 public:
-	FInsetMeshRegion(FDynamicMesh3* mesh);
+	UE_API FInsetMeshRegion(FDynamicMesh3* mesh);
 
-	virtual ~FInsetMeshRegion();
+	UE_API virtual ~FInsetMeshRegion();
 
 
 	/**
@@ -130,14 +132,16 @@ public:
 	 * Apply the Inset operation to the input mesh.
 	 * @return true if the algorithm succeeds
 	 */
-	virtual bool Apply();
+	UE_API virtual bool Apply();
 
 
 protected:
 
-	virtual bool ApplyInset(FInsetInfo& Region, FMeshNormals* UseNormals = nullptr);
+	UE_API virtual bool ApplyInset(FInsetInfo& Region, FMeshNormals* UseNormals = nullptr);
 };
 
 
 } // end namespace UE::Geometry
-} // end namespace UE
\ No newline at end of file
+} // end namespace UE
+
+#undef UE_API
diff --git a/Engine/Plugins/Runtime/GeometryProcessing/Source/DynamicMesh/Public/Operations/IntrinsicCorrespondenceUtils.h b/Engine/Plugins/Runtime/GeometryProcessing/Source/DynamicMesh/Public/Operations/IntrinsicCorrespondenceUtils.h
--- a/Engine/Plugins/Runtime/GeometryProcessing/Source/DynamicMesh/Public/Operations/IntrinsicCorrespondenceUtils.h
+++ b/Engine/Plugins/Runtime/GeometryProcessing/Source/DynamicMesh/Public/Operations/IntrinsicCorrespondenceUtils.h
@@ -9,6 +9,8 @@
 #include "VectorTypes.h"
 #include "MathUtil.h"
 
+#define UE_API DYNAMICMESH_API
+
 namespace UE
 {
 namespace Geometry
@@ -19,7 +21,7 @@ namespace IntrinsicCorrespondenceUtils
 {
 	
 	// Location on an extrinsic (surface) mesh stored as a union.
-	struct DYNAMICMESH_API FSurfacePoint
+	struct FSurfacePoint
 	{
 		enum class EPositionType
 		{
@@ -29,10 +31,10 @@ namespace IntrinsicCorrespondenceUtils
 		};
 
 		// constructors respect the correct position type
-		FSurfacePoint(); // defaults to invalid vertex position
-		FSurfacePoint(const int32 VID);
-		FSurfacePoint(const int32 EdgeID, const double Alpha);
-		FSurfacePoint(const int32 TriID, const FVector3d& BaryCentrics);
+		UE_API FSurfacePoint(); // defaults to invalid vertex position
+		UE_API FSurfacePoint(const int32 VID);
+		UE_API FSurfacePoint(const int32 EdgeID, const double Alpha);
+		UE_API FSurfacePoint(const int32 TriID, const FVector3d& BaryCentrics);
 
 		struct FVertexPosition
 		{
@@ -94,17 +96,17 @@ namespace IntrinsicCorrespondenceUtils
 	}
 
 	// A Connection stores the information needed to define a local reference direction for each vertex and triangle.
-	struct DYNAMICMESH_API FMeshConnection
+	struct FMeshConnection
 	{
 		FMeshConnection() : SurfaceMesh(nullptr) {};
 
-		FMeshConnection(const FDynamicMesh3& SurfMesh);
+		UE_API FMeshConnection(const FDynamicMesh3& SurfMesh);
 
 		/** Reset the connection to the given (not owned) surface mesh*/
-		void Reset(const FDynamicMesh3& SurfMesh);     
+		UE_API void Reset(const FDynamicMesh3& SurfMesh);     
 
 		/** Reset the connection to used a internally owned surface mesh*/
-		void Reset(TUniquePtr<FDynamicMesh3>& SurfMesh);
+		UE_API void Reset(TUniquePtr<FDynamicMesh3>& SurfMesh);
 		
 
 		
@@ -120,19 +122,19 @@ namespace IntrinsicCorrespondenceUtils
 	* This structure assumes that the surface mesh is fixed, and the intrinsic mesh shares the same vertex set as the surface mesh and
 	* is initialized with the same connectivity as the surface mesh. Note, it is assumed that the surface mesh has no bow-ties.
 	*/
-	struct DYNAMICMESH_API FNormalCoordinates : public FMeshConnection
+	struct FNormalCoordinates : public FMeshConnection
 	{
 		typedef FMeshConnection  MyBase;
 
 		FNormalCoordinates(): MyBase() {}
 
-		FNormalCoordinates(const FDynamicMesh3& SurfMesh);
+		UE_API FNormalCoordinates(const FDynamicMesh3& SurfMesh);
 
 		/** Reset the connection to the given (not owned) surface mesh*/
-		void Reset(const FDynamicMesh3& SurfMesh);
+		UE_API void Reset(const FDynamicMesh3& SurfMesh);
 
 		/** Reset the connection to used an internally owned surface mesh*/
-		void Reset(TUniquePtr<FDynamicMesh3>& SurfMesh);
+		UE_API void Reset(TUniquePtr<FDynamicMesh3>& SurfMesh);
 
 
 
@@ -175,20 +177,20 @@ namespace IntrinsicCorrespondenceUtils
 		* triangle edges of the two affected triangles prior (!) to flipping the specified edge.
 		* @return false if something failed
 		*/
-		bool  OnFlipEdge(const int32 T0ID, const FIndex3i& BeforeT0EIDs, const int32 T0OppVID,
+		UE_API bool  OnFlipEdge(const int32 T0ID, const FIndex3i& BeforeT0EIDs, const int32 T0OppVID,
 			             const int32 T1ID, const FIndex3i& BeforeT1EIDs, const int32 T1OppVID,
 			             const int32 FlippedEID);
 
 		/**
 		* @return the edgeID of the N-th out-going edge in the SurfaceMesh (counting CCW from the reference edge for this vertex).
 		*/
-		int32 GetNthEdgeID(const int32 VID, const int32 N) const;
+		UE_API int32 GetNthEdgeID(const int32 VID, const int32 N) const;
 
 		/**
 		* @return the Order of the edge EID about VID (counting CCW from the surface reference edge for this vertex).
 		* Note: this assumes that edge EID is adjacent to vertex VID - otherwise -1 is returned.
 		*/
-		int32 GetEdgeOrder(const int32 VID, const int32 SurfaceEID) const;
+		UE_API int32 GetEdgeOrder(const int32 VID, const int32 SurfaceEID) const;
 
 		TDynamicVector<int32>     NormalCoord;     // Per intrinsic edge: number of times src mesh edge crosses. By convention NormalCoord[EID] = -1 iff same as edge on surface mesh   
 		TDynamicVector<FIndex3i>  RoundaboutOrder; // Per intrinsic Triangle, per (directed) edge: indicates the closest (when traveling ccw) reference mesh edge
@@ -198,7 +200,7 @@ namespace IntrinsicCorrespondenceUtils
 
 	protected:
 		/** rebuild the normal coordinate data in this connection. Assumes the base connection data is consistent with this surface mesh*/
-		void RebuildNormalCoordinates(const FDynamicMesh3& SurfMesh);
+		UE_API void RebuildNormalCoordinates(const FDynamicMesh3& SurfMesh);
 	};
 
 
@@ -207,10 +209,10 @@ namespace IntrinsicCorrespondenceUtils
 	* This structure assumes that the surface mesh is fixed, but allows the intrinsic mesh to support operations like edge splits and triangle pokes.
 	* is initialized with the same connectivity as the surface mesh.  Note, it is assumed that the surface mesh has no bow-ties.
 	*/
-	struct DYNAMICMESH_API FSignpost : public FMeshConnection
+	struct FSignpost : public FMeshConnection
 	{
 		typedef FMeshConnection  MyBase;
-		FSignpost(const FDynamicMesh3& SurfaceMesh);
+		UE_API FSignpost(const FDynamicMesh3& SurfaceMesh);
 
 		/**
 		* Update the signpost after an edge flip.  This requires information about the both the pre and post flip state.
@@ -221,7 +223,7 @@ namespace IntrinsicCorrespondenceUtils
 		* @param Opp0NewInternalAngle - internal angle of opposing vert 0 after the flip. i.e. InternalAngles[Tris[1]][1]
 		* @param Opp1NewInternalAngle - internal angle of opposing vert 1 after the flip. i.e  InternalAngles[Tris[0]][1];
 		*/
-		void OnFlipEdge(const int32 EID, const FIndex2i AdjTris, const FIndex2i OpposingVerts, const FIndex2i PreFlipIndexOf, 
+		UE_API void OnFlipEdge(const int32 EID, const FIndex2i AdjTris, const FIndex2i OpposingVerts, const FIndex2i PreFlipIndexOf, 
 		                const double Opp0NewInternalAngle, const double Opp1NewInternalAngle);
 
 		struct FSurfaceTraceResult
@@ -372,4 +374,6 @@ namespace IntrinsicCorrespondenceUtils
 	}
 };
 };
-};
\ No newline at end of file
+};
+
+#undef UE_API
diff --git a/Engine/Plugins/Runtime/GeometryProcessing/Source/DynamicMesh/Public/Operations/IntrinsicTriangulationMesh.h b/Engine/Plugins/Runtime/GeometryProcessing/Source/DynamicMesh/Public/Operations/IntrinsicTriangulationMesh.h
--- a/Engine/Plugins/Runtime/GeometryProcessing/Source/DynamicMesh/Public/Operations/IntrinsicTriangulationMesh.h
+++ b/Engine/Plugins/Runtime/GeometryProcessing/Source/DynamicMesh/Public/Operations/IntrinsicTriangulationMesh.h
@@ -10,6 +10,8 @@
 #include "VectorTypes.h"
 #include "MathUtil.h"
 
+#define UE_API DYNAMICMESH_API
+
 
 namespace UE
 {
@@ -60,7 +62,7 @@ namespace Geometry
 * In addition to managing vertices, triangles and edges this class also tracks edge lengths (as measured on the surface mesh)
 * and for convenience the internal angles for each triangle.
 */
-class DYNAMICMESH_API FSimpleIntrinsicEdgeFlipMesh
+class FSimpleIntrinsicEdgeFlipMesh
 {
 public:
 	
@@ -78,15 +80,15 @@ public:
 	FSimpleIntrinsicEdgeFlipMesh(){};
 
 	/** Constructor does ID-matching deep copy the basic mesh topology (but no attributes or groups) */
-	FSimpleIntrinsicEdgeFlipMesh(const FDynamicMesh3& SrcMesh);
+	UE_API FSimpleIntrinsicEdgeFlipMesh(const FDynamicMesh3& SrcMesh);
 
 	FSimpleIntrinsicEdgeFlipMesh(const FSimpleIntrinsicEdgeFlipMesh&) = delete;
 
 	/** Discard all data */
-	void Clear();
+	UE_API void Clear();
 
 	/** Reset the mesh to an ID-matching deep copy the basic mesh topology (but no attributes or groups) */
-	void Reset(const FDynamicMesh3& SrcMesh);
+	UE_API void Reset(const FDynamicMesh3& SrcMesh);
 	/**
 	* Flip a single edge in the Intrinsic Mesh.
 	* Note: After a successful flip, both triangles will share the same orientation (matching that of the pre-flip Edge.Tri.A)
@@ -97,7 +99,7 @@ public:
 	*
 	* @return a success or failure info.
 	*/
-	EMeshResult FlipEdge(int32 EID, FEdgeFlipInfo& EdgeFlipInfo);
+	UE_API EMeshResult FlipEdge(int32 EID, FEdgeFlipInfo& EdgeFlipInfo);
 
 
 	/** @return the intrinsic edge length of specified edge */
@@ -139,7 +141,7 @@ public:
 	*
 	* Note: if EID is a mesh boundary beta_{ij} is replaced by -MAX_DOUBLE
 	*/
-	FVector2d EdgeOpposingAngles(int32 EID) const;
+	UE_API FVector2d EdgeOpposingAngles(int32 EID) const;
 
 
 	/**
@@ -149,7 +151,7 @@ public:
 	*
 	* Note: if EID is a mesh boundary, cot(alpha_ij) is returned
 	*/
-	double EdgeCotanWeight(int32 EID) const;
+	UE_API double EdgeCotanWeight(int32 EID) const;
 
 	// -- Minimal mesh interface --//
 
@@ -226,7 +228,7 @@ public:
 	}
 
 	/** @return the vertex pair opposite specified edge. The second vertex may be InvalidID */
-	FIndex2i GetEdgeOpposingV(int32 EID) const;
+	UE_API FIndex2i GetEdgeOpposingV(int32 EID) const;
 
 	/** @return true if intrinsic mesh contains specified vertex */
 	inline bool IsVertex(int32 VertexID) const
@@ -353,13 +355,13 @@ protected:
 	* the order of the angles matches the order of the vertices that define the triangle, e.g. if (a, b, c) = GetTriangle()
 	* the (angle at a, angle at b, angle at c) = ComputeTriInternalAnglesR
 	*/  
-	FVector3d ComputeTriInternalAnglesR(const int32 TID) const;
+	UE_API FVector3d ComputeTriInternalAnglesR(const int32 TID) const;
 
 	/**
 	* @return the intrinsic distance between the two vertices opposite the specified edge
 	* NB: Will fail if EID is a boundary edge.
 	*/
-	double GetOpposingVerticesDistance(int32 EID) const;
+	UE_API double GetOpposingVerticesDistance(int32 EID) const;
 
 protected:
 	// basic mesh operations
@@ -368,7 +370,7 @@ protected:
 	* update the mesh topology (e.g. triangles and edges) by doing an edge flip, but does not update the intrinsic edge lengths or angles.
 	* Note: this requires precomputed FlipInfo.bHadSameOrientation. 
 	*/
-	EMeshResult FlipEdgeTopology(int32 eab, FEdgeFlipInfo& FlipInfo);
+	UE_API EMeshResult FlipEdgeTopology(int32 eab, FEdgeFlipInfo& FlipInfo);
 
 	/** update the existing edge to reference vertices 'a' and 'b'.*/
 	void SetEdgeVerticesInternal(int32 EdgeID, int32 a, int32 b)
@@ -387,7 +389,7 @@ protected:
 		Edges[EdgeID].Tri[1] = t1;
 	}
 	/** replace tOld with tNew in an existing edge */
-	int32 ReplaceEdgeTriangle(int32 eID, int32 tOld, int32 tNew);
+	UE_API int32 ReplaceEdgeTriangle(int32 eID, int32 tOld, int32 tNew);
 	/** replace eOld with eNew in and existing triangle */
 	inline int32 ReplaceTriangleEdge(int32 tID, int32 eOld, int32 eNew)
 	{
@@ -408,7 +410,7 @@ protected:
 	* note: the vertices may be the same.
 	* note: assumes EID is an edge in Triangle TID
 	*/
-	FIndex2i GetOrientedEdgeV(int32 EID, int32 TID) const;
+	UE_API FIndex2i GetOrientedEdgeV(int32 EID, int32 TID) const;
 protected:
 
 	const int32 AddOneModThree[3] = { 1, 2, 0 };
@@ -449,7 +451,7 @@ protected:
 * with the surface mesh ( "normal coordinates" or "signpost data") that can track these surface mesh locations 
 * and can be used to reconstruct paths (e.g. intrinsic edges) on the surface mesh.
 */
-class DYNAMICMESH_API FSimpleIntrinsicMesh : public FSimpleIntrinsicEdgeFlipMesh
+class FSimpleIntrinsicMesh : public FSimpleIntrinsicEdgeFlipMesh
 {
 public:
 	typedef FSimpleIntrinsicEdgeFlipMesh  MyBase;
@@ -479,7 +481,7 @@ public:
 	* 
 	* @return Ok on success, or enum value indicates why operation cannot be applied. Mesh remains unmodified on error.
 	*/
-	EMeshResult PokeTriangle(int32 TID, const FVector3d& BaryCoordinates, FPokeTriangleInfo& PokeInfo);
+	UE_API EMeshResult PokeTriangle(int32 TID, const FVector3d& BaryCoordinates, FPokeTriangleInfo& PokeInfo);
 
 
 	/**
@@ -494,7 +496,7 @@ public:
 	* 
 	* @return Ok on success, or enum value indicates why operation cannot be applied. Mesh remains unmodified on error.
 	*/
-	EMeshResult SplitEdge(int32 EdgeAB, FEdgeSplitInfo& SplitInfo, double SplitParameterT = 0.5);
+	UE_API EMeshResult SplitEdge(int32 EdgeAB, FEdgeSplitInfo& SplitInfo, double SplitParameterT = 0.5);
 
 protected:
 
@@ -502,13 +504,13 @@ protected:
 	* Updates the mesh connectivity by adding a new vertex at the vertex-averaged location. Note: this does not update any of the intrinsic lengths, positions, angles, etc.
 	* those quantities must be updated after.  See ::PokeTriangle()
 	*/ 
-	EMeshResult PokeTriangleTopology(int32 TID, FPokeTriangleInfo& PokeInfo);
+	UE_API EMeshResult PokeTriangleTopology(int32 TID, FPokeTriangleInfo& PokeInfo);
 
 	/**
 	* Updates the mesh connectivity by adding a new vertex at the vertex-averaged location. Note: this does not update any of the intrinsic lengths, positions, angles, etc.
 	* those quantities must be updated after.  See ::SplitEdge()
 	*/
-	EMeshResult SplitEdgeTopology(int32 EdgeAB, FEdgeSplitInfo& SplitInfo);
+	UE_API EMeshResult SplitEdgeTopology(int32 EdgeAB, FEdgeSplitInfo& SplitInfo);
 
 	/** allocate, or clear existing edge list for specified vertex*/
 	inline void AllocateEdgesList(int32 VertexID)
@@ -570,7 +572,7 @@ protected:
 * such as the FIntrinsicMesh and FIntrinsicEdgeFlipMesh.
 */ 
 template <typename IntrinsicMeshType>
-struct DYNAMICMESH_API TEdgeCorrespondence
+struct TEdgeCorrespondence
 {
 
 	using FEdgeAndCrossingIdx = IntrinsicCorrespondenceUtils::FNormalCoordinates::FEdgeAndCrossingIdx;
@@ -638,7 +640,7 @@ struct DYNAMICMESH_API TEdgeCorrespondence
 * 
 * Note:  This assumes but does not check that the FDynamicMesh3 "Surface Mesh" has consistent orientation of all triangles.
 */
-class DYNAMICMESH_API FIntrinsicEdgeFlipMesh : public  FSimpleIntrinsicEdgeFlipMesh
+class FIntrinsicEdgeFlipMesh : public  FSimpleIntrinsicEdgeFlipMesh
 {
 public:
 	typedef FSimpleIntrinsicEdgeFlipMesh MyBase;
@@ -647,7 +649,7 @@ public:
 	using FSurfacePoint       = IntrinsicCorrespondenceUtils::FSurfacePoint;
 	using FEdgeAndCrossingIdx = IntrinsicCorrespondenceUtils::FNormalCoordinates::FEdgeAndCrossingIdx;
 
-	FIntrinsicEdgeFlipMesh(const FDynamicMesh3& SurfaceMesh);
+	UE_API FIntrinsicEdgeFlipMesh(const FDynamicMesh3& SurfaceMesh);
 
 	FIntrinsicEdgeFlipMesh() = delete;
 	FIntrinsicEdgeFlipMesh(const FIntrinsicEdgeFlipMesh&) = delete;
@@ -669,7 +671,7 @@ public:
 	* Note: when tracing few edges, this method is preferred to using the FEdgeCorrespondence based tracing (see FEdgeCorrespondence below).  
 	* But when tracing a large percentage of the mesh edges the FEdgeCorrespondence will be faster as it reuses intermediate results.
 	*/
-	TArray<FSurfacePoint> TraceEdge(int32 IntrinsicEID, double CoalesceThreshold = 0., bool bReverse = false) const;
+	UE_API TArray<FSurfacePoint> TraceEdge(int32 IntrinsicEID, double CoalesceThreshold = 0., bool bReverse = false) const;
 
 
 	/**
@@ -678,7 +680,7 @@ public:
 	*
 	* @return a success or failure info.
 	*/
-	EMeshResult FlipEdge(int32 EID, FEdgeFlipInfo& EdgeFlipInfo);
+	UE_API EMeshResult FlipEdge(int32 EID, FEdgeFlipInfo& EdgeFlipInfo);
 
 
 	typedef TEdgeCorrespondence< FIntrinsicEdgeFlipMesh >  FEdgeCorrespondence;
@@ -704,7 +706,7 @@ public:
 	* 
 	* @return array of intrinsic edges crossed between the specified vertices.  This can be used to compute the actual MeshPoints on the surface if desired.
 	*/ 
-	TArray<FEdgeAndCrossingIdx> GetImplicitEdgeCrossings(const int32 SurfaceEID, const bool bReverse = false) const;
+	UE_API TArray<FEdgeAndCrossingIdx> GetImplicitEdgeCrossings(const int32 SurfaceEID, const bool bReverse = false) const;
 
 	/**
 	* Trace the specified surface (extrinsic) mesh edge across the intrinsic mesh.  By default this traces from vertex EdgeV.A to vertex EdgeV.B
@@ -714,7 +716,7 @@ public:
 	* 
 	* @return array of surface points defined relative to the intrinsic mesh.  In particular, the array will be a vertex followed by a sequence of (intrinsic) edge crossings.
 	*/ 
-	TArray<FSurfacePoint> TraceSurfaceEdge(int32 SurfaceEID, double CoalesceThreshold, bool bReverse = false) const;
+	UE_API TArray<FSurfacePoint> TraceSurfaceEdge(int32 SurfaceEID, double CoalesceThreshold, bool bReverse = false) const;
 
 	/**
 	* @return true if the intrinsic vertex corresponds to a surface vertex.  
@@ -761,7 +763,7 @@ protected:
 * Note:  This assumes but does not check that the FDynamicMesh3 "Surface Mesh" has consistent orientation of all triangles.
 */
 
-class DYNAMICMESH_API FIntrinsicMesh  : public  FSimpleIntrinsicMesh
+class FIntrinsicMesh  : public  FSimpleIntrinsicMesh
 {
 public:
 	typedef FSimpleIntrinsicMesh     MyBase;
@@ -776,7 +778,7 @@ public:
 	FIntrinsicMesh() = delete;
 	FIntrinsicMesh(const FIntrinsicMesh&) = delete;
 
-	FIntrinsicMesh(const FDynamicMesh3& SurfaceMesh);
+	UE_API FIntrinsicMesh(const FDynamicMesh3& SurfaceMesh);
 
 	/** @return pointer the reference mesh, note the FIntrinsicTriangulation does not manage the lifetime of this mesh */
 	const FDynamicMesh3* GetExtrinsicMesh()  const
@@ -791,7 +793,7 @@ public:
 	*
 	* @return a success or failure info.
 	*/
-	EMeshResult FlipEdge(int32 EID, FEdgeFlipInfo& EdgeFlipInfo);
+	UE_API EMeshResult FlipEdge(int32 EID, FEdgeFlipInfo& EdgeFlipInfo);
 
 	/**
 	* Insert a new vertex inside an intrinsic triangle, ie do a 1 to 3 triangle split
@@ -801,7 +803,7 @@ public:
 	* 
 	* @return Ok on success, or enum value indicates why operation cannot be applied. Mesh remains unmodified on error.
 	*/
-	EMeshResult PokeTriangle(int32 TID, const FVector3d& BaryCoordinates, FPokeTriangleInfo& PokeInfo);
+	UE_API EMeshResult PokeTriangle(int32 TID, const FVector3d& BaryCoordinates, FPokeTriangleInfo& PokeInfo);
 
 	/**
 	* Split an intrinsic edge of the mesh by inserting a vertex. This creates a new triangle on either side of the edge (ie a 2-4 split).
@@ -815,7 +817,7 @@ public:
 	* 
 	* @return Ok on success, or enum value indicates why operation cannot be applied. Mesh remains unmodified on error.
 	*/
-	EMeshResult SplitEdge(int32 EdgeAB, FEdgeSplitInfo& SplitInfo, double SplitParameterT = 0.5);
+	UE_API EMeshResult SplitEdge(int32 EdgeAB, FEdgeSplitInfo& SplitInfo, double SplitParameterT = 0.5);
 
 
 	/**
@@ -837,7 +839,7 @@ public:
 	* Note: when tracing few edges, this method is preferred to using the FEdgeCorrespondence based tracing (see FEdgeCorrespondence below).  
 	* But when tracing a large percentage of the mesh edges the FEdgeCorrespondence will be faster as it reuses intermediate results.
 	*/
-	TArray<FSurfacePoint> TraceEdge(int32 IntrinsicEID, double CoalesceThreshold = 0., bool bReverse = false) const;
+	UE_API TArray<FSurfacePoint> TraceEdge(int32 IntrinsicEID, double CoalesceThreshold = 0., bool bReverse = false) const;
 
 	/**
 	* Trace the specified surface (extrinsic) mesh edge across the intrinsic mesh.  By default this traces from vertex EdgeV.A to vertex EdgeV.B
@@ -847,7 +849,7 @@ public:
 	* 
 	* @return array of intrinsic edges crossed between the specified vertices.  This can be used to compute the actual MeshPoints on the surface if desired.
 	*/ 
-	TArray<FEdgeAndCrossingIdx> GetImplicitEdgeCrossings(const int32 SurfaceEID, const bool bReverse = false) const;
+	UE_API TArray<FEdgeAndCrossingIdx> GetImplicitEdgeCrossings(const int32 SurfaceEID, const bool bReverse = false) const;
 
 	/**
 	* Trace the specified surface (extrinsic) mesh edge across the intrinsic mesh.  By default this traces from vertex EdgeV.A to vertex EdgeV.B
@@ -857,7 +859,7 @@ public:
 	* 
 	* @return array of surface points defined relative to the intrinsic mesh.  In particular, the array will be a vertex followed by a sequence of (intrinsic) edge crossings.
 	*/ 
-	TArray<FSurfacePoint> TraceSurfaceEdge(int32 SurfaceEID, double CoalesceThreshold, bool bReverse = false) const;
+	UE_API TArray<FSurfacePoint> TraceSurfaceEdge(int32 SurfaceEID, double CoalesceThreshold, bool bReverse = false) const;
 
 
 	typedef TEdgeCorrespondence< FIntrinsicMesh >  FEdgeCorrespondence;
@@ -911,7 +913,7 @@ protected:
 * 
 * Note:  This assumes but does not check that the FDynamicMesh3 "Surface Mesh" has consistent orientation of all triangles.
 */ 
-class DYNAMICMESH_API FIntrinsicTriangulation : public  FSimpleIntrinsicMesh
+class FIntrinsicTriangulation : public  FSimpleIntrinsicMesh
 {
 
 public:
@@ -923,7 +925,7 @@ public:
 	using FPokeTriangleInfo = FSimpleIntrinsicMesh::FPokeTriangleInfo;
 	using FSurfacePoint  = IntrinsicCorrespondenceUtils::FSurfacePoint;
 
-	FIntrinsicTriangulation(const FDynamicMesh3& SrcMesh);
+	UE_API FIntrinsicTriangulation(const FDynamicMesh3& SrcMesh);
 
 	FIntrinsicTriangulation() = delete;
 	FIntrinsicTriangulation(const FIntrinsicTriangulation&) = delete;
@@ -944,7 +946,7 @@ public:
 	*                            and any resulting repetitions of vertex surface points are replaced with a single vertex surface point.
 	* @param bReverse          - if true, trace edge from EdgeV.B to EdgeV.A
 	*/
-	TArray<FIntrinsicTriangulation::FSurfacePoint> TraceEdge(int32 EID, double CoalesceThreshold = 0., bool bReverse = false) const;
+	UE_API TArray<FIntrinsicTriangulation::FSurfacePoint> TraceEdge(int32 EID, double CoalesceThreshold = 0., bool bReverse = false) const;
 
 	/**
 	* @return FSurfacePoint for the specified intrinsic vertex.
@@ -969,7 +971,7 @@ public:
 	*
 	* @return a success or failure info.
 	*/
-	EMeshResult FlipEdge(int32 EID, FEdgeFlipInfo& EdgeFlipInfo);
+	UE_API EMeshResult FlipEdge(int32 EID, FEdgeFlipInfo& EdgeFlipInfo);
 
 
 	/**
@@ -980,7 +982,7 @@ public:
 	* 
 	* @return Ok on success, or enum value indicates why operation cannot be applied. Mesh remains unmodified on error.
 	*/
-	EMeshResult PokeTriangle(int32 TID, const FVector3d& BaryCoordinates, FPokeTriangleInfo& PokeInfo);
+	UE_API EMeshResult PokeTriangle(int32 TID, const FVector3d& BaryCoordinates, FPokeTriangleInfo& PokeInfo);
 
 
 	/**
@@ -995,7 +997,7 @@ public:
 	 * 
 	 * @return Ok on success, or enum value indicates why operation cannot be applied. Mesh remains unmodified on error.
 	 */
-	EMeshResult SplitEdge(int32 EdgeAB, FEdgeSplitInfo& SplitInfo, double SplitParameterT = 0.5);
+	UE_API EMeshResult SplitEdge(int32 EdgeAB, FEdgeSplitInfo& SplitInfo, double SplitParameterT = 0.5);
 
 	
 
@@ -1014,7 +1016,7 @@ protected:
 	* @param TraceDistance  - Distance, measured on the extrinsic mesh surface to trace.
 	*
 	*/
-	double UpdateVertexByEdgeTrace(const int32 UpdateVID, const int32 StartVID, const double TracePolarDir, const double TraceDistance);
+	UE_API double UpdateVertexByEdgeTrace(const int32 UpdateVID, const int32 StartVID, const double TracePolarDir, const double TraceDistance);
 
 
 protected:
@@ -1041,3 +1043,5 @@ int32 DYNAMICMESH_API FlipToDelaunay(FIntrinsicTriangulation& IntrinsicMesh, TSe
 
 } // end namespace Geometry
 } // end namespace UE
+
+#undef UE_API
diff --git a/Engine/Plugins/Runtime/GeometryProcessing/Source/DynamicMesh/Public/Operations/JoinMeshLoops.h b/Engine/Plugins/Runtime/GeometryProcessing/Source/DynamicMesh/Public/Operations/JoinMeshLoops.h
--- a/Engine/Plugins/Runtime/GeometryProcessing/Source/DynamicMesh/Public/Operations/JoinMeshLoops.h
+++ b/Engine/Plugins/Runtime/GeometryProcessing/Source/DynamicMesh/Public/Operations/JoinMeshLoops.h
@@ -7,6 +7,8 @@
 #include "GeometryTypes.h"
 #include "MeshRegionBoundaryLoops.h"
 
+#define UE_API DYNAMICMESH_API
+
 
 
 namespace UE
@@ -26,7 +28,7 @@ class FMeshNormals;
  *    - the entire strip becomes a separate "normal island", ie hard normals at border with existing mesh
  *    - the strip is unwrapped into a UV rectangle starting at vertex 0 and ending back at 0 again
  */
-class DYNAMICMESH_API FJoinMeshLoops
+class FJoinMeshLoops
 {
 public:
 
@@ -62,8 +64,8 @@ public:
 	TArray<int32> QuadGroups;
 
 public:
-	FJoinMeshLoops(FDynamicMesh3* Mesh);
-	FJoinMeshLoops(FDynamicMesh3* Mesh, const TArray<int32>& LoopA, const TArray<int32>& LoopB);
+	UE_API FJoinMeshLoops(FDynamicMesh3* Mesh);
+	UE_API FJoinMeshLoops(FDynamicMesh3* Mesh, const TArray<int32>& LoopA, const TArray<int32>& LoopB);
 
 	virtual ~FJoinMeshLoops() {}
 
@@ -84,7 +86,7 @@ public:
 	 * Apply the Extrude operation to the input mesh.
 	 * @return true if the algorithm succeeds
 	 */
-	virtual bool Apply();
+	UE_API virtual bool Apply();
 
 
 protected:
@@ -93,4 +95,6 @@ protected:
 };
 
 } // end namespace UE::Geometry
-} // end namespace UE
\ No newline at end of file
+} // end namespace UE
+
+#undef UE_API
diff --git a/Engine/Plugins/Runtime/GeometryProcessing/Source/DynamicMesh/Public/Operations/MeshAttributeTransfer.h b/Engine/Plugins/Runtime/GeometryProcessing/Source/DynamicMesh/Public/Operations/MeshAttributeTransfer.h
--- a/Engine/Plugins/Runtime/GeometryProcessing/Source/DynamicMesh/Public/Operations/MeshAttributeTransfer.h
+++ b/Engine/Plugins/Runtime/GeometryProcessing/Source/DynamicMesh/Public/Operations/MeshAttributeTransfer.h
@@ -6,6 +6,8 @@
 #include "VectorTypes.h"
 #include "DynamicMesh/DynamicMeshAABBTree3.h"
 
+#define UE_API DYNAMICMESH_API
+
 namespace UE
 {
 namespace Geometry
@@ -30,7 +32,7 @@ enum class EMeshAttributeTransferError
 /**
  * FMeshAttributeTransfer transfers attributes from a SourceMesh to a TargetMesh
  */
-class DYNAMICMESH_API FMeshAttributeTransfer
+class FMeshAttributeTransfer
 {
 public:
 	/** The mesh that we are transferring from */
@@ -46,7 +48,7 @@ public:
 	TArray<EMeshAttributeTransferError> Errors;
 
 public:
-	FMeshAttributeTransfer(const FDynamicMesh3* SourceMeshIn, FDynamicMesh3* TargetMeshIn);
+	UE_API FMeshAttributeTransfer(const FDynamicMesh3* SourceMeshIn, FDynamicMesh3* TargetMeshIn);
 
 	virtual ~FMeshAttributeTransfer() {}
 
@@ -54,15 +56,17 @@ public:
 	 * Run the operation and modify .Mesh
 	 * @return true if the algorithm succeeds
 	 */
-	virtual bool Apply();
+	UE_API virtual bool Apply();
 
 protected:
 
 	TUniquePtr<FDynamicMeshAABBTree3> SourceSpatial;
 
 
-	bool Apply_MaterialID();
+	UE_API bool Apply_MaterialID();
 };
 
 } // end namespace UE::Geometry
-} // end namespace UE
\ No newline at end of file
+} // end namespace UE
+
+#undef UE_API
diff --git a/Engine/Plugins/Runtime/GeometryProcessing/Source/DynamicMesh/Public/Operations/MeshBevel.h b/Engine/Plugins/Runtime/GeometryProcessing/Source/DynamicMesh/Public/Operations/MeshBevel.h
--- a/Engine/Plugins/Runtime/GeometryProcessing/Source/DynamicMesh/Public/Operations/MeshBevel.h
+++ b/Engine/Plugins/Runtime/GeometryProcessing/Source/DynamicMesh/Public/Operations/MeshBevel.h
@@ -6,6 +6,8 @@
 #include "Selections/QuadGridPatch.h"
 #include "Math/InterpCurve.h"
 
+#define UE_API DYNAMICMESH_API
+
 namespace UE
 {
 namespace Geometry
@@ -44,7 +46,7 @@ class FEdgeLoop;
  * Currently "updating" the bevel shape, once the topology is known, is not supported, but could be implemented
  * relatively easily, as all the relevant information is already tracked and stored.
  */
-class DYNAMICMESH_API FMeshBevel
+class FMeshBevel
 {
 public:
 
@@ -94,18 +96,18 @@ public:
 	/**
 	 * Initialize the bevel with all edges of the given GroupTopology
 	 */
-	void InitializeFromGroupTopology(const FDynamicMesh3& Mesh, const FGroupTopology& Topology);
+	UE_API void InitializeFromGroupTopology(const FDynamicMesh3& Mesh, const FGroupTopology& Topology);
 
 	/**
 	 * Initialize the bevel with the specified edges of a GroupTopology
 	 */
-	void InitializeFromGroupTopologyEdges(const FDynamicMesh3& Mesh, const FGroupTopology& Topology, const TArray<int32>& GroupEdges);
+	UE_API void InitializeFromGroupTopologyEdges(const FDynamicMesh3& Mesh, const FGroupTopology& Topology, const TArray<int32>& GroupEdges);
 
 	/**
 	 * Initialize the bevel with the specified mesh triangle edges
 	 * @param IsCornerVertex Function to determine whether a vertex should be treated as a corner when beveling. Only affects vertices with exactly two adjacent TriangleEdges; otherwise corners will be automatically detected.
 	 */
-	void InitializeFromTriangleEdges(const FDynamicMesh3& Mesh, TConstArrayView<int32> TriangleEdges, TFunctionRef<bool(int32)> IsCornerVertex);
+	UE_API void InitializeFromTriangleEdges(const FDynamicMesh3& Mesh, TConstArrayView<int32> TriangleEdges, TFunctionRef<bool(int32)> IsCornerVertex);
 
 	/**
 	 * Initialize the bevel with the specified mesh triangle edges
@@ -119,18 +121,18 @@ public:
 	* Initialize the bevel with the specified faces of a GroupTopology
 	* @return false if any selection-bowtie vertices were found, in this case we cannot compute the bevel
 	*/
-	bool InitializeFromGroupTopologyFaces(const FDynamicMesh3& Mesh, const FGroupTopology& Topology, const TArray<int32>& GroupFaces);
+	UE_API bool InitializeFromGroupTopologyFaces(const FDynamicMesh3& Mesh, const FGroupTopology& Topology, const TArray<int32>& GroupFaces);
 
 	/**
 	* Initialize the bevel with border loops of the selected triangles. 
 	* @return false if any selection-bowtie vertices were found, in this case we cannot compute the bevel
 	*/
-	bool InitializeFromTriangleSet(const FDynamicMesh3& Mesh, const TArray<int32>& Triangles);
+	UE_API bool InitializeFromTriangleSet(const FDynamicMesh3& Mesh, const TArray<int32>& Triangles);
 
 	/**
 	 * Apply the bevel operation to the mesh, and optionally track changes
 	 */
-	bool Apply(FDynamicMesh3& Mesh, FDynamicMeshChangeTracker* ChangeTracker);
+	UE_API bool Apply(FDynamicMesh3& Mesh, FDynamicMeshChangeTracker* ChangeTracker);
 
 
 
@@ -300,75 +302,77 @@ protected:
 
 protected:
 
-	FBevelVertex* GetBevelVertexFromVertexID(int32 VertexID, int32* IndexOut = nullptr);
+	UE_API FBevelVertex* GetBevelVertexFromVertexID(int32 VertexID, int32* IndexOut = nullptr);
 
 	// Setup phase: register Edges (spans) and (isolated) Loops that need to be beveled and precompute/store any mesh topology that must be tracked across the operation
 	// Required BevelVertex's are added by AddBevelGroupEdge()
 	// Once edges are configured, BuildVertexSets() is called to precompute the vertex topological information
 
-	void AddBevelGroupEdge(const FDynamicMesh3& Mesh, const FGroupTopology& Topology, int32 GroupEdgeID);
-	void AddBevelEdgeLoop(const FDynamicMesh3& Mesh, const FEdgeLoop& Loop);
-	void BuildVertexSets(const FDynamicMesh3& Mesh);
-	void BuildJunctionVertex(FBevelVertex& Vertex, const FDynamicMesh3& Mesh);
-	void BuildTerminatorVertex(FBevelVertex& Vertex, const FDynamicMesh3& Mesh);
+	UE_API void AddBevelGroupEdge(const FDynamicMesh3& Mesh, const FGroupTopology& Topology, int32 GroupEdgeID);
+	UE_API void AddBevelEdgeLoop(const FDynamicMesh3& Mesh, const FEdgeLoop& Loop);
+	UE_API void BuildVertexSets(const FDynamicMesh3& Mesh);
+	UE_API void BuildJunctionVertex(FBevelVertex& Vertex, const FDynamicMesh3& Mesh);
+	UE_API void BuildTerminatorVertex(FBevelVertex& Vertex, const FDynamicMesh3& Mesh);
 
 	// Unlink phase - disconnect triangles along bevel edges/loops, and at vertices. 
 	// Vertices may expand out into multiple "wedges" depending on incoming bevel-edge topology.
 
-	void UnlinkEdges( FDynamicMesh3& Mesh, FDynamicMeshChangeTracker* ChangeTracker);
-	void UnlinkBevelEdgeInterior( FDynamicMesh3& Mesh, FBevelEdge& BevelEdge, FDynamicMeshChangeTracker* ChangeTracker);
+	UE_API void UnlinkEdges( FDynamicMesh3& Mesh, FDynamicMeshChangeTracker* ChangeTracker);
+	UE_API void UnlinkBevelEdgeInterior( FDynamicMesh3& Mesh, FBevelEdge& BevelEdge, FDynamicMeshChangeTracker* ChangeTracker);
 
-	void UnlinkLoops( FDynamicMesh3& Mesh, FDynamicMeshChangeTracker* ChangeTracker);
-	void UnlinkBevelLoop( FDynamicMesh3& Mesh, FBevelLoop& BevelLoop, FDynamicMeshChangeTracker* ChangeTracker);
+	UE_API void UnlinkLoops( FDynamicMesh3& Mesh, FDynamicMeshChangeTracker* ChangeTracker);
+	UE_API void UnlinkBevelLoop( FDynamicMesh3& Mesh, FBevelLoop& BevelLoop, FDynamicMeshChangeTracker* ChangeTracker);
 
-	void UnlinkVertices( FDynamicMesh3& Mesh, FDynamicMeshChangeTracker* ChangeTracker);
-	void UnlinkJunctionVertex(FDynamicMesh3& Mesh, FBevelVertex& BevelVertex, FDynamicMeshChangeTracker* ChangeTracker);
-	void UnlinkTerminatorVertex(FDynamicMesh3& Mesh, FBevelVertex& BevelVertex, FDynamicMeshChangeTracker* ChangeTracker);
+	UE_API void UnlinkVertices( FDynamicMesh3& Mesh, FDynamicMeshChangeTracker* ChangeTracker);
+	UE_API void UnlinkJunctionVertex(FDynamicMesh3& Mesh, FBevelVertex& BevelVertex, FDynamicMeshChangeTracker* ChangeTracker);
+	UE_API void UnlinkTerminatorVertex(FDynamicMesh3& Mesh, FBevelVertex& BevelVertex, FDynamicMeshChangeTracker* ChangeTracker);
 
-	void FixUpUnlinkedBevelEdges(FDynamicMesh3& Mesh);
+	UE_API void FixUpUnlinkedBevelEdges(FDynamicMesh3& Mesh);
 
 	// Displace phase - move unlinked vertices to new positions
 
-	void DisplaceVertices(FDynamicMesh3& Mesh, double Distance);
+	UE_API void DisplaceVertices(FDynamicMesh3& Mesh, double Distance);
 
 	// Meshing phase - append quad-strips between unlinked edge spans/loops, polygons at junction vertices where required,
 	// and triangles at terminator vertices
 
 	/* Meshing functions for chamfer bevel, ie no subdivisions */
-	void CreateBevelMeshing(FDynamicMesh3& Mesh);
-	void AppendJunctionVertexPolygon(FDynamicMesh3& Mesh, FBevelVertex& Vertex);
-	void AppendTerminatorVertexTriangle(FDynamicMesh3& Mesh, FBevelVertex& Vertex);
-	void AppendTerminatorVertexPairQuad(FDynamicMesh3& Mesh, FBevelVertex& Vertex0, FBevelVertex& Vertex1);
-	void AppendEdgeQuads(FDynamicMesh3& Mesh, FBevelEdge& Edge);
-	void AppendLoopQuads(FDynamicMesh3& Mesh, FBevelLoop& Loop);
+	UE_API void CreateBevelMeshing(FDynamicMesh3& Mesh);
+	UE_API void AppendJunctionVertexPolygon(FDynamicMesh3& Mesh, FBevelVertex& Vertex);
+	UE_API void AppendTerminatorVertexTriangle(FDynamicMesh3& Mesh, FBevelVertex& Vertex);
+	UE_API void AppendTerminatorVertexPairQuad(FDynamicMesh3& Mesh, FBevelVertex& Vertex0, FBevelVertex& Vertex1);
+	UE_API void AppendEdgeQuads(FDynamicMesh3& Mesh, FBevelEdge& Edge);
+	UE_API void AppendLoopQuads(FDynamicMesh3& Mesh, FBevelLoop& Loop);
 
 	/** Meshing functions for multi-segment bevel with optional round profile */
-	void CreateBevelMeshing_Multi(FDynamicMesh3& Mesh);
-	void AppendEdgeQuads_Multi(FDynamicMesh3& Mesh, FBevelEdge& Edge);
-	void AppendLoopQuads_Multi(FDynamicMesh3& Mesh, FBevelLoop& Loop);
-	void AppendJunctionVertexPolygon_Multi(FDynamicMesh3& Mesh, FBevelVertex& Vertex);
-	void AppendTerminatorVertexTriangles_Multi(FDynamicMesh3& Mesh, FBevelVertex& Vertex);
-	void AppendTerminatorVertexPairQuad_Multi(FDynamicMesh3& Mesh, FBevelVertex& Vertex0, FBevelVertex& Vertex1);
-	void ApplyProfileShape_Round(FDynamicMesh3& Mesh);
-	FInterpCurveVector MakeArcSplineCurve(const FVector3d& PosA, FVector3d& NormalA, const FVector3d& PosB, FVector3d& NormalB) const;
+	UE_API void CreateBevelMeshing_Multi(FDynamicMesh3& Mesh);
+	UE_API void AppendEdgeQuads_Multi(FDynamicMesh3& Mesh, FBevelEdge& Edge);
+	UE_API void AppendLoopQuads_Multi(FDynamicMesh3& Mesh, FBevelLoop& Loop);
+	UE_API void AppendJunctionVertexPolygon_Multi(FDynamicMesh3& Mesh, FBevelVertex& Vertex);
+	UE_API void AppendTerminatorVertexTriangles_Multi(FDynamicMesh3& Mesh, FBevelVertex& Vertex);
+	UE_API void AppendTerminatorVertexPairQuad_Multi(FDynamicMesh3& Mesh, FBevelVertex& Vertex0, FBevelVertex& Vertex1);
+	UE_API void ApplyProfileShape_Round(FDynamicMesh3& Mesh);
+	UE_API FInterpCurveVector MakeArcSplineCurve(const FVector3d& PosA, FVector3d& NormalA, const FVector3d& PosB, FVector3d& NormalB) const;
 
 
 	// Normals phase - calculate normals for new geometry
 
-	void ComputeNormals(FDynamicMesh3& Mesh);
-	void ComputeUVs(FDynamicMesh3& Mesh);
-	void ComputeMaterialIDs(FDynamicMesh3& Mesh);
+	UE_API void ComputeNormals(FDynamicMesh3& Mesh);
+	UE_API void ComputeUVs(FDynamicMesh3& Mesh);
+	UE_API void ComputeMaterialIDs(FDynamicMesh3& Mesh);
 
 private:
 	// Detect and fix any bowtie vertices in the bevel operation set
 	// Called by "Apply" (because the setup methods all operate on a const Mesh) before the actual bevel operation
-	void FixBowties(FDynamicMesh3& Mesh, FDynamicMeshChangeTracker* ChangeTracker);
+	UE_API void FixBowties(FDynamicMesh3& Mesh, FDynamicMeshChangeTracker* ChangeTracker);
 
-	void InitVertexSet(const FDynamicMesh3& Mesh, FBevelVertex& Vertex);
-	void FinalizeTerminatorVertex(const FDynamicMesh3& Mesh, FBevelVertex& Vertex);
+	UE_API void InitVertexSet(const FDynamicMesh3& Mesh, FBevelVertex& Vertex);
+	UE_API void FinalizeTerminatorVertex(const FDynamicMesh3& Mesh, FBevelVertex& Vertex);
 };
 
 
 
 } // end namespace UE::Geometry
-} // end namespace UE
\ No newline at end of file
+} // end namespace UE
+
+#undef UE_API
diff --git a/Engine/Plugins/Runtime/GeometryProcessing/Source/DynamicMesh/Public/Operations/MeshConvexHull.h b/Engine/Plugins/Runtime/GeometryProcessing/Source/DynamicMesh/Public/Operations/MeshConvexHull.h
--- a/Engine/Plugins/Runtime/GeometryProcessing/Source/DynamicMesh/Public/Operations/MeshConvexHull.h
+++ b/Engine/Plugins/Runtime/GeometryProcessing/Source/DynamicMesh/Public/Operations/MeshConvexHull.h
@@ -8,6 +8,8 @@
 #include "GeometryTypes.h"
 #include "DynamicMesh/DynamicMesh3.h"
 
+#define UE_API DYNAMICMESH_API
+
 class FProgressCancel;
 
 namespace UE
@@ -18,7 +20,7 @@ namespace Geometry
 /**
  * Calculate Convex Hull of a Mesh
  */
-class DYNAMICMESH_API FMeshConvexHull
+class FMeshConvexHull
 {
 public:
 	/** Input Mesh */
@@ -47,12 +49,12 @@ public:
 	 *  @param CellSizeAsPercentOfBounds	Grid cell size expressed as a percentage of the mesh bounding box size
 	 *  @param OutSamples					Indices of chosen vertices
 	 */
-	static void GridSample(const FDynamicMesh3& Mesh, 
+	static UE_API void GridSample(const FDynamicMesh3& Mesh, 
 						   int GridResolutionMaxAxis,
 						   TArray<int32>& OutSamples);
 
 	/** Used for testing/debugging */
-	static FVector3i DebugGetCellIndex(const FDynamicMesh3& Mesh,
+	static UE_API FVector3i DebugGetCellIndex(const FDynamicMesh3& Mesh,
 									   int GridResolutionMaxAxis,
 									   int VertexIndex);
 
@@ -67,21 +69,23 @@ public:
 	 * Calculate output ConvexHull mesh for vertices of input Mesh
 	 * @return true on success
 	 */
-	bool Compute(FProgressCancel* Progress = nullptr);
+	UE_API bool Compute(FProgressCancel* Progress = nullptr);
 
 	/**
 	 * Simplify the output convex hull. Assumes ConvexHull is already computed.
 	 * @return true on success
 	 */
-	static bool SimplifyHull(FDynamicMesh3& HullMesh, int32 MaxTargetFaceCount, FProgressCancel* Progress = nullptr);
+	static UE_API bool SimplifyHull(FDynamicMesh3& HullMesh, int32 MaxTargetFaceCount, FProgressCancel* Progress = nullptr);
 
 protected:
-	bool Compute_FullMesh(FProgressCancel* Progress);
-	bool Compute_VertexSubset(FProgressCancel* Progress);
+	UE_API bool Compute_FullMesh(FProgressCancel* Progress);
+	UE_API bool Compute_VertexSubset(FProgressCancel* Progress);
 private:
-	bool Compute_Helper(FProgressCancel* Progress, int32 MaxVertexIndex, TFunctionRef<FVector3d(int32)> GetVertex, TFunctionRef<bool(int32)> IsVertex, bool bTestMinDimension = true);
+	UE_API bool Compute_Helper(FProgressCancel* Progress, int32 MaxVertexIndex, TFunctionRef<FVector3d(int32)> GetVertex, TFunctionRef<bool(int32)> IsVertex, bool bTestMinDimension = true);
 };
 
 
 } // end namespace UE::Geometry
-} // end namespace UE
\ No newline at end of file
+} // end namespace UE
+
+#undef UE_API
diff --git a/Engine/Plugins/Runtime/GeometryProcessing/Source/DynamicMesh/Public/Operations/MeshGeodesicSurfaceTracer.h b/Engine/Plugins/Runtime/GeometryProcessing/Source/DynamicMesh/Public/Operations/MeshGeodesicSurfaceTracer.h
--- a/Engine/Plugins/Runtime/GeometryProcessing/Source/DynamicMesh/Public/Operations/MeshGeodesicSurfaceTracer.h
+++ b/Engine/Plugins/Runtime/GeometryProcessing/Source/DynamicMesh/Public/Operations/MeshGeodesicSurfaceTracer.h
@@ -9,6 +9,8 @@
 #include "MathUtil.h" // for max real
 #include "TriangleTypes.h"
 
+#define UE_API DYNAMICMESH_API
+
 
 
 namespace UE {
@@ -197,7 +199,7 @@ public:
 * in terms of barycentric coordinates, triangle references and local path direction relative 
 * to a local frame (in the form of one of the edges of the triangle).
 */
-class DYNAMICMESH_API FMeshGeodesicSurfaceTracer
+class FMeshGeodesicSurfaceTracer
 {
 public:
 
@@ -234,7 +236,7 @@ public:
 	*
 	* NB: the BaryPoint must be a valid point within or on the boundary of the triangle.
 	*/
-	double TraceMeshFromBaryPoint(const int32 TriID, const FVector3d& BaryPoint, const FVector3d& RayDir3, double MaxDistance = TMathUtilConstants<double>::MaxReal);
+	UE_API double TraceMeshFromBaryPoint(const int32 TriID, const FVector3d& BaryPoint, const FVector3d& RayDir3, double MaxDistance = TMathUtilConstants<double>::MaxReal);
 
 	/**
 	* Populate internal SurfaceTrace array of FTraceResults while following a "straight" path on the surface of a mesh.
@@ -250,7 +252,7 @@ public:
 	* NB: the BaryPoint must be a valid point within or on the boundary of the triangle.
 	* NB: it is assumed the edge Direction.EdgeID is part of the triangle TriID.  ( otherwise the result will be unexpected)
 	*/
-	double TraceMeshFromBaryPoint(const int32 TriID, const FVector3d& BaryPoint, const FMeshSurfaceDirection& Direction, double MaxDistance = TMathUtilConstants<double>::MaxReal);
+	UE_API double TraceMeshFromBaryPoint(const int32 TriID, const FVector3d& BaryPoint, const FMeshSurfaceDirection& Direction, double MaxDistance = TMathUtilConstants<double>::MaxReal);
 
 	/**
 	* Populate internal SurfaceTrace array of FTraceResults while following a "straight" path on the surface of a mesh.
@@ -263,7 +265,7 @@ public:
 	*
 	* NB: it is assumed the edge Direction.EdgeID is adjacent to VID.  ( otherwise the result will be unexpected)
 	*/
-	double TraceMeshFromVertex( const FMeshTangentDirection& Direction, double MaxDistance);
+	UE_API double TraceMeshFromVertex( const FMeshTangentDirection& Direction, double MaxDistance);
 
 	/** Access to the last tri traced in 2d form.*/
 	const FTangentTri2& GetLastTri()
@@ -378,3 +380,4 @@ static double AsZeroToTwoPi(double AngleR)
 }; // end namespace Geometry
 }; // end namespace UE
 
+#undef UE_API
diff --git a/Engine/Plugins/Runtime/GeometryProcessing/Source/DynamicMesh/Public/Operations/MeshIsoCurves.h b/Engine/Plugins/Runtime/GeometryProcessing/Source/DynamicMesh/Public/Operations/MeshIsoCurves.h
--- a/Engine/Plugins/Runtime/GeometryProcessing/Source/DynamicMesh/Public/Operations/MeshIsoCurves.h
+++ b/Engine/Plugins/Runtime/GeometryProcessing/Source/DynamicMesh/Public/Operations/MeshIsoCurves.h
@@ -6,6 +6,8 @@
 #include "VectorTypes.h"
 #include "GeometryTypes.h"
 
+#define UE_API DYNAMICMESH_API
+
 namespace UE
 {
 namespace Geometry
@@ -31,7 +33,7 @@ struct FMeshIsoCurveSettings
 /**
  * Insert edges on a mesh along the isocurve where some scalar value function over the mesh surface crosses a specified value
  */
-class DYNAMICMESH_API FMeshIsoCurves
+class FMeshIsoCurves
 {
 public:
 
@@ -47,7 +49,7 @@ public:
 	 * @param EdgeCutFn Given the vertices of an edge and their values, return the parameter where the edge should be cut. Only called if IsoValue is crossed between ValueA and ValueB.
 	 * @param IsoValue Value at which to insert a new curve on the mesh
 	 */
-	void Cut(FDynamicMesh3& Mesh, TFunctionRef<float(int32)> VertexFn, TFunctionRef<float(int32 VertA, int32 VertB, float ValueA, float ValueB)> EdgeCutFn, float IsoValue = 0);
+	UE_API void Cut(FDynamicMesh3& Mesh, TFunctionRef<float(int32)> VertexFn, TFunctionRef<float(int32 VertA, int32 VertB, float ValueA, float ValueB)> EdgeCutFn, float IsoValue = 0);
 
 	/**
 	 * Insert new edges on the given mesh along the curve where a function over the mesh vertices, and linearly interpolated over edges, crosses a given isovalue
@@ -67,11 +69,13 @@ public:
 	}
 
 private:
-	void SplitCrossingEdges(FDynamicMesh3& Mesh, const TArray<float>& VertexValues, TSet<int32>& OnCutEdges,
+	UE_API void SplitCrossingEdges(FDynamicMesh3& Mesh, const TArray<float>& VertexValues, TSet<int32>& OnCutEdges,
 		TFunctionRef<float(int32 VIDMin, int32 VIDMax, float ValMin, float ValMax)> EdgeCutFn,
 		float IsoValue);
 };
 
 
 } // end namespace UE::Geometry
-} // end namespace UE
\ No newline at end of file
+} // end namespace UE
+
+#undef UE_API
diff --git a/Engine/Plugins/Runtime/GeometryProcessing/Source/DynamicMesh/Public/Operations/MeshMirror.h b/Engine/Plugins/Runtime/GeometryProcessing/Source/DynamicMesh/Public/Operations/MeshMirror.h
--- a/Engine/Plugins/Runtime/GeometryProcessing/Source/DynamicMesh/Public/Operations/MeshMirror.h
+++ b/Engine/Plugins/Runtime/GeometryProcessing/Source/DynamicMesh/Public/Operations/MeshMirror.h
@@ -4,12 +4,14 @@
 #include "DynamicMesh/DynamicMesh3.h"
 #include "Util/ProgressCancel.h"
 
+#define UE_API DYNAMICMESH_API
+
 namespace UE
 {
 namespace Geometry
 {
 
-class DYNAMICMESH_API FMeshMirror
+class FMeshMirror
 {
 
 public:
@@ -42,15 +44,17 @@ public:
 	 *
 	 * @param Progress Object used to cancel the operation early. This leaves the mesh in an undefined state.
 	 */
-	void Mirror(FProgressCancel *Progress = nullptr);
+	UE_API void Mirror(FProgressCancel *Progress = nullptr);
 
 	/**
 	 * Appends a mirrored copy of the mesh to the mesh.
 	 *
 	 * @param Progress Object used to cancel the operation early. This leaves the mesh in an undefined state.
 	 */
-	void MirrorAndAppend(FProgressCancel* Progress = nullptr);
+	UE_API void MirrorAndAppend(FProgressCancel* Progress = nullptr);
 };
 
 } // end namespace UE::Geometry
-} // end namespace UE
\ No newline at end of file
+} // end namespace UE
+
+#undef UE_API
diff --git a/Engine/Plugins/Runtime/GeometryProcessing/Source/DynamicMesh/Public/Operations/MeshPlaneCut.h b/Engine/Plugins/Runtime/GeometryProcessing/Source/DynamicMesh/Public/Operations/MeshPlaneCut.h
--- a/Engine/Plugins/Runtime/GeometryProcessing/Source/DynamicMesh/Public/Operations/MeshPlaneCut.h
+++ b/Engine/Plugins/Runtime/GeometryProcessing/Source/DynamicMesh/Public/Operations/MeshPlaneCut.h
@@ -12,6 +12,8 @@
 
 #include "Operations/LocalPlanarSimplify.h"
 
+#define UE_API DYNAMICMESH_API
+
 
 namespace UE
 {
@@ -45,7 +47,7 @@ public:
  *	  6) (optionally) change an attribute tag for all triangles on positive side
  *    7) find loops through valid boundary edges (ie connected to splits, or on-plane edges) (if second half was kept, do this separately for each separate mesh ID label)
  */
-class DYNAMICMESH_API FMeshPlaneCut : public FMeshPlaneCutFixLayout
+class FMeshPlaneCut : public FMeshPlaneCutFixLayout
 {
 public:
 
@@ -142,7 +144,7 @@ public:
 	 * on the positive side of the cutting plane.
 	 * @return true if operation succeeds
 	 */
-	virtual bool Cut();
+	UE_API virtual bool Cut();
 
 	/**
 	 * Compute the plane cut by splitting mesh edges that cross the cut plane, but not deleting triangles on positive side.
@@ -154,7 +156,7 @@ public:
 	 * @param bAddBoundariesSecondHalf add open boundaries on "second" half to OpenBoundaries list (default true)
 	 * @return true if operation succeeds
 	 */
-	virtual bool CutWithoutDelete(bool bSplitVerticesAtPlane, float OffsetSeparatedPortion = 0.0f,
+	UE_API virtual bool CutWithoutDelete(bool bSplitVerticesAtPlane, float OffsetSeparatedPortion = 0.0f,
 		TDynamicMeshScalarTriangleAttribute<int>* TriLabels = nullptr, int NewLabelStartID = 0,
 		bool bAddBoundariesFirstHalf = true, bool bAddBoundariesSecondHalf = true);
 
@@ -178,36 +180,36 @@ public:
 	/**
 	 *  Fill cut loops with FSimpleHoleFiller
 	 */
-	virtual bool SimpleHoleFill(int ConstantGroupID = -1);
+	UE_API virtual bool SimpleHoleFill(int ConstantGroupID = -1);
 
-	virtual bool MinimalHoleFill(int ConstantGroupID = -1);
+	UE_API virtual bool MinimalHoleFill(int ConstantGroupID = -1);
 
 	/**
 	 *  Fill cut loops with FPlanarHoleFiller, using a caller-provided triangulation function
 	 */
-	virtual bool HoleFill(TFunction<TArray<FIndex3i>(const FGeneralPolygon2d&)> PlanarTriangulationFunc, bool bFillSpans, int ConstantGroupID = -1, int MaterialID = -1);
+	UE_API virtual bool HoleFill(TFunction<TArray<FIndex3i>(const FGeneralPolygon2d&)> PlanarTriangulationFunc, bool bFillSpans, int ConstantGroupID = -1, int MaterialID = -1);
 
 	
-	virtual void TransferTriangleLabelsToHoleFillTriangles(TDynamicMeshScalarTriangleAttribute<int>* TriLabels);
+	UE_API virtual void TransferTriangleLabelsToHoleFillTriangles(TDynamicMeshScalarTriangleAttribute<int>* TriLabels);
 
 protected:
 
 	UE_DEPRECATED(5.3, "Use the single-set CollapseDegenerateEdges instead")
-	void CollapseDegenerateEdges(const TSet<int>& OnCutEdges, const TSet<int>& ZeroEdges);
+	UE_API void CollapseDegenerateEdges(const TSet<int>& OnCutEdges, const TSet<int>& ZeroEdges);
 	
 	/// Collapse degenerate edges
 	/// @param Edges									Edges to consider for collapse; will updated by removing edges at they are collapsed
 	/// @param bRemoveAllDegenerateFromInputSet			Whether we should also check whether the neighbor edges removed by collapse were also in the set
 	/// 												(typically this is not needed, and has additional cost)
 	/// @param TriangleSelection						Optional set tracking an active selection. Any triangles removed by collapse will also be removed from the set.
-	void CollapseDegenerateEdges(TSet<int>& Edges, bool bRemoveAllDegenerateFromInputSet, TSet<int>* TriangleSelection = nullptr);
+	UE_API void CollapseDegenerateEdges(TSet<int>& Edges, bool bRemoveAllDegenerateFromInputSet, TSet<int>* TriangleSelection = nullptr);
 	
-	void SplitCrossingEdges(TArray<double>& Signs, TSet<int>& ZeroEdges, TSet<int>& OnCutEdges, bool bDeleteTrisOnPlane = true);
-	void SplitCrossingEdges(TArray<double>& Signs, TSet<int>& ZeroEdges, TSet<int>& OnCutEdges, TSet<int>& OnSplitEdges, bool bDeleteTrisOnPlane = true);
-	void SplitCrossingEdges(bool bDeleteTrisOnPlane, TArray<double>& Signs, TSet<int>& AlreadyOnPlaneEdges, TSet<int32>& CutPlaneEdges, 
+	UE_API void SplitCrossingEdges(TArray<double>& Signs, TSet<int>& ZeroEdges, TSet<int>& OnCutEdges, bool bDeleteTrisOnPlane = true);
+	UE_API void SplitCrossingEdges(TArray<double>& Signs, TSet<int>& ZeroEdges, TSet<int>& OnCutEdges, TSet<int>& OnSplitEdges, bool bDeleteTrisOnPlane = true);
+	UE_API void SplitCrossingEdges(bool bDeleteTrisOnPlane, TArray<double>& Signs, TSet<int>& AlreadyOnPlaneEdges, TSet<int32>& CutPlaneEdges, 
 		TSet<int>* SplitEdges = nullptr, TSet<int>* OnPlaneVertices = nullptr, TSet<int>* TriangleSelection = nullptr);
 	
-	bool ExtractBoundaryLoops(const TSet<int>& OnCutEdges, const TSet<int>& ZeroEdges, FMeshPlaneCut::FOpenBoundary& Boundary);
+	UE_API bool ExtractBoundaryLoops(const TSet<int>& OnCutEdges, const TSet<int>& ZeroEdges, FMeshPlaneCut::FOpenBoundary& Boundary);
 
 	// set of vertices lying on plane after calling SplitCrossingEdges
 	UE_DEPRECATED(5.3, "If needed, explicitly request the vertices on the cut plane the the OnPlaneVertices argument of SplitCrossingEdges")
@@ -216,11 +218,13 @@ protected:
 private:
 
 	// Helper to compute signed distances from the cutting plane for all vertices of the mesh.  Value at Invalid Vertex IDs will be set to InvalidDist.
-	void ComputeVertexSignedDistances(TArray<double>& Signs, double InvalidDist);
+	UE_API void ComputeVertexSignedDistances(TArray<double>& Signs, double InvalidDist);
 
-	bool SplitEdgesOnlyHelper(bool bAssignNewGroups, TSet<int32>* OptionalTriangleSelection, bool bAddDeprecatedResultSeedTriangles);
+	UE_API bool SplitEdgesOnlyHelper(bool bAssignNewGroups, TSet<int32>* OptionalTriangleSelection, bool bAddDeprecatedResultSeedTriangles);
 };
 
 
 } // end namespace UE::Geometry
-} // end namespace UE
\ No newline at end of file
+} // end namespace UE
+
+#undef UE_API
diff --git a/Engine/Plugins/Runtime/GeometryProcessing/Source/DynamicMesh/Public/Operations/MeshProjectionHull.h b/Engine/Plugins/Runtime/GeometryProcessing/Source/DynamicMesh/Public/Operations/MeshProjectionHull.h
--- a/Engine/Plugins/Runtime/GeometryProcessing/Source/DynamicMesh/Public/Operations/MeshProjectionHull.h
+++ b/Engine/Plugins/Runtime/GeometryProcessing/Source/DynamicMesh/Public/Operations/MeshProjectionHull.h
@@ -8,6 +8,8 @@
 #include "Polygon2.h"
 #include "DynamicMesh/DynamicMesh3.h"
 
+#define UE_API DYNAMICMESH_API
+
 namespace UE
 {
 namespace Geometry
@@ -17,7 +19,7 @@ namespace Geometry
  * Calculate a Convex Hull for a Mesh by first Projecting all vertices to a plane, computing a
  * 2D convex polygon that contains them, and then sweeping that 2D hull to create an extruded 3D volume.
  */
-class DYNAMICMESH_API FMeshProjectionHull
+class FMeshProjectionHull
 {
 public:
 	/** Input Mesh */
@@ -70,14 +72,16 @@ public:
 	 * Calculate output 2D Convex Polygon and Swept-Polygon 3D Mesh for vertices of input Mesh
 	 * @return true on success
 	 */
-	bool Compute();
+	UE_API bool Compute();
 
 private:
 
 	// Helper to compute the 3D hull when Keep3DHullSide is not None; called by Compute after the 2D projection hull is computed
-	bool ComputeWith3DHullSide(FFrame3d CenterFrame, double ExtrudeLength);
+	UE_API bool ComputeWith3DHullSide(FFrame3d CenterFrame, double ExtrudeLength);
 
 };
 
 } // end namespace UE::Geometry
-} // end namespace UE
\ No newline at end of file
+} // end namespace UE
+
+#undef UE_API
diff --git a/Engine/Plugins/Runtime/GeometryProcessing/Source/DynamicMesh/Public/Operations/MeshRegionOperator.h b/Engine/Plugins/Runtime/GeometryProcessing/Source/DynamicMesh/Public/Operations/MeshRegionOperator.h
--- a/Engine/Plugins/Runtime/GeometryProcessing/Source/DynamicMesh/Public/Operations/MeshRegionOperator.h
+++ b/Engine/Plugins/Runtime/GeometryProcessing/Source/DynamicMesh/Public/Operations/MeshRegionOperator.h
@@ -28,7 +28,7 @@ namespace Geometry
    trying to guess it here, by making some assumptions about what happens. It works for now, but
    it would better if FDynamicMeshEditor returned this information.
  */
-class DYNAMICMESH_API FMeshRegionOperator
+class FMeshRegionOperator
 {
 public:
 	FDynamicMesh3* BaseMesh;
@@ -257,4 +257,4 @@ public:
 
 
 } // end namespace UE::Geometry
-} // end namespace UE
\ No newline at end of file
+} // end namespace UE
diff --git a/Engine/Plugins/Runtime/GeometryProcessing/Source/DynamicMesh/Public/Operations/MeshResolveTJunctions.h b/Engine/Plugins/Runtime/GeometryProcessing/Source/DynamicMesh/Public/Operations/MeshResolveTJunctions.h
--- a/Engine/Plugins/Runtime/GeometryProcessing/Source/DynamicMesh/Public/Operations/MeshResolveTJunctions.h
+++ b/Engine/Plugins/Runtime/GeometryProcessing/Source/DynamicMesh/Public/Operations/MeshResolveTJunctions.h
@@ -5,6 +5,8 @@
 #include "MathUtil.h"
 #include "VectorTypes.h"
 
+#define UE_API DYNAMICMESH_API
+
 namespace UE
 {
 namespace Geometry
@@ -39,11 +41,11 @@ class FDynamicMesh3;
  * Current implementation is O(N*M) in the number of boundary edges (N) and boundary vertices (M).
  * Could be improved with a spatial data structure. 
  */
-class DYNAMICMESH_API FMeshResolveTJunctions
+class FMeshResolveTJunctions
 {
 public:
 	/** default tolerance is float ZeroTolerance */
-	static const double DEFAULT_TOLERANCE;  // = FMathf::ZeroTolerance;
+	static UE_API const double DEFAULT_TOLERANCE;  // = FMathf::ZeroTolerance;
 
 	/** The mesh that we are modifying */
 	FDynamicMesh3* Mesh;
@@ -67,7 +69,7 @@ public:
 	 * Run the resolve operation and modify .Mesh
 	 * @return true if the algorithm succeeds
 	 */
-	bool Apply();
+	UE_API bool Apply();
 
 };
 
@@ -120,4 +122,6 @@ public:
 };
 
 } // end namespace UE::Geometry
-} // end namespace UE
\ No newline at end of file
+} // end namespace UE
+
+#undef UE_API
diff --git a/Engine/Plugins/Runtime/GeometryProcessing/Source/DynamicMesh/Public/Operations/MinimalHoleFiller.h b/Engine/Plugins/Runtime/GeometryProcessing/Source/DynamicMesh/Public/Operations/MinimalHoleFiller.h
--- a/Engine/Plugins/Runtime/GeometryProcessing/Source/DynamicMesh/Public/Operations/MinimalHoleFiller.h
+++ b/Engine/Plugins/Runtime/GeometryProcessing/Source/DynamicMesh/Public/Operations/MinimalHoleFiller.h
@@ -15,6 +15,8 @@
 #include "QueueRemesher.h"
 #include "MeshConstraintsUtil.h"
 
+#define UE_API DYNAMICMESH_API
+
 namespace UE
 {
 namespace Geometry
@@ -24,7 +26,7 @@ namespace Geometry
 /// Construct a "minimal" fill surface for the hole. This surface is often quasi-developable, reconstructs sharp edges, 
 /// etc. There are various options.
 
-class DYNAMICMESH_API FMinimalHoleFiller : public IHoleFiller
+class FMinimalHoleFiller : public IHoleFiller
 {
 public:
 
@@ -34,7 +36,7 @@ public:
 		FillMesh(nullptr)
 	{}
 
-	bool Fill(int32 GroupID = -1) override;
+	UE_API bool Fill(int32 GroupID = -1) override;
 
 	// Settings
 	bool bIgnoreBoundaryTriangles = false;
@@ -54,24 +56,26 @@ private:
 	TMap<int, double> ExteriorAngleSums;
 	TArray<double> Curvatures;
 
-	void AddAllEdges(int EdgeID, TSet<int>& EdgeSet);
+	UE_API void AddAllEdges(int EdgeID, TSet<int>& EdgeSet);
 
-	double AspectMetric(int eid);
-	static double GetTriAspect(const FDynamicMesh3& mesh, FIndex3i& tri);
+	UE_API double AspectMetric(int eid);
+	static UE_API double GetTriAspect(const FDynamicMesh3& mesh, FIndex3i& tri);
 
-	void UpdateCurvature(int vid);
-	double CurvatureMetricCached(int a, int b, int c, int d);
-	double CurvatureMetricEval(int a, int b, int c, int d);
-	double ComputeGaussCurvature(int vid);
+	UE_API void UpdateCurvature(int vid);
+	UE_API double CurvatureMetricCached(int a, int b, int c, int d);
+	UE_API double CurvatureMetricEval(int a, int b, int c, int d);
+	UE_API double ComputeGaussCurvature(int vid);
 
 	// Steps in Fill:
-	void CollapseToMinimal();
-	void RemoveRemainingInteriorVerts();
-	void FlipToFlatter();
-	void FlipToMinimizeCurvature();
-	void FlipToImproveAspectRatios();
+	UE_API void CollapseToMinimal();
+	UE_API void RemoveRemainingInteriorVerts();
+	UE_API void FlipToFlatter();
+	UE_API void FlipToMinimizeCurvature();
+	UE_API void FlipToImproveAspectRatios();
 };
 
 
 } // end namespace UE::Geometry
-} // end namespace UE
\ No newline at end of file
+} // end namespace UE
+
+#undef UE_API
diff --git a/Engine/Plugins/Runtime/GeometryProcessing/Source/DynamicMesh/Public/Operations/OffsetMeshRegion.h b/Engine/Plugins/Runtime/GeometryProcessing/Source/DynamicMesh/Public/Operations/OffsetMeshRegion.h
--- a/Engine/Plugins/Runtime/GeometryProcessing/Source/DynamicMesh/Public/Operations/OffsetMeshRegion.h
+++ b/Engine/Plugins/Runtime/GeometryProcessing/Source/DynamicMesh/Public/Operations/OffsetMeshRegion.h
@@ -8,6 +8,8 @@
 #include "MeshRegionBoundaryLoops.h"
 #include "DynamicMesh/DynamicMeshChangeTracker.h"
 
+#define UE_API DYNAMICMESH_API
+
 
 namespace UE
 {
@@ -35,7 +37,7 @@ public:
  * 
  * Each quad of the border loop is assigned it's own normal and UVs (ie each is a separate UV-island)
  */
-class DYNAMICMESH_API FOffsetMeshRegion : public FOffsetMeshRegionFixLayout
+class FOffsetMeshRegion : public FOffsetMeshRegionFixLayout
 {
 public:
 
@@ -222,7 +224,7 @@ public:
 	TArray<int32> AllModifiedAndNewTriangles;
 
 public:
-	FOffsetMeshRegion(FDynamicMesh3* mesh);
+	UE_API FOffsetMeshRegion(FDynamicMesh3* mesh);
 
 	virtual ~FOffsetMeshRegion() {}
 
@@ -244,18 +246,20 @@ public:
 	 * Apply the Offset operation to the input mesh.
 	 * @return true if the algorithm succeeds
 	 */
-	virtual bool Apply();
+	UE_API virtual bool Apply();
 
 	// The default function used for LoopEdgesShouldHaveSameGroup. 
-	static bool EdgesSeparateSameGroupsAndAreColinearAtBorder(FDynamicMesh3* Mesh,
+	static UE_API bool EdgesSeparateSameGroupsAndAreColinearAtBorder(FDynamicMesh3* Mesh,
 		int32 Eid1, int32 Eid2, bool bCheckColinearityAtBorder);
 protected:
 
-	virtual bool ApplyOffset(FOffsetInfo& Region);
+	UE_API virtual bool ApplyOffset(FOffsetInfo& Region);
 
-	virtual bool ApplyOffset_Legacy(FOffsetInfo& Region);
-	virtual bool ApplyOffset_Version1(FOffsetInfo& Region);
+	UE_API virtual bool ApplyOffset_Legacy(FOffsetInfo& Region);
+	UE_API virtual bool ApplyOffset_Version1(FOffsetInfo& Region);
 };
 
 } // end namespace UE::Geometry
 } // end namespace UE
+
+#undef UE_API
diff --git a/Engine/Plugins/Runtime/GeometryProcessing/Source/DynamicMesh/Public/Operations/PNTriangles.h b/Engine/Plugins/Runtime/GeometryProcessing/Source/DynamicMesh/Public/Operations/PNTriangles.h
--- a/Engine/Plugins/Runtime/GeometryProcessing/Source/DynamicMesh/Public/Operations/PNTriangles.h
+++ b/Engine/Plugins/Runtime/GeometryProcessing/Source/DynamicMesh/Public/Operations/PNTriangles.h
@@ -4,6 +4,8 @@
 
 #include "GeometryTypes.h"
 
+#define UE_API DYNAMICMESH_API
+
 class FProgressCancel;
 
 namespace UE
@@ -20,7 +22,7 @@ class FDynamicMesh3;
  * points. The patch matches the point and normal information at the vertices of the original flat triangle.
  */
 
-class DYNAMICMESH_API FPNTriangles
+class FPNTriangles
 {
 public:
 
@@ -74,8 +76,10 @@ public:
 	 * Generate PN Triangles geometry and optionally compute quadratically varying normals.
 	 * @return true if the algorithm succeeds, false if it failed or was canceled by the user.
 	 */
-	virtual bool Compute();
+	UE_API virtual bool Compute();
 };
 
 } // end namespace UE::Geometry
 } // end namespace UE
+
+#undef UE_API
diff --git a/Engine/Plugins/Runtime/GeometryProcessing/Source/DynamicMesh/Public/Operations/PlanarFlipsOptimization.h b/Engine/Plugins/Runtime/GeometryProcessing/Source/DynamicMesh/Public/Operations/PlanarFlipsOptimization.h
--- a/Engine/Plugins/Runtime/GeometryProcessing/Source/DynamicMesh/Public/Operations/PlanarFlipsOptimization.h
+++ b/Engine/Plugins/Runtime/GeometryProcessing/Source/DynamicMesh/Public/Operations/PlanarFlipsOptimization.h
@@ -4,6 +4,8 @@
 
 #include "CoreMinimal.h"
 
+#define UE_API DYNAMICMESH_API
+
 namespace UE
 {
 namespace Geometry
@@ -17,7 +19,7 @@ class FDynamicMesh3;
  * the edge without affecting the shape to try to remove it. This code does one or more passes of
  * such an optimization.
  */
-class DYNAMICMESH_API FPlanarFlipsOptimization
+class FPlanarFlipsOptimization
 {
 public:
 
@@ -36,11 +38,13 @@ public:
 	/**
 	 * Apply the operation to the mesh.
 	 */
-	void Apply();
+	UE_API void Apply();
 
 protected:
-	void ApplySinglePass();
+	UE_API void ApplySinglePass();
 };
 
 } // end namespace UE::Geometry
 } // end namespace UE
+
+#undef UE_API
diff --git a/Engine/Plugins/Runtime/GeometryProcessing/Source/DynamicMesh/Public/Operations/PlanarHoleFiller.h b/Engine/Plugins/Runtime/GeometryProcessing/Source/DynamicMesh/Public/Operations/PlanarHoleFiller.h
--- a/Engine/Plugins/Runtime/GeometryProcessing/Source/DynamicMesh/Public/Operations/PlanarHoleFiller.h
+++ b/Engine/Plugins/Runtime/GeometryProcessing/Source/DynamicMesh/Public/Operations/PlanarHoleFiller.h
@@ -13,6 +13,8 @@
 #include "Curve/GeneralPolygon2.h"
 #include "MeshBoundaryLoops.h"
 
+#define UE_API DYNAMICMESH_API
+
 
 namespace UE
 {
@@ -24,7 +26,7 @@ class FDynamicMesh3;
 /**
  * Fill a set of boundary loops with planar surfaces.  User must provide the triangulation function.
  */
-class DYNAMICMESH_API FPlanarHoleFiller : public IHoleFiller
+class FPlanarHoleFiller : public IHoleFiller
 {
 public:
 	//
@@ -59,9 +61,11 @@ public:
 		return EOperationValidationResult::Ok;
 	}
 
-	bool Fill(int GroupID = -1) override;
+	UE_API bool Fill(int GroupID = -1) override;
 };
 
 
 } // end namespace UE::Geometry
-} // end namespace UE
\ No newline at end of file
+} // end namespace UE
+
+#undef UE_API
diff --git a/Engine/Plugins/Runtime/GeometryProcessing/Source/DynamicMesh/Public/Operations/PolygroupRemesh.h b/Engine/Plugins/Runtime/GeometryProcessing/Source/DynamicMesh/Public/Operations/PolygroupRemesh.h
--- a/Engine/Plugins/Runtime/GeometryProcessing/Source/DynamicMesh/Public/Operations/PolygroupRemesh.h
+++ b/Engine/Plugins/Runtime/GeometryProcessing/Source/DynamicMesh/Public/Operations/PolygroupRemesh.h
@@ -17,6 +17,8 @@
 
 #include "Util/ProgressCancel.h"
 
+#define UE_API DYNAMICMESH_API
+
 
 namespace UE
 {
@@ -34,7 +36,7 @@ public:
  * PolygroupRemesh -- remesh only considering polygroup features (topological corners and bends on polygroup edges)
  * This can help clean up low poly meshes that have extra vertices along straight edges e.g. after mesh Boolean operations
  */
-class DYNAMICMESH_API FPolygroupRemesh : public FPolygroupRemeshFixLayout
+class FPolygroupRemesh : public FPolygroupRemeshFixLayout
 {
 public:
 
@@ -83,7 +85,7 @@ public:
 	 * on the positive side of the cutting plane.
 	 * @return true if operation succeeds
 	 */
-	bool Compute();
+	UE_API bool Compute();
 
 protected:
 	/** If this returns true, abort computation.  */
@@ -96,4 +98,6 @@ protected:
 
 
 } // end namespace UE::Geometry
-} // end namespace UE
\ No newline at end of file
+} // end namespace UE
+
+#undef UE_API
diff --git a/Engine/Plugins/Runtime/GeometryProcessing/Source/DynamicMesh/Public/Operations/RepairOrientation.h b/Engine/Plugins/Runtime/GeometryProcessing/Source/DynamicMesh/Public/Operations/RepairOrientation.h
--- a/Engine/Plugins/Runtime/GeometryProcessing/Source/DynamicMesh/Public/Operations/RepairOrientation.h
+++ b/Engine/Plugins/Runtime/GeometryProcessing/Source/DynamicMesh/Public/Operations/RepairOrientation.h
@@ -9,6 +9,8 @@
 #include "DynamicMesh/DynamicMesh3.h"
 #include "DynamicMesh/DynamicMeshAABBTree3.h"
 
+#define UE_API DYNAMICMESH_API
+
 namespace UE
 {
 namespace Geometry
@@ -17,7 +19,7 @@ namespace Geometry
 /**
  * Invert triangles as needed to a consistent ~"outward" orientation
  */
-class DYNAMICMESH_API FMeshRepairOrientation
+class FMeshRepairOrientation
 {
 public:
 	//
@@ -42,10 +44,10 @@ public:
 	}
 	
 	// locally orient connected components
-	void OrientComponents();
+	UE_API void OrientComponents();
 
 	// globally orient -- to be called after OrientComponents if a more globally consistent result is desired
-	void SolveGlobalOrientation(FDynamicMeshAABBTree3* Tree);
+	UE_API void SolveGlobalOrientation(FDynamicMeshAABBTree3* Tree);
 
 protected:
 
@@ -57,10 +59,12 @@ protected:
 	};
 	TArray<Component> Components;
 
-	void ComputeStatistics(FDynamicMeshAABBTree3* Tree);
-	void ComputeComponentStatistics(FDynamicMeshAABBTree3* Tree, Component& c);
+	UE_API void ComputeStatistics(FDynamicMeshAABBTree3* Tree);
+	UE_API void ComputeComponentStatistics(FDynamicMeshAABBTree3* Tree, Component& c);
 };
 
 
 } // end namespace UE::Geometry
 } // end namespace UE
+
+#undef UE_API
diff --git a/Engine/Plugins/Runtime/GeometryProcessing/Source/DynamicMesh/Public/Operations/SelectiveTessellate.h b/Engine/Plugins/Runtime/GeometryProcessing/Source/DynamicMesh/Public/Operations/SelectiveTessellate.h
--- a/Engine/Plugins/Runtime/GeometryProcessing/Source/DynamicMesh/Public/Operations/SelectiveTessellate.h
+++ b/Engine/Plugins/Runtime/GeometryProcessing/Source/DynamicMesh/Public/Operations/SelectiveTessellate.h
@@ -8,6 +8,8 @@
 #include "Containers/ArrayView.h"
 #include "Math/MathFwd.h"
 
+#define UE_API DYNAMICMESH_API
+
 class FProgressCancel;
 
 namespace UE
@@ -48,7 +50,7 @@ class FDynamicMesh3;
  *      - Using the new vertices (along edges and inside the triangle), provide connectivity for a triangle patch by 
  *        stitching those vertices together and creating triangles.
  */
-class DYNAMICMESH_API FTessellationPattern
+class FTessellationPattern
 {
 public:
     
@@ -180,7 +182,7 @@ protected:
      * 
      * @note Vertices should be passed in CW order.
      */
-    void ComputeInnerConcentricTriangle(const FVector3d& U,
+    UE_API void ComputeInnerConcentricTriangle(const FVector3d& U,
                                         const FVector3d& V,
                                         const FVector3d& W,
                                         const int EdgeTessellationLevel, 
@@ -201,7 +203,7 @@ protected:
  * 
  * @note Currently does not support interpolation of the GenericAttributes besides Skin Weights.
  */
-class DYNAMICMESH_API FSelectiveTessellate
+class FSelectiveTessellate
 {
 public:
 
@@ -317,7 +319,7 @@ public:
 	 * Generate tessellated geometry.
 	 * @return true if the algorithm succeeds, false if it failed or was cancelled by the user.
 	 */
-	virtual bool Compute();
+	UE_API virtual bool Compute();
 
 
 
@@ -366,21 +368,21 @@ public:
      * @param InInnerTessLevels Array of InMesh->MaxTriangleID() size, where InInnerTessLevels[TriangleID] is the 
      *                          number of new vertices along each of the edges of the first inner ring.
      */
-    static TUniquePtr<FTessellationPattern> CreateConcentricRingsTessellationPattern(const FDynamicMesh3* InMesh,
+    static UE_API TUniquePtr<FTessellationPattern> CreateConcentricRingsTessellationPattern(const FDynamicMesh3* InMesh,
                                                                                      const TArray<int>& InEdgeTessLevels, 
                                                                                      const TArray<int>& InInnerTessLevels);
     
     /** 
      * Tessellate the whole mesh such that all the edge and inner tessellation levels are equal to InTessellationLevel. 
      */
-    static TUniquePtr<FTessellationPattern> CreateConcentricRingsTessellationPattern(const FDynamicMesh3* InMesh,
+    static UE_API TUniquePtr<FTessellationPattern> CreateConcentricRingsTessellationPattern(const FDynamicMesh3* InMesh,
                                                                                      const int InTessellationLevel);
 
     /** 
      * Tessellate only selected triangles such that the edge and inner tessellation levels of those triangles are set to
      * InTessellationLevel.
      */
-    static TUniquePtr<FTessellationPattern> CreateConcentricRingsTessellationPattern(const FDynamicMesh3* InMesh,
+    static UE_API TUniquePtr<FTessellationPattern> CreateConcentricRingsTessellationPattern(const FDynamicMesh3* InMesh,
                                                                                      const int InTessellationLevel,
                                                                                      const TArray<int>& InTriangleList);
 
@@ -391,7 +393,7 @@ public:
      * 
      * @note not yet implemented
      */
-    static TUniquePtr<FTessellationPattern> CreateConcentricRingsTessellationPattern(const TFunctionRef<int(const int EdgeID)> InEdgeFunc, 
+    static UE_API TUniquePtr<FTessellationPattern> CreateConcentricRingsTessellationPattern(const TFunctionRef<int(const int EdgeID)> InEdgeFunc, 
                                                                                      const TFunctionRef<int(const int TriangleID)> InTriFunc);
 
     
@@ -399,7 +401,7 @@ public:
      * Tessellate only triangles that belong to a given triangle group id. Triangle groups are defined by the 
      * FDynamicMesh3::GetTriangleGroup().
      */
-    static TUniquePtr<FTessellationPattern> CreateConcentricRingsPatternFromTriangleGroup(const FDynamicMesh3* InMesh,
+    static UE_API TUniquePtr<FTessellationPattern> CreateConcentricRingsPatternFromTriangleGroup(const FDynamicMesh3* InMesh,
                                                                                           const int InTessellationLevel,
                                                                                           const int InPolygroupID);
 
@@ -410,7 +412,7 @@ public:
      * 
      * @return nullptr if InMesh has no attributes or if no polygroup layer attribute with the given name or id exists.
      */
-    static TUniquePtr<FTessellationPattern> CreateConcentricRingsPatternFromPolyGroup(const FDynamicMesh3* InMesh,
+    static UE_API TUniquePtr<FTessellationPattern> CreateConcentricRingsPatternFromPolyGroup(const FDynamicMesh3* InMesh,
                                                                                       const int InTessellationLevel,
                                                                                       const FString& InLayerName,
                                                                                       const int InPolygroupID);
@@ -421,7 +423,7 @@ public:
      * 
      * @return nullptr if InMesh has no attributes, no material attribute or if no material attribute with the given id exists.
      */
-	static TUniquePtr<FTessellationPattern> CreateConcentricRingsPatternFromMaterial(const FDynamicMesh3* InMesh,
+	static UE_API TUniquePtr<FTessellationPattern> CreateConcentricRingsPatternFromMaterial(const FDynamicMesh3* InMesh,
 																				 	 const int InTessellationLevel,
 																					 const int MaterialID);
 
@@ -431,7 +433,7 @@ public:
 	*
 	*  @return nullptr if InMesh has no attributes, no material attribute, if no material attribute with the given id exists, or if selection is invalid
 	*/
-	static TUniquePtr<FTessellationPattern> CreateConcentricRingsPatternFromSelectionAndMaterial(const FDynamicMesh3* InMesh,
+	static UE_API TUniquePtr<FTessellationPattern> CreateConcentricRingsPatternFromSelectionAndMaterial(const FDynamicMesh3* InMesh,
 																					  const int InTessellationLevel,
 																					  const int MaterialID,
 																					  const TArray<int>& SelectedTriangles);
@@ -461,7 +463,7 @@ public:
      * @note not yet implemented
      */         
 
-    static TUniquePtr<FTessellationPattern> CreateInnerUnifromTessellationPattern(const FDynamicMesh3* InMesh,
+    static UE_API TUniquePtr<FTessellationPattern> CreateInnerUnifromTessellationPattern(const FDynamicMesh3* InMesh,
                                                                                   const TArray<int>& InEdgeTessLevels, 
                                                                                   const TArray<int>& InInnerTessLevels);
 
@@ -490,7 +492,7 @@ public:
      * @note not yet implemented
      */
     
-    static TUniquePtr<FTessellationPattern> CreateUniformTessellationPattern(const FDynamicMesh3* InMesh,
+    static UE_API TUniquePtr<FTessellationPattern> CreateUniformTessellationPattern(const FDynamicMesh3* InMesh,
                                                                              const int InTessellationLevel,
                                                                              const TArray<int>& InTriangleList);
 
@@ -499,8 +501,10 @@ public:
 protected:
 
 	/** @return true if we need to abort the computation. */
-	virtual bool Cancelled();
+	UE_API virtual bool Cancelled();
 };
 
 } // end namespace UE::Geometry
 } // end namespace UE
+
+#undef UE_API
diff --git a/Engine/Plugins/Runtime/GeometryProcessing/Source/DynamicMesh/Public/Operations/SimpleHoleFiller.h b/Engine/Plugins/Runtime/GeometryProcessing/Source/DynamicMesh/Public/Operations/SimpleHoleFiller.h
--- a/Engine/Plugins/Runtime/GeometryProcessing/Source/DynamicMesh/Public/Operations/SimpleHoleFiller.h
+++ b/Engine/Plugins/Runtime/GeometryProcessing/Source/DynamicMesh/Public/Operations/SimpleHoleFiller.h
@@ -10,6 +10,8 @@
 #include "GeometryTypes.h"
 #include "MeshRegionBoundaryLoops.h"
 
+#define UE_API DYNAMICMESH_API
+
 namespace UE
 {
 namespace Geometry
@@ -21,7 +23,7 @@ class FDynamicMesh3;
  * Fill an EdgeLoop hole with triangles.
  * Supports two fill modes, either a fan connected to a new central vertex, or a triangulation of the boundary polygon
  */
-class DYNAMICMESH_API FSimpleHoleFiller : public IHoleFiller
+class FSimpleHoleFiller : public IHoleFiller
 {
 public:
 	enum class EFillType
@@ -75,7 +77,7 @@ public:
 		return EOperationValidationResult::Ok;
 	}
 
-	bool Fill(int32 GroupID = -1) override;	
+	UE_API bool Fill(int32 GroupID = -1) override;	
 
 	/**
 	 * Updates the normals and UV's of NewTriangles. UV's are taken from VidUVMaps,
@@ -95,14 +97,16 @@ public:
 	 * @returns false if there is an error, usually if VidsToUVsMap did not have an entry
 	 *  for a needed vertex ID.
 	 */
-	bool UpdateAttributes(TArray<FMeshRegionBoundaryLoops::VidOverlayMap<FVector2f>>& VidUVMaps);
+	UE_API bool UpdateAttributes(TArray<FMeshRegionBoundaryLoops::VidOverlayMap<FVector2f>>& VidUVMaps);
 
 
 protected:
-	bool Fill_Fan(int32 NewGroupID);
-	bool Fill_EarClip(int32 NewGroupID);
+	UE_API bool Fill_Fan(int32 NewGroupID);
+	UE_API bool Fill_EarClip(int32 NewGroupID);
 };
 
 
 } // end namespace UE::Geometry
-} // end namespace UE
\ No newline at end of file
+} // end namespace UE
+
+#undef UE_API
diff --git a/Engine/Plugins/Runtime/GeometryProcessing/Source/DynamicMesh/Public/Operations/SmoothBoneWeights.h b/Engine/Plugins/Runtime/GeometryProcessing/Source/DynamicMesh/Public/Operations/SmoothBoneWeights.h
--- a/Engine/Plugins/Runtime/GeometryProcessing/Source/DynamicMesh/Public/Operations/SmoothBoneWeights.h
+++ b/Engine/Plugins/Runtime/GeometryProcessing/Source/DynamicMesh/Public/Operations/SmoothBoneWeights.h
@@ -8,6 +8,8 @@
 #include "BoneIndices.h"
 #include "BoneWeights.h" //  UE::AnimationCore::MaxInlineBoneWeightCount;
 
+#define UE_API DYNAMICMESH_API
+
 // Forward declarations
 class FProgressCancel;
 namespace UE::Geometry 
@@ -83,7 +85,7 @@ public:
  * }
  */
 template <typename BoneIndexType, typename BoneWeightType>
-class DYNAMICMESH_API TSmoothBoneWeights
+class TSmoothBoneWeights
 {
 public:
 
@@ -165,7 +167,7 @@ protected:
  * 		}
  * }
  */
-class DYNAMICMESH_API FSmoothDynamicMeshVertexSkinWeights : public TSmoothBoneWeights<FBoneIndexType, float>
+class FSmoothDynamicMeshVertexSkinWeights : public TSmoothBoneWeights<FBoneIndexType, float>
 {
 public:
 
@@ -182,18 +184,18 @@ public:
 	 * @param InSourceMesh The mesh whose vertices contain bone weights data.
 	 * @param InProfile    The name of the skin weights profile containing the bone weights data.
 	 */
-	FSmoothDynamicMeshVertexSkinWeights(const FDynamicMesh3* InSourceMesh, const FName InProfile);
+	UE_API FSmoothDynamicMeshVertexSkinWeights(const FDynamicMesh3* InSourceMesh, const FName InProfile);
 
 	/**
 	 * @param InSourceMesh The mesh whose vertices contain bone weights data.
 	 * @param InAttribute  The skin weight attribute containing the bone data.
 	 */
-	FSmoothDynamicMeshVertexSkinWeights(const FDynamicMesh3* InSourceMesh, FDynamicMeshVertexSkinWeightsAttribute* InAttribute);
+	UE_API FSmoothDynamicMeshVertexSkinWeights(const FDynamicMesh3* InSourceMesh, FDynamicMeshVertexSkinWeightsAttribute* InAttribute);
 
 	virtual ~FSmoothDynamicMeshVertexSkinWeights() = default;
 	
 	/** @return EOperationValidationResult::Ok if we can apply operation, or error code if we cannot. */
-	virtual EOperationValidationResult Validate() override;
+	UE_API virtual EOperationValidationResult Validate() override;
 	
 	/**
 	 * Compute smoothed weights at a vertex and write the result into FDynamicMeshVertexSkinWeightsAttribute passed 
@@ -204,7 +206,7 @@ public:
 	 *
 	 * @return true if the algorithm succeeds, false if it failed or was canceled by the user.
 	 */
-	bool SmoothWeightsAtVertex(const int32 VertexID, const float VertexFalloff);
+	UE_API bool SmoothWeightsAtVertex(const int32 VertexID, const float VertexFalloff);
 
 	/**
 	 * Compute smoothed weights at vertices and neighbors within distance and write the result into 
@@ -218,7 +220,7 @@ public:
 	 *
 	 * @return true if the algorithm succeeds, false if it failed or was canceled by the user.
 	 */
-	bool SmoothWeightsAtVerticesWithinDistance(const TArray<int32>& Vertices, const float Strength, const double FloodFillUpToDistance, const int32 NumIterations);
+	UE_API bool SmoothWeightsAtVerticesWithinDistance(const TArray<int32>& Vertices, const float Strength, const double FloodFillUpToDistance, const int32 NumIterations);
 
 protected:
 	TUniquePtr<TBoneWeightsDataSource<FBoneIndexType, float>> SkinWeightsAttributeDataSource = nullptr;
@@ -228,3 +230,5 @@ protected:
 
 } // end namespace UE::Geometry
 } // end namespace UE
+
+#undef UE_API
diff --git a/Engine/Plugins/Runtime/GeometryProcessing/Source/DynamicMesh/Public/Operations/SmoothDynamicMeshAttributes.h b/Engine/Plugins/Runtime/GeometryProcessing/Source/DynamicMesh/Public/Operations/SmoothDynamicMeshAttributes.h
--- a/Engine/Plugins/Runtime/GeometryProcessing/Source/DynamicMesh/Public/Operations/SmoothDynamicMeshAttributes.h
+++ b/Engine/Plugins/Runtime/GeometryProcessing/Source/DynamicMesh/Public/Operations/SmoothDynamicMeshAttributes.h
@@ -6,6 +6,8 @@
 #include "Containers/Map.h"
 #include "Containers/Array.h"
 
+#define UE_API DYNAMICMESH_API
+
 // Forward declarations
 class FProgressCancel;
 namespace UE::Geometry 
@@ -63,7 +65,7 @@ namespace Geometry
  *  // Run the smoothing
  *  BlurOp.SmoothOverlay(Colors);
  */
-class DYNAMICMESH_API FSmoothDynamicMeshAttributes
+class FSmoothDynamicMeshAttributes
 {
 public:
 
@@ -131,7 +133,7 @@ protected:
 
 public:
 
-	FSmoothDynamicMeshAttributes(const FDynamicMesh3& InMesh);
+	UE_API FSmoothDynamicMeshAttributes(const FDynamicMesh3& InMesh);
 
 	virtual ~FSmoothDynamicMeshAttributes() = default;
 
@@ -147,7 +149,7 @@ public:
 protected:
 	
     /** @return if true, abort the computation. */
-	virtual bool Cancelled();
+	UE_API virtual bool Cancelled();
 
     /** @return the index of the vertex into ExpandedSelection array if Selection is not empty, otherwise returns VertexID. */
     FORCEINLINE int32 MeshVIDToSelectionIdx(int32 VertexID) const
@@ -157,4 +159,6 @@ protected:
 };
 
 }
-}
\ No newline at end of file
+}
+
+#undef UE_API
diff --git a/Engine/Plugins/Runtime/GeometryProcessing/Source/DynamicMesh/Public/Operations/SmoothHoleFiller.h b/Engine/Plugins/Runtime/GeometryProcessing/Source/DynamicMesh/Public/Operations/SmoothHoleFiller.h
--- a/Engine/Plugins/Runtime/GeometryProcessing/Source/DynamicMesh/Public/Operations/SmoothHoleFiller.h
+++ b/Engine/Plugins/Runtime/GeometryProcessing/Source/DynamicMesh/Public/Operations/SmoothHoleFiller.h
@@ -10,6 +10,8 @@
 #include "DynamicMesh/DynamicMeshAABBTree3.h"
 #include "ProjectionTargets.h"
 
+#define UE_API DYNAMICMESH_API
+
 namespace UE
 {
 namespace Geometry
@@ -70,14 +72,14 @@ struct FSmoothFillOptions
  * This fills a hole in a mesh by doing a trivial fill, then doing a remesh, then a laplacian smooth, then a second 
  * remesh.
  */
-class DYNAMICMESH_API FSmoothHoleFiller : public IHoleFiller
+class FSmoothHoleFiller : public IHoleFiller
 {
 public:
-	FSmoothHoleFiller(FDynamicMesh3& Mesh, const FEdgeLoop& FillLoop);
+	UE_API FSmoothHoleFiller(FDynamicMesh3& Mesh, const FEdgeLoop& FillLoop);
 
 	FSmoothFillOptions FillOptions;
 
-	bool Fill(int32 GroupID) override;
+	UE_API bool Fill(int32 GroupID) override;
 
 protected:
 
@@ -90,16 +92,18 @@ protected:
 	// initialized to the average edge length of FillLoop / FillDensityScalar
 	double RemeshingTargetEdgeLength;
 
-	void SmoothAndRemeshPreserveRegion(FMeshFaceSelection& tris, bool bFinal);
+	UE_API void SmoothAndRemeshPreserveRegion(FMeshFaceSelection& tris, bool bFinal);
 
-	void SmoothAndRemesh(FMeshFaceSelection& tris);
+	UE_API void SmoothAndRemesh(FMeshFaceSelection& tris);
 
-	void ConstrainSubregionSeams(FSubRegionRemesher& Remesher, bool bConstrainROIBoundary);
+	UE_API void ConstrainSubregionSeams(FSubRegionRemesher& Remesher, bool bConstrainROIBoundary);
 
-	void DefaultConfigureRemesher(FSubRegionRemesher& Remesher, bool bConstrainROIBoundary);
+	UE_API void DefaultConfigureRemesher(FSubRegionRemesher& Remesher, bool bConstrainROIBoundary);
 
 };
 
 
 } // end namespace UE::Geometry
-} // end namespace UE
\ No newline at end of file
+} // end namespace UE
+
+#undef UE_API
diff --git a/Engine/Plugins/Runtime/GeometryProcessing/Source/DynamicMesh/Public/Operations/TransferBoneWeights.h b/Engine/Plugins/Runtime/GeometryProcessing/Source/DynamicMesh/Public/Operations/TransferBoneWeights.h
--- a/Engine/Plugins/Runtime/GeometryProcessing/Source/DynamicMesh/Public/Operations/TransferBoneWeights.h
+++ b/Engine/Plugins/Runtime/GeometryProcessing/Source/DynamicMesh/Public/Operations/TransferBoneWeights.h
@@ -7,6 +7,8 @@
 #include "TransformTypes.h"
 #include "BoneWeights.h" // UE::AnimationCore::MaxInlineBoneWeightCount
 
+#define UE_API DYNAMICMESH_API
+
 // Forward declarations
 class FProgressCancel;
 
@@ -85,7 +87,7 @@ namespace Geometry
 
 
 
-class DYNAMICMESH_API FTransferBoneWeights
+class FTransferBoneWeights
 {
 public:
 
@@ -216,16 +218,16 @@ public:
 	 * @note Assumes that the InSourceMesh has bone attributes, use bIgnoreBoneAttributes flag to ignore the bone 
 	 * 		 attributes and skip re-indexing.
 	 */
-	FTransferBoneWeights(const FDynamicMesh3* InSourceMesh, 
+	UE_API FTransferBoneWeights(const FDynamicMesh3* InSourceMesh, 
 					     const FName& InSourceProfileName,
 					     const FDynamicMeshAABBTree3* SourceBVH = nullptr); 
 	
-	virtual ~FTransferBoneWeights();
+	UE_API virtual ~FTransferBoneWeights();
 
 	/**
 	 * @return EOperationValidationResult::Ok if we can apply operation, or error code if we cannot.
 	 */
-	virtual EOperationValidationResult Validate();
+	UE_API virtual EOperationValidationResult Validate();
 
 	/**
      * Transfer the bone weights from the source mesh to the given target mesh and store the result in the skin weight  
@@ -240,7 +242,7 @@ public:
 	 * @note Assumes that the InOutTargetMesh has bone attributes, use bIgnoreBoneAttributes flag to ignore the bone 
 	 * 		 attributes and skip re-indexing.
 	 */
-	virtual bool TransferWeightsToMesh(FDynamicMesh3& InOutTargetMesh, const FName& InTargetProfileName);
+	UE_API virtual bool TransferWeightsToMesh(FDynamicMesh3& InOutTargetMesh, const FName& InTargetProfileName);
 
 
 	/**
@@ -254,7 +256,7 @@ public:
 	 * 
 	 * @return true if the algorithm succeeds, false if it failed to find the matching point or was canceled by the user.
 	 */
-	virtual bool TransferWeightsToPoint(UE::AnimationCore::FBoneWeights& OutWeights, 
+	UE_API virtual bool TransferWeightsToPoint(UE::AnimationCore::FBoneWeights& OutWeights, 
 										const FVector3d& InPoint, 
 										const TMap<FName, uint16>* TargetBoneToIndex = nullptr,
 										const FVector3f& InNormal = FVector3f::Zero());
@@ -283,7 +285,7 @@ public:
 protected:
 	
     /** @return if true, abort the computation. */
-	virtual bool Cancelled();
+	UE_API virtual bool Cancelled();
 	
 	/** 
 	 * Find the closest point on the surface of the source mesh and return the ID of the triangle containing it and its 
@@ -291,9 +293,11 @@ protected:
 	 * 
 	 * @return true if point is found, false otherwise
 	 */
-	bool FindClosestPointOnSourceSurface(const FVector3d& InPoint, const FTransformSRT3d& InToWorld, int32& OutTriID, FVector3d& OutBary);
+	UE_API bool FindClosestPointOnSourceSurface(const FVector3d& InPoint, const FTransformSRT3d& InToWorld, int32& OutTriID, FVector3d& OutBary);
 
 };
 
 } // end namespace UE::Geometry
 } // end namespace UE
+
+#undef UE_API
diff --git a/Engine/Plugins/Runtime/GeometryProcessing/Source/DynamicMesh/Public/Operations/TransferDynamicMeshAttributes.h b/Engine/Plugins/Runtime/GeometryProcessing/Source/DynamicMesh/Public/Operations/TransferDynamicMeshAttributes.h
--- a/Engine/Plugins/Runtime/GeometryProcessing/Source/DynamicMesh/Public/Operations/TransferDynamicMeshAttributes.h
+++ b/Engine/Plugins/Runtime/GeometryProcessing/Source/DynamicMesh/Public/Operations/TransferDynamicMeshAttributes.h
@@ -6,6 +6,8 @@
 #include "UObject/NameTypes.h"
 #include "TransformTypes.h"
 
+#define UE_API DYNAMICMESH_API
+
 // Forward declarations
 class FProgressCancel;
 
@@ -27,7 +29,7 @@ namespace Geometry
  * Transfer vertex colors from one mesh (source) to another (target).
  */
 
-class DYNAMICMESH_API FTransferVertexColorAttribute
+class FTransferVertexColorAttribute
 {
 public:
 
@@ -141,14 +143,14 @@ public:
 	 * 
 	 * @note Assumes that the InSourceMesh has a primary colors attribute.
 	 */
-	FTransferVertexColorAttribute(const FDynamicMesh3* InSourceMesh, const FDynamicMeshAABBTree3* SourceBVH = nullptr); 
+	UE_API FTransferVertexColorAttribute(const FDynamicMesh3* InSourceMesh, const FDynamicMeshAABBTree3* SourceBVH = nullptr); 
 	
-	virtual ~FTransferVertexColorAttribute();
+	UE_API virtual ~FTransferVertexColorAttribute();
 
 	/**
 	 * @return EOperationValidationResult::Ok if we can apply operation, or error code if we cannot.
 	 */
-	virtual EOperationValidationResult Validate();
+	UE_API virtual EOperationValidationResult Validate();
 
 	/**
      * Transfer the colors from the source mesh to the given target mesh and store the result in the primary colors attribute.
@@ -157,7 +159,7 @@ public:
      * 
      * @return true if the algorithm succeeds, false if it failed or was canceled by the user.
 	 */
-	virtual bool TransferColorsToMesh(FDynamicMesh3& InOutTargetMesh);
+	UE_API virtual bool TransferColorsToMesh(FDynamicMesh3& InOutTargetMesh);
 
 
 	/**
@@ -169,12 +171,12 @@ public:
 	 * 
 	 * @return true if the algorithm succeeds, false if it failed to find the matching point or was canceled by the user.
 	 */
-	bool TransferColorToPoint(FVector4f& OutColor, const FVector3d& InPoint, const FVector3f& InNormal = FVector3f::Zero()) const;
+	UE_API bool TransferColorToPoint(FVector4f& OutColor, const FVector3d& InPoint, const FVector3f& InNormal = FVector3f::Zero()) const;
 	
 protected:
 	
     /** @return if true, abort the computation. */
-	virtual bool Cancelled();
+	UE_API virtual bool Cancelled();
 	
 	/** 
 	 * Find the closest point on the surface of the source mesh and return the ID of the triangle containing it and its 
@@ -182,14 +184,16 @@ protected:
 	 * 
 	 * @return true if point is found, false otherwise
 	 */
-	bool FindClosestPointOnSourceSurface(const FVector3d& InPoint, const FTransformSRT3d& InToWorld, int32& OutTriID, FVector3d& OutBary) const;
+	UE_API bool FindClosestPointOnSourceSurface(const FVector3d& InPoint, const FTransformSRT3d& InToWorld, int32& OutTriID, FVector3d& OutBary) const;
 
 	/**
      * Transfer the colors from the source mesh to the given target mesh using the closest point algorithm.
      * @return the number of matching vertices / elements
 	 */
-	int32 TransferUsingClosestPoint(FDynamicMesh3& InOutTargetMesh, const TUniquePtr<FMeshNormals>& InTargetMeshNormals);
+	UE_API int32 TransferUsingClosestPoint(FDynamicMesh3& InOutTargetMesh, const TUniquePtr<FMeshNormals>& InTargetMeshNormals);
 };
 
 } // end namespace UE::Geometry
 } // end namespace UE
+
+#undef UE_API
diff --git a/Engine/Plugins/Runtime/GeometryProcessing/Source/DynamicMesh/Public/Operations/UniformTessellate.h b/Engine/Plugins/Runtime/GeometryProcessing/Source/DynamicMesh/Public/Operations/UniformTessellate.h
--- a/Engine/Plugins/Runtime/GeometryProcessing/Source/DynamicMesh/Public/Operations/UniformTessellate.h
+++ b/Engine/Plugins/Runtime/GeometryProcessing/Source/DynamicMesh/Public/Operations/UniformTessellate.h
@@ -5,6 +5,8 @@
 #include "GeometryTypes.h"
 #include "IndexTypes.h"
 
+#define UE_API DYNAMICMESH_API
+
 class FProgressCancel;
 
 namespace UE
@@ -35,7 +37,7 @@ class FDynamicMesh3;
  *   o-----------------o               o-----x-----x-----o 
  */
 
-class DYNAMICMESH_API FUniformTessellate
+class FUniformTessellate
 {
 public:
 	//
@@ -120,10 +122,10 @@ public:
 	}
 
 	/** @return The number of vertices after the tessellation. */
-	static int32 ExpectedNumVertices(const FDynamicMesh3& Mesh, const int32 TessellationNum);
+	static UE_API int32 ExpectedNumVertices(const FDynamicMesh3& Mesh, const int32 TessellationNum);
 	
 	/** @return The number of triangles after the tessellation. */
-	static int32 ExpectedNumTriangles(const FDynamicMesh3& Mesh, const int32 TessellationNum);
+	static UE_API int32 ExpectedNumTriangles(const FDynamicMesh3& Mesh, const int32 TessellationNum);
 
 	/**
 	 * @return EOperationValidationResult::Ok if we can apply operation, or error code if we cannot.
@@ -146,13 +148,15 @@ public:
 	 * 
 	 * @return true if the algorithm succeeds, false if it failed or was cancelled by the user.
 	 */
-	virtual bool Compute();
+	UE_API virtual bool Compute();
 
 protected:
 
 	/** If this returns true, abort computation. */
-	virtual bool Cancelled();
+	UE_API virtual bool Cancelled();
 };
 
 } // end namespace UE::Geometry
 } // end namespace UE
+
+#undef UE_API
diff --git a/Engine/Plugins/Runtime/GeometryProcessing/Source/DynamicMesh/Public/Operations/WeldEdgeSequence.h b/Engine/Plugins/Runtime/GeometryProcessing/Source/DynamicMesh/Public/Operations/WeldEdgeSequence.h
--- a/Engine/Plugins/Runtime/GeometryProcessing/Source/DynamicMesh/Public/Operations/WeldEdgeSequence.h
+++ b/Engine/Plugins/Runtime/GeometryProcessing/Source/DynamicMesh/Public/Operations/WeldEdgeSequence.h
@@ -3,6 +3,8 @@
 
 #include "EdgeSpan.h"
 
+#define UE_API DYNAMICMESH_API
+
 namespace UE
 {
 namespace Geometry
@@ -20,7 +22,7 @@ class FDynamicMesh3;
  *  where the group edges are connected by an edge at the end points.
  */
 
-class DYNAMICMESH_API FWeldEdgeSequence
+class FWeldEdgeSequence
 {
 public:
 	enum class EWeldResult
@@ -74,14 +76,14 @@ public:
 	*
 	* @return EWeldResult::OK on success
 	*/
-	EWeldResult Weld();
+	UE_API EWeldResult Weld();
 
 	/**
 	 * Helper that splits the edges in the shorter span until the spans have the same number
 	 *  of edges. Weld() will automatically do this, but this is public in case the user wants
 	 *  to equalize spans that get concatenated together before all being welded at once.
 	 */
-	static EWeldResult SplitEdgesToEqualizeSpanLengths(FDynamicMesh3& Mesh, FEdgeSpan& Span1, FEdgeSpan& Span2);
+	static UE_API EWeldResult SplitEdgesToEqualizeSpanLengths(FDynamicMesh3& Mesh, FEdgeSpan& Span1, FEdgeSpan& Span2);
 protected:
 	/**
 	* Verifies validity of input edges by ensuring they are
@@ -89,7 +91,7 @@ protected:
 	*
 	* @return EWeldResult::OK on success
 	*/
-	EWeldResult CheckInput();
+	UE_API EWeldResult CheckInput();
 
 	/**
 	* Splits largest edges of the span with fewest vertices so that
@@ -97,19 +99,21 @@ protected:
 	* 
 	* @return EWeldResult::OK on success
 	*/
-	EWeldResult SplitSmallerSpan();
+	UE_API EWeldResult SplitSmallerSpan();
 
 
 	UE_DEPRECATED(5.5, "Side triangles are handled appropriately in WeldEdgeSequence")
-	EWeldResult CheckForAndCollapseSideTriangles();
+	UE_API EWeldResult CheckForAndCollapseSideTriangles();
 
 	/**
 	* Welds edge sequence together
 	*
 	* @return EWeldResult::OK on success
 	*/
-	EWeldResult WeldEdgeSequence();
+	UE_API EWeldResult WeldEdgeSequence();
 };
 
 } // end namespace UE::Geometry
 } // end namespace UE
+
+#undef UE_API
diff --git a/Engine/Plugins/Runtime/GeometryProcessing/Source/DynamicMesh/Public/Parameterization/DynamicMeshUVEditor.h b/Engine/Plugins/Runtime/GeometryProcessing/Source/DynamicMesh/Public/Parameterization/DynamicMeshUVEditor.h
--- a/Engine/Plugins/Runtime/GeometryProcessing/Source/DynamicMesh/Public/Parameterization/DynamicMeshUVEditor.h
+++ b/Engine/Plugins/Runtime/GeometryProcessing/Source/DynamicMesh/Public/Parameterization/DynamicMeshUVEditor.h
@@ -6,13 +6,15 @@
 #include "DynamicMesh/DynamicMeshAttributeSet.h"
 #include "DynamicMesh/DynamicMeshAABBTree3.h"
 
+#define UE_API DYNAMICMESH_API
+
 namespace UE
 {
 namespace Geometry
 {
 
 
-struct DYNAMICMESH_API FUVEditResult
+struct FUVEditResult
 {
 	TArray<int32> NewUVElements;
 
@@ -22,7 +24,7 @@ struct DYNAMICMESH_API FUVEditResult
 /**
  * FDynamicMeshUVEditor implements various UV overlay editing operations.
  */
-class DYNAMICMESH_API FDynamicMeshUVEditor
+class FDynamicMeshUVEditor
 {
 private:
 	/** The mesh we will be editing */
@@ -38,12 +40,12 @@ public:
 	 * @param UVLayerIndex index of target UV layer
 	 * @param bCreateIfMissing if true, target UV layers up to UVLayerIndex will be created if it is not there. Otherwise UVOverlay will be nullptr and class is incomplete.
 	 */
-	explicit FDynamicMeshUVEditor(FDynamicMesh3* MeshIn, int32 UVLayerIndex, bool bCreateIfMissing);
+	UE_API explicit FDynamicMeshUVEditor(FDynamicMesh3* MeshIn, int32 UVLayerIndex, bool bCreateIfMissing);
 
 	/**
 	 * Construct UV Editor for a UV Overlay of the given Mesh.
 	 */
-	explicit FDynamicMeshUVEditor(FDynamicMesh3* MeshIn, FDynamicMeshUVOverlay* UVOverlayIn);
+	UE_API explicit FDynamicMeshUVEditor(FDynamicMesh3* MeshIn, FDynamicMeshUVOverlay* UVOverlayIn);
 
 
 	FDynamicMesh3* GetMesh() { return Mesh; }
@@ -55,37 +57,37 @@ public:
 	/**
 	 * Create specified UVLayer if it does not exist
 	 */
-	void CreateUVLayer(int32 UVLayerIndex);
+	UE_API void CreateUVLayer(int32 UVLayerIndex);
 
 	/**
 	 * Append new UV layer to the end of the array, returning the newly added index, or -1 if at max layers already
 	 */
-	int32 AddUVLayer();
+	UE_API int32 AddUVLayer();
 
 	/**
 	 * Switch to new layer
 	 */
-	void SwitchActiveLayer(int32 UVLayerIndex);
+	UE_API void SwitchActiveLayer(int32 UVLayerIndex);
 
 	/**
 	 * Remove active layer, setting active layer to highest, preceeding layer. Will not remove the final layer from the mesh. Returns the resulting active layer index.
 	 */
-	int32 RemoveUVLayer();
+	UE_API int32 RemoveUVLayer();
 
 	/**
 	 * Clear UVs for all triangles on active layer
 	 */
-	void ResetUVs();
+	UE_API void ResetUVs();
 
 	/**
 	* Clear UVs for given triangles on active layer
 	*/
-	void ResetUVs(const TArray<int32>& Triangles);
+	UE_API void ResetUVs(const TArray<int32>& Triangles);
 
 	/**
 	 * Copy UVs from another overlay
 	 */
-	bool CopyUVLayer(FDynamicMeshUVOverlay* FromUVOverlay);
+	UE_API bool CopyUVLayer(FDynamicMeshUVOverlay* FromUVOverlay);
 
 
 	/**
@@ -93,36 +95,36 @@ public:
 	 * @param VertexToUVOut returned new overlay UV ID for each mesh Vertex ID (array size is mesh MaxVertexID)
 	 * @param bIsIdentityMapOut returned as true if UV ID == Vertex ID for every vertex  (even if there are gaps)
 	 */
-	void SetToPerVertexUVs(TArray<int32>& VertexToUVOut, bool& bIsIdentityMapOut, FUVEditResult* Result = nullptr);
+	UE_API void SetToPerVertexUVs(TArray<int32>& VertexToUVOut, bool& bIsIdentityMapOut, FUVEditResult* Result = nullptr);
 
 
 	/**
 	 * Create new UV island for each Triangle, by planar projection onto plane of Triangle. No transforms/etc are applied.
 	 */
-	void SetPerTriangleUVs(const TArray<int32>& Triangles, double ScaleFactor = 1.0, FUVEditResult* Result = nullptr);
+	UE_API void SetPerTriangleUVs(const TArray<int32>& Triangles, double ScaleFactor = 1.0, FUVEditResult* Result = nullptr);
 
 	/**
 	 * Create new UV island for given Triangles, and set UVs by planar projection to ProjectionFrame. No transforms/etc are applied.
 	 */
-	void SetPerTriangleUVs(double ScaleFactor = 1.0, FUVEditResult* Result = nullptr);
+	UE_API void SetPerTriangleUVs(double ScaleFactor = 1.0, FUVEditResult* Result = nullptr);
 
 	/**
 	 * Apply arbitrary transform to UV elements
 	 */
-	void TransformUVElements(const TArray<int32>& ElementIDs, TFunctionRef<FVector2f(const FVector2f&)> TransformFunc);
+	UE_API void TransformUVElements(const TArray<int32>& ElementIDs, TFunctionRef<FVector2f(const FVector2f&)> TransformFunc);
 
 
 	/**
 	 * Create new UV island for given Triangles, and set UVs by planar projection to ProjectionFrame. No transforms/etc are applied.
 	 */
-	void SetTriangleUVsFromProjection(const TArray<int32>& Triangles, const FFrame3d& ProjectionFrame, FUVEditResult* Result = nullptr);
+	UE_API void SetTriangleUVsFromProjection(const TArray<int32>& Triangles, const FFrame3d& ProjectionFrame, FUVEditResult* Result = nullptr);
 
 	/**
 	 * Create new UV island for given Triangles, and set UVs by planar projection to ProjectionFrame. 
 	 * PointTransform is applied to points before projectiong onto ProjectionFrame X/Y axes
 	 * Projected U/V coordinates are divided by Dimensions.X/Y
 	 */
-	void SetTriangleUVsFromPlanarProjection(const TArray<int32>& Triangles, TFunctionRef<FVector3d(const FVector3d&)> PointTransform, const FFrame3d& ProjectionFrame, const FVector2d& Dimensions, FUVEditResult* Result = nullptr);
+	UE_API void SetTriangleUVsFromPlanarProjection(const TArray<int32>& Triangles, TFunctionRef<FVector3d(const FVector3d&)> PointTransform, const FFrame3d& ProjectionFrame, const FVector2d& Dimensions, FUVEditResult* Result = nullptr);
 
 
 	/**
@@ -147,7 +149,7 @@ public:
 	 * ExpMap center-point is calculated by finding maximum (Dijkstra-approximated) geodesic distance from border of island.
 	 * @warning computes a single ExpMap, so input triangle set must be connected, however this is not verified internally
 	 */
-	bool SetTriangleUVsFromExpMap(const TArray<int32>& Triangles, const FExpMapOptions& Options = FExpMapOptions(), FUVEditResult* Result = nullptr);
+	UE_API bool SetTriangleUVsFromExpMap(const TArray<int32>& Triangles, const FExpMapOptions& Options = FExpMapOptions(), FUVEditResult* Result = nullptr);
 
 
 	/**
@@ -155,7 +157,7 @@ public:
 	 * ExpMap center-point is calculated by finding maximum (Dijkstra-approximated) geodesic distance from border of island.
 	 * @warning computes a single ExpMap, so input triangle set must be connected, however this is not verified internally
 	 */
-	bool SetTriangleUVsFromExpMap(
+	UE_API bool SetTriangleUVsFromExpMap(
 		const TArray<int32>& Triangles, 
 		TFunctionRef<FVector3d(const FVector3d&)> PointTransform,
 		const FFrame3d& ProjectionFrame,
@@ -170,7 +172,7 @@ public:
 	 * Create new UV island for given Triangles, and set UVs for that island using Discrete Natural Conformal Map (equivalent to Least-Squares Conformal Map)
 	 * @warning computes a single parameterization, so input triangle set must be connected, however this is not verified internally
 	 */
-	bool SetTriangleUVsFromFreeBoundaryConformal(const TArray<int32>& Triangles, FUVEditResult* Result = nullptr);
+	UE_API bool SetTriangleUVsFromFreeBoundaryConformal(const TArray<int32>& Triangles, FUVEditResult* Result = nullptr);
 
 	/**
 	 * Create new UV island for given Triangles, and set UVs for that island using Discrete Natural Conformal Map (equivalent to Least-Squares Conformal Map)
@@ -178,7 +180,7 @@ public:
 	 * @param bUseExistingUVTopology if true, re-solve for existing UV set, rather than constructing per-vertex UVs from triangle set. Allows for solving w/ partial seams, interior cuts, etc. 
 	 * @warning computes a single parameterization, so input triangle set must be connected, however this is not verified internally
 	 */
-	bool SetTriangleUVsFromFreeBoundaryConformal(const TArray<int32>& Triangles, bool bUseExistingUVTopology, FUVEditResult* Result = nullptr);
+	UE_API bool SetTriangleUVsFromFreeBoundaryConformal(const TArray<int32>& Triangles, bool bUseExistingUVTopology, FUVEditResult* Result = nullptr);
 
 	/**
 	 * Create new UV island for given Triangles, and set UVs for that island using Discrete Natural Conformal Map (equivalent to Least-Squares Conformal Map).
@@ -188,7 +190,7 @@ public:
 	 * @param PinnedElementIDs Elements for which to use unmodified UV values.
 	 * @warning computes a single parameterization, so input triangle set must be connected, however this is not verified internally
 	 */
-	bool SetTriangleUVsFromFreeBoundaryConformal(const TArray<int32>& Triangles, const TSet<int32>& PinnedElementIDs, FUVEditResult* Result = nullptr);
+	UE_API bool SetTriangleUVsFromFreeBoundaryConformal(const TArray<int32>& Triangles, const TSet<int32>& PinnedElementIDs, FUVEditResult* Result = nullptr);
 
 	/**
 	 * Create new UV island for given Triangles, and set UVs for that island using Spectral Conformal Map.
@@ -197,13 +199,13 @@ public:
 	 * @param bPreserveIrregularity if true, reduces distortion for meshes with triangles of vastly different sizes.
 	 * @warning computes a single parameterization, so input triangle set must be connected, however this is not verified internally
 	 */
-	bool SetTriangleUVsFromFreeBoundarySpectralConformal(const TArray<int32>& Triangles, bool bUseExistingUVTopology, bool bPreserveIrregularity, FUVEditResult* Result = nullptr);
+	UE_API bool SetTriangleUVsFromFreeBoundarySpectralConformal(const TArray<int32>& Triangles, bool bUseExistingUVTopology, bool bPreserveIrregularity, FUVEditResult* Result = nullptr);
 
 	/**
 	 * Initializes any uninitialized triangles in the set to per-vertex (0,0) UVs. Doesn't clear or modify any existing 
 	 *  UVs, and only operates on the given set.
 	 */
-	void MakeSureUVsAreSet(const TSet<int32>& Triangles, FUVEditResult* Result = nullptr, TSet<int32>* ChangedTrianglesOut = nullptr);
+	UE_API void MakeSureUVsAreSet(const TSet<int32>& Triangles, FUVEditResult* Result = nullptr, TSet<int32>* ChangedTrianglesOut = nullptr);
 
 	/** 
 	 * Merge existing UV topology with a set of edges, removing seams at edges if they exist within the UV topology.
@@ -211,7 +213,7 @@ public:
 	 * @param EidsToRemoveAsSeams list of edges to remove seams from
 	 * @return true on success
 	 */
-	bool RemoveSeamsAtEdges(const TSet<int32>& EidsToRemoveAsSeams);
+	UE_API bool RemoveSeamsAtEdges(const TSet<int32>& EidsToRemoveAsSeams);
 
 	/**
 	 * Cut existing UV topology with a set of edges. This allows for creating partial seams/darts, interior cuts, etc.
@@ -224,7 +226,7 @@ public:
 	 * @param Result if non-null, list of new UV elements created along the path will be stored here (not ordered)
 	 * @return true on success
 	 */
-	bool CreateSeamsAtEdges(const TSet<int32>& EidsToMakeIntoSeams, FUVEditResult* Result = nullptr);
+	UE_API bool CreateSeamsAtEdges(const TSet<int32>& EidsToMakeIntoSeams, FUVEditResult* Result = nullptr);
 
 	/**
 	 * Takes the currently selected triangles and makes a separate UV island out of them, i.e. any interior seams
@@ -239,7 +241,7 @@ public:
 	 *   superset, so some might not have actually changed. Not cleared before use.
 	 * @return true on success
 	 */
-	bool MakeIsland(const TSet<int32>& TidsToMakeIntoIsland, FUVEditResult* Result = nullptr, TSet<int32>* ChangedTrianglesOut = nullptr);
+	UE_API bool MakeIsland(const TSet<int32>& TidsToMakeIntoIsland, FUVEditResult* Result = nullptr, TSet<int32>* ChangedTrianglesOut = nullptr);
 
 	/**
 	 * Set UVs by box projection. Triangles will be grouped to "best" box face
@@ -247,7 +249,7 @@ public:
 	 * Projected U/V coordinates are divided by Dimensions.X/Y
 	 * @param MinIslandTriCount Any UV island with fewer triangles than this count will be merged into a neighbouring island
 	 */
-	void SetTriangleUVsFromBoxProjection(const TArray<int32>& Triangles, TFunctionRef<FVector3d(const FVector3d&)> PointTransform, const FFrame3d& BoxFrame, const FVector3d& BoxDimensions, int32 MinIslandTriCount = 2, FUVEditResult* Result = nullptr);
+	UE_API void SetTriangleUVsFromBoxProjection(const TArray<int32>& Triangles, TFunctionRef<FVector3d(const FVector3d&)> PointTransform, const FFrame3d& BoxFrame, const FVector3d& BoxDimensions, int32 MinIslandTriCount = 2, FUVEditResult* Result = nullptr);
 
 
 	/**
@@ -255,7 +257,7 @@ public:
 	 * PointTransform is applied to points before projectiong onto ProjectionFrame X/Y axes
 	 * Projected U/V coordinates are divided by Dimensions.X/Y
 	 */
-	void SetTriangleUVsFromCylinderProjection(const TArray<int32>& Triangles, TFunctionRef<FVector3d(const FVector3d&)> PointTransform, const FFrame3d& CylFrame, const FVector3d& CylDimensions, float CylinderSplitAngle, FUVEditResult* Result = nullptr);
+	UE_API void SetTriangleUVsFromCylinderProjection(const TArray<int32>& Triangles, TFunctionRef<FVector3d(const FVector3d&)> PointTransform, const FFrame3d& CylFrame, const FVector3d& CylDimensions, float CylinderSplitAngle, FUVEditResult* Result = nullptr);
 
 
 	struct FTransferFromMeshViaProjectionSettings
@@ -282,7 +284,7 @@ public:
 	 * @param Settings Projection settings
 	 * @param Result If non-null, contains new UV elements IDs created on target mesh
 	 */
-	void TransferTriangleUVsFromMeshViaDirectionProjection(TArrayView<const int32> Triangles, TFunctionRef<FVector3d(const FVector3d&)> TransformTargetToSourceSpacePosition, 
+	UE_API void TransferTriangleUVsFromMeshViaDirectionProjection(TArrayView<const int32> Triangles, TFunctionRef<FVector3d(const FVector3d&)> TransformTargetToSourceSpacePosition, 
 		const FVector3d& ProjectionDirection, double ProjectionOffset, TFunctionRef<FVector3d(const FVector3d&)> TransformProjectionToSourceSpaceVector,
 		const FDynamicMeshAABBTree3& SourceMeshSpatial, int32 SourceMeshUVChannel,
 		const FTransferFromMeshViaProjectionSettings& Settings, FUVEditResult* Result = nullptr);
@@ -293,7 +295,7 @@ public:
 	* to the 3D area. Scaling is around the center of the UV bounding box.
 	* @param bRecenterAtOrigin if true, UVs are translated after scaling such that the UV bounding box center is at (0,0)
 	*/
-	bool ScaleUVAreaTo3DArea(const TArray<int32>& Triangles, bool bRecenterAtOrigin);
+	UE_API bool ScaleUVAreaTo3DArea(const TArray<int32>& Triangles, bool bRecenterAtOrigin);
 
 
 	/**
@@ -302,7 +304,7 @@ public:
 	* @param bPreserveAspectRatio if true, when scaling UVs, maintain original aspect ratio instead of fitting to box bounds
 	* @param bRecenterAtBoundingBox if true, UVs are translated such that their centroid is at the bounding box center
 	*/
-	bool ScaleUVAreaToBoundingBox(const TArray<int32>& Triangles, const FAxisAlignedBox2f& BoundingBox, bool bPreserveAspectRatio, bool bRecenterAtBoundingBox);
+	UE_API bool ScaleUVAreaToBoundingBox(const TArray<int32>& Triangles, const FAxisAlignedBox2f& BoundingBox, bool bPreserveAspectRatio, bool bRecenterAtBoundingBox);
 
 
 	/**
@@ -310,20 +312,20 @@ public:
 	* box is aligned with the X axis. 
 	* The UV-space box is found by computing the 2D Convex Hull and then finding the hull edge that results in the minimal box.
 	*/
-	bool AutoOrientUVArea(const TArray<int32>& Triangles);
+	UE_API bool AutoOrientUVArea(const TArray<int32>& Triangles);
 
 
 	/**
 	 * Pack UVs into unit rectangle
 	 */
-	bool QuickPack(int32 TargetTextureResolution = 512, float GutterSize = 1.0f);
+	UE_API bool QuickPack(int32 TargetTextureResolution = 512, float GutterSize = 1.0f);
 
 	/**
     * Pack specific UVs triangles into the specific unit rectangle at the specified index
 	* @param UDIMCoordsIn Which unit rectangle tile should the UVs be packed in, indicated by the coordinate of it's lower left hand corner
 	* @param Triangles Which UV triangles to pack, all triangles if nullptr
     */
-	bool UDIMPack(int32 TargetTextureResolution = 512, float GutterSize = 1.0f, const FVector2i& UDIMCoordsIn = FVector2i(0,0), const TArray<int32>* Triangles = nullptr);
+	UE_API bool UDIMPack(int32 TargetTextureResolution = 512, float GutterSize = 1.0f, const FVector2i& UDIMCoordsIn = FVector2i(0,0), const TArray<int32>* Triangles = nullptr);
 
 
 	//
@@ -334,24 +336,24 @@ public:
 	 * Compute geodesic center of given Mesh
 	 * @return true if a central vertex was found (but if false, will still give a frame for *some* vertex as long as there is at least one vertex in the mesh)
 	 */
-	static bool EstimateGeodesicCenterFrameVertex(const FDynamicMesh3& Mesh, FFrame3d& FrameOut, int32 &VertexIDOut, bool bAlignToUnitAxes = true);
+	static UE_API bool EstimateGeodesicCenterFrameVertex(const FDynamicMesh3& Mesh, FFrame3d& FrameOut, int32 &VertexIDOut, bool bAlignToUnitAxes = true);
 
 	/**
 	 * Compute geodesic center of given Mesh triangles (assumes they are connected)
 	 * @return true if a central vertex was found (but if false, will still give a frame for *some* vertex as long as there is at least one triangle)
 	 */
-	static bool EstimateGeodesicCenterFrameVertex(const FDynamicMesh3& Mesh, const TArray<int32>& Triangles, FFrame3d& FrameOut, int32& VertexIDOut, bool bAlignToUnitAxes = true);
+	static UE_API bool EstimateGeodesicCenterFrameVertex(const FDynamicMesh3& Mesh, const TArray<int32>& Triangles, FFrame3d& FrameOut, int32& VertexIDOut, bool bAlignToUnitAxes = true);
 
 	/**
 	* Compute the occupied 2D area for the UVs for the specified triangles and channel.
 	* @param BoundingBox if not null, return the overall bounding box of the UVs for the specified triangles along with the exact area occupied.
 	*/
-	static double DetermineAreaFromUVs(const FDynamicMeshUVOverlay& UVOverlay, const TArray<int32>& Triangles, FAxisAlignedBox2f* BoundingBox = nullptr);
+	static UE_API double DetermineAreaFromUVs(const FDynamicMeshUVOverlay& UVOverlay, const TArray<int32>& Triangles, FAxisAlignedBox2f* BoundingBox = nullptr);
 
 	/**
 	 * Transform UVs of the specified triangles
 	 */
-	static void TransformTriangleSelectionUVs(FDynamicMeshUVOverlay& UVOverlay, TConstArrayView<int32> Triangles, TFunctionRef<FVector2f(const FVector2f&)> TransformFunc);
+	static UE_API void TransformTriangleSelectionUVs(FDynamicMeshUVOverlay& UVOverlay, TConstArrayView<int32> Triangles, TFunctionRef<FVector2f(const FVector2f&)> TransformFunc);
 
 private: 
 
@@ -366,13 +368,14 @@ private:
 	/**
 	 * Helper method to solve either the Least Squares Conformal Map or the Spectral Conformal Map problems.
 	 */
-	bool SetTriangleUVsFromConformal(const TArray<int32>& Triangles, const FSetUVsFromConformalOptions& Options, FUVEditResult* Result);
+	UE_API bool SetTriangleUVsFromConformal(const TArray<int32>& Triangles, const FSetUVsFromConformalOptions& Options, FUVEditResult* Result);
 
 	// Helper to split bowties on element IDs
-	void SplitBowtiesOnUVElements(TArray<int32>& UVElementIDs, bool bAddNewElementsToInputArray = true);
+	UE_API void SplitBowtiesOnUVElements(TArray<int32>& UVElementIDs, bool bAddNewElementsToInputArray = true);
 };
 
 
 } // end namespace UE::Geometry
 } // end namespace UE
 
+#undef UE_API
diff --git a/Engine/Plugins/Runtime/GeometryProcessing/Source/DynamicMesh/Public/Parameterization/MeshPlanarSymmetry.h b/Engine/Plugins/Runtime/GeometryProcessing/Source/DynamicMesh/Public/Parameterization/MeshPlanarSymmetry.h
--- a/Engine/Plugins/Runtime/GeometryProcessing/Source/DynamicMesh/Public/Parameterization/MeshPlanarSymmetry.h
+++ b/Engine/Plugins/Runtime/GeometryProcessing/Source/DynamicMesh/Public/Parameterization/MeshPlanarSymmetry.h
@@ -7,6 +7,8 @@
 #include "Quaternion.h"
 #include "Spatial/PointHashGrid3.h"
 
+#define UE_API DYNAMICMESH_API
+
 namespace UE
 {
 namespace Geometry
@@ -22,7 +24,7 @@ namespace Geometry
  * The Symmetry Plane is specified by a Frame3d where the XY axes define the plane and the Z axis is the plane normal.
  * So references to "Positive" side below are relative to that Z axis/plane
  */
-class DYNAMICMESH_API FMeshPlanarSymmetry
+class FMeshPlanarSymmetry
 {
 public:
 
@@ -34,13 +36,13 @@ public:
 	 * planar/mirror-symmetry relationships, as well as "on plane" vertices
 	 * @return false if any non-on-plane vertex fails to find a match. 
 	 */
-	bool Initialize(FDynamicMesh3* Mesh, FDynamicMeshAABBTree3* Spatial, FFrame3d SymmetryFrameIn);
+	UE_API bool Initialize(FDynamicMesh3* Mesh, FDynamicMeshAABBTree3* Spatial, FFrame3d SymmetryFrameIn);
 	/**
 	 * Given a Mesh and its bounding box, and a Symmetry Plane/Frame, detect any pairs of vertices with
 	 * planar/mirror-symmetry relationships, as well as "on plane" vertices
 	 * @return false if any non-on-plane vertex fails to find a match.
 	 */
-	bool Initialize(FDynamicMesh3* Mesh, const FAxisAlignedBox3d& Bounds, FFrame3d SymmetryFrameIn);
+	UE_API bool Initialize(FDynamicMesh3* Mesh, const FAxisAlignedBox3d& Bounds, FFrame3d SymmetryFrameIn);
 
 	/**
 	 * Given a Mesh and its bounding box, find a Symmetry Plane/Frame and detect any pairs of vertices with
@@ -49,41 +51,41 @@ public:
 	 * @param PreferredNormals Optionally try to find a symmetry frame aligned to any normals passed in to this array. Tries the normals in order, so the first normal that fits (if any) will be used.
 	 * @return false if any non-on-plane vertex fails to find a match.
 	 */
-	bool FindPlaneAndInitialize(FDynamicMesh3* Mesh, const FAxisAlignedBox3d& Bounds, FFrame3d& SymmetryFrameOut, TArrayView<const FVector3d> PreferredNormals = TArrayView<const FVector3d>());
+	UE_API bool FindPlaneAndInitialize(FDynamicMesh3* Mesh, const FAxisAlignedBox3d& Bounds, FFrame3d& SymmetryFrameOut, TArrayView<const FVector3d> PreferredNormals = TArrayView<const FVector3d>());
 
 	/**
 	 * @return the input Point mirrored across the Symmetry plane
 	 */
-	FVector3d GetMirroredPosition(const FVector3d& Position) const;
+	UE_API FVector3d GetMirroredPosition(const FVector3d& Position) const;
 
 	/**
 	 * @return the input Vector/Axis mirrored across the Symmetry plane
 	 */
-	FVector3d GetMirroredAxis(const FVector3d& Axis) const;
+	UE_API FVector3d GetMirroredAxis(const FVector3d& Axis) const;
 
 	/**
 	 * @return the input Quaternion mirrored across the Symmetry plane
 	 */
-	FQuaterniond GetMirroredOrientation(const FQuaterniond& Orientation) const;
+	UE_API FQuaterniond GetMirroredOrientation(const FQuaterniond& Orientation) const;
 
 	/**
 	 * @return the input Frame mirrored to the "positive" side (unchanged if it is already on the positive side)
 	 */
-	FFrame3d GetPositiveSideFrame(FFrame3d FromFrame) const;
+	UE_API FFrame3d GetPositiveSideFrame(FFrame3d FromFrame) const;
 
 
 	/**
 	 * Update all the symmetry vertex positions based on their Positive-side pair vertex.
 	 * Also snap all on-plane vertices to the symmetry plane
 	 */
-	void FullSymmetryUpdate();
+	UE_API void FullSymmetryUpdate();
 
 
 	/** @return true if the TestMesh vertices match the current tracked vertex symmetry correspondence (i.e., the TestMesh vertex IDs match the known correspondence, and have the expected mirror symmetry) */
-	bool ValidateSymmetry(const FDynamicMesh3& TestMesh) const;
+	UE_API bool ValidateSymmetry(const FDynamicMesh3& TestMesh) const;
 
 	/** @return true if the TestMesh vertices in VertexROI match the current tracked vertex symmetry correspondence (i.e., the vertex IDs match the known correspondence, and have the expected mirror symmetry) */
-	bool ValidateSymmetry(const FDynamicMesh3& TestMesh, TConstArrayView<int32> VertexROI) const;
+	UE_API bool ValidateSymmetry(const FDynamicMesh3& TestMesh, TConstArrayView<int32> VertexROI) const;
 
 	//
 	// The set of functions below are intended to be used together in situations like 3D sculpting
@@ -97,13 +99,13 @@ public:
 	 * If bForceSameSizeWithGaps is true, MirrorVertexROIOut will be the same size as VertexROI, and -1 is stored 
 	 * for vertices in VertexROI that are source vertices or on the symmetry plane. Otherwise those vertices are skipped.
 	 */
-	void GetMirrorVertexROI(const TArray<int>& VertexROI, TArray<int>& MirrorVertexROIOut, bool bForceSameSizeWithGaps) const;
+	UE_API void GetMirrorVertexROI(const TArray<int>& VertexROI, TArray<int>& MirrorVertexROIOut, bool bForceSameSizeWithGaps) const;
 
 	/**
 	 * For any vertices in VertexIndices that are on-plane, take the position in VertexPositionsInOut, apply the
 	 * on-plane constraint, and return the constrained position in VertexPositionsInOut  (ie VertexPositionsInOut is updated to new positions)
 	 */
-	void ApplySymmetryPlaneConstraints(const TArray<int>& VertexIndices, TArray<FVector3d>& VertexPositionsInOut) const;
+	UE_API void ApplySymmetryPlaneConstraints(const TArray<int>& VertexIndices, TArray<FVector3d>& VertexPositionsInOut) const;
 
 	/**
 	 * Given the pairing (SourceVertexROI, SourceVertexPositions), compute the symmetry-constrained vertex positions for
@@ -111,7 +113,7 @@ public:
 	 * by calling GetMirrorVertexROI(SourceVertexROI, MirrorVertexROI, bForceSameSizeWithGaps=true), ie all the arrays
 	 * must be the same length
 	 */
-	void ComputeSymmetryConstrainedPositions(
+	UE_API void ComputeSymmetryConstrainedPositions(
 		const TArray<int>& SourceVertexROI,
 		const TArray<int>& MirrorVertexROI,
 		const TArray<FVector3d>& SourceVertexPositions,
@@ -139,8 +141,8 @@ protected:
 	// list of vertices, size == TargetMesh.MaxVertexID()
 	TArray<FSymmetryVertex> Vertices;
 
-	void UpdateSourceVertex(int32 VertexID);
-	void UpdatePlaneVertex(int32 VertexID);
+	UE_API void UpdateSourceVertex(int32 VertexID);
+	UE_API void UpdatePlaneVertex(int32 VertexID);
 
 private:
 
@@ -172,4 +174,6 @@ private:
 
 
 } // end namespace UE::Geometry
-} // end namespace UE
\ No newline at end of file
+} // end namespace UE
+
+#undef UE_API
diff --git a/Engine/Plugins/Runtime/GeometryProcessing/Source/DynamicMesh/Public/Parameterization/MeshRegionGraph.h b/Engine/Plugins/Runtime/GeometryProcessing/Source/DynamicMesh/Public/Parameterization/MeshRegionGraph.h
--- a/Engine/Plugins/Runtime/GeometryProcessing/Source/DynamicMesh/Public/Parameterization/MeshRegionGraph.h
+++ b/Engine/Plugins/Runtime/GeometryProcessing/Source/DynamicMesh/Public/Parameterization/MeshRegionGraph.h
@@ -5,6 +5,8 @@
 #include "DynamicMesh/DynamicMesh3.h"
 #include "Selections/MeshConnectedComponents.h"
 
+#define UE_API DYNAMICMESH_API
+
 namespace UE
 {
 namespace Geometry
@@ -23,7 +25,7 @@ namespace Geometry
  * 
  * An integer ID can be associated with each Region using the ExternalIDFunc passed to the various Build() functions.
  */
-class DYNAMICMESH_API FMeshRegionGraph
+class FMeshRegionGraph
 {
 public:
 
@@ -31,7 +33,7 @@ public:
 	 * Build a region graph for a Mesh using the given ConnectedComponents
 	 * @param ExternalIDFunc This allows the client to provide an integer ID that will be associated with each Region - called with the region index in Components
 	 */
-	void BuildFromComponents(const FDynamicMesh3& Mesh, 
+	UE_API void BuildFromComponents(const FDynamicMesh3& Mesh, 
 							 const FMeshConnectedComponents& Components,
 							 TFunctionRef<int32(int32)> ExternalIDFunc,
 							 TFunctionRef<bool(int32, int32)> TrisConnectedPredicate = [](int, int) { return true; } );
@@ -40,7 +42,7 @@ public:
 	 * Build a region graph for a Mesh using the given triangle sets
 	 * @param ExternalIDFunc This allows the client to provide an integer ID that will be associated with each Region - called with the region index in TriangleSets
 	 */
-	void BuildFromTriangleSets(const FDynamicMesh3& Mesh,
+	UE_API void BuildFromTriangleSets(const FDynamicMesh3& Mesh,
 							 const TArray<TArray<int32>>& TriangleSets,
 							 TFunctionRef<int32(int32)> ExternalIDFunc,
 		                     TFunctionRef<bool(int32, int32)> TrisConnectedPredicate = [](int, int) { return true; });
@@ -73,12 +75,12 @@ public:
 	/** 
 	 * @return list of Neigbour Region Indices for a Region 
 	 */
-	TArray<int32> GetNeighbours(int32 RegionIdx) const;
+	UE_API TArray<int32> GetNeighbours(int32 RegionIdx) const;
 
 	/**
 	 * @return true if the two regions are adjacent, ie they share an edge on at least one triangle
 	 */
-	bool AreRegionsConnected(int32 RegionAIndex, int32 RegionBIndex) const;
+	UE_API bool AreRegionsConnected(int32 RegionAIndex, int32 RegionBIndex) const;
 
 	/** @return true if the two triangles are in the same region */
 	bool AreTrianglesConnected(int32 TriangleA, int32 TriangleB) const
@@ -91,7 +93,7 @@ public:
 	 * A Small region is merged to it's most-similar neighbour, as defined by RegionSimilarityFunc.
 	 * The merge order is sorted by largest Similarity, ie most-similar regions are merged first.
 	 */
-	bool MergeSmallRegions(int32 SmallThreshold,
+	UE_API bool MergeSmallRegions(int32 SmallThreshold,
 						   TFunctionRef<float(int32 SmallRgnIdx, int32 NbrRgnIdx)> RegionSimilarityFunc );
 
 	/**
@@ -99,7 +101,7 @@ public:
 	 * @param RemoveRgnIdx this region will be removed
 	 * @param MergeToRgnIdx this region will be added to
 	 */
-	bool MergeRegion(int32 RemoveRgnIdx, int32 MergeToRgnIdx);
+	UE_API bool MergeRegion(int32 RemoveRgnIdx, int32 MergeToRgnIdx);
 
 	/**
 	 * Optimize region borders by swapping triangles. The swap criteria is, if 2/3 of the
@@ -107,7 +109,7 @@ public:
 	 * This swapping process is repeated in rounds, MaxRounds times.
 	 * @return true if any swaps were applied
 	 */
-	bool OptimizeBorders(int32 MaxRounds = 25);
+	UE_API bool OptimizeBorders(int32 MaxRounds = 25);
 
 
 	/**
@@ -146,13 +148,15 @@ public:
 	TArray<int32> EmptyTriSet;
 
 protected:
-	void BuildNeigbours(int32 RegionIdx);
+	UE_API void BuildNeigbours(int32 RegionIdx);
 	
-	bool MergeSmallRegionsPass(int32 SmallThreshold,
+	UE_API bool MergeSmallRegionsPass(int32 SmallThreshold,
 						   TFunctionRef<float(int32 SmallRgnIdx, int32 NbrRgnIdx)> RegionSimilarityFunc);
 
 };
 
 
 } // end namespace UE::Geometry
-} // end namespace UE
\ No newline at end of file
+} // end namespace UE
+
+#undef UE_API
diff --git a/Engine/Plugins/Runtime/GeometryProcessing/Source/DynamicMesh/Public/Parameterization/MeshUDIMClassifier.h b/Engine/Plugins/Runtime/GeometryProcessing/Source/DynamicMesh/Public/Parameterization/MeshUDIMClassifier.h
--- a/Engine/Plugins/Runtime/GeometryProcessing/Source/DynamicMesh/Public/Parameterization/MeshUDIMClassifier.h
+++ b/Engine/Plugins/Runtime/GeometryProcessing/Source/DynamicMesh/Public/Parameterization/MeshUDIMClassifier.h
@@ -5,6 +5,8 @@
 #include "DynamicMesh/DynamicMesh3.h"
 #include "DynamicMesh/DynamicMeshAttributeSet.h"
 
+#define UE_API DYNAMICMESH_API
+
 
 namespace UE
 {
@@ -17,21 +19,21 @@ namespace UE
 		/**
 		* FDynamicMeshUDIMClassifier is a utility class for identifying active UDIMs from a FDynamicMesh's UV overlay
 		 */
-		class DYNAMICMESH_API FDynamicMeshUDIMClassifier
+		class FDynamicMeshUDIMClassifier
 		{
 		public:			
-			explicit FDynamicMeshUDIMClassifier(const FDynamicMeshUVOverlay* UVOverlay, TOptional<TArray<int32>> Selection = TOptional<TArray<int32>>());
+			UE_API explicit FDynamicMeshUDIMClassifier(const FDynamicMeshUVOverlay* UVOverlay, TOptional<TArray<int32>> Selection = TOptional<TArray<int32>>());
 
-			TArray<FVector2i> ActiveTiles() const;
-			TArray<int32> TidsForTile(FVector2i TileIndexIn) const;
+			UE_API TArray<FVector2i> ActiveTiles() const;
+			UE_API TArray<int32> TidsForTile(FVector2i TileIndexIn) const;
 
-			static FVector2i ClassifyTrianglesToUDIM(const FDynamicMeshUVOverlay* UVOverlay, TArray<int32> Tids);
-			static FVector2i ClassifyBoundingBoxToUDIM(const FDynamicMeshUVOverlay* UVOverlay, const FAxisAlignedBox2d& BoundingBox);
-			static FVector2i ClassifyPointToUDIM(const FVector2f& UVPoint);
+			static UE_API FVector2i ClassifyTrianglesToUDIM(const FDynamicMeshUVOverlay* UVOverlay, TArray<int32> Tids);
+			static UE_API FVector2i ClassifyBoundingBoxToUDIM(const FDynamicMeshUVOverlay* UVOverlay, const FAxisAlignedBox2d& BoundingBox);
+			static UE_API FVector2i ClassifyPointToUDIM(const FVector2f& UVPoint);
 
 		protected:
 
-			void ClassifyUDIMs();
+			UE_API void ClassifyUDIMs();
 
 			/** The UV Overlay to analyze for UDIMs */
 			const FDynamicMeshUVOverlay* UVOverlay = nullptr;
@@ -41,4 +43,6 @@ namespace UE
 			TMap<FVector2i, TArray<int32> > UDIMs;
 		};
 	}
-}
\ No newline at end of file
+}
+
+#undef UE_API
diff --git a/Engine/Plugins/Runtime/GeometryProcessing/Source/DynamicMesh/Public/Parameterization/MeshUVPacking.h b/Engine/Plugins/Runtime/GeometryProcessing/Source/DynamicMesh/Public/Parameterization/MeshUVPacking.h
--- a/Engine/Plugins/Runtime/GeometryProcessing/Source/DynamicMesh/Public/Parameterization/MeshUVPacking.h
+++ b/Engine/Plugins/Runtime/GeometryProcessing/Source/DynamicMesh/Public/Parameterization/MeshUVPacking.h
@@ -5,6 +5,8 @@
 #include "DynamicMesh/DynamicMesh3.h"
 #include "DynamicMesh/DynamicMeshAttributeSet.h"
 
+#define UE_API DYNAMICMESH_API
+
 
 namespace UE
 {
@@ -19,7 +21,7 @@ class FMeshConnectedComponents;
  * UV Overlay. The island topology and UV unwraps must already be created, this
  * class simply scales/rotates/translates the islands to fit.
  */
-class DYNAMICMESH_API FDynamicMeshUVPacker
+class FDynamicMeshUVPacker
 {
 public:
 	/** The UV Overlay we will be repacking */
@@ -46,8 +48,8 @@ public:
 	/** Attempt to rescale islands to match texel-to-world-space ratio across islands, based on ratio of World- and UV-space edge lengths */
 	bool bScaleIslandsByWorldSpaceTexelRatio = false;
 
-	explicit FDynamicMeshUVPacker(FDynamicMeshUVOverlay* UVOverlay);
-	explicit FDynamicMeshUVPacker(FDynamicMeshUVOverlay* UVOverlay, TUniquePtr<TArray<int32>>&& TidsToRepackIn);
+	UE_API explicit FDynamicMeshUVPacker(FDynamicMeshUVOverlay* UVOverlay);
+	UE_API explicit FDynamicMeshUVPacker(FDynamicMeshUVOverlay* UVOverlay, TUniquePtr<TArray<int32>>&& TidsToRepackIn);
 
 
 	/**
@@ -55,22 +57,24 @@ public:
 	 * All UV islands are packed into standard positive-unit-square.
 	 * Only supports single-pixel border size.
 	 */
-	bool StandardPack();
+	UE_API bool StandardPack();
 
 	/**
 	 * Uniformly scale all UV islands so that the largest fits in positive-unit-square,
 	 * and translate each islands separately so that it's bbox-min is at the origin.
 	 * So the islands are "stacked" and all fit in the unit box.
 	 */
-	bool StackPack();
+	UE_API bool StackPack();
 
 
 protected:
 
-	FMeshConnectedComponents CollectUVIslandsToPack(const FUVOverlayView& MeshView);
+	UE_API FMeshConnectedComponents CollectUVIslandsToPack(const FUVOverlayView& MeshView);
 
 };
 
 
 } // end namespace UE::Geometry
-} // end namespace UE
\ No newline at end of file
+} // end namespace UE
+
+#undef UE_API
diff --git a/Engine/Plugins/Runtime/GeometryProcessing/Source/DynamicMesh/Public/Polygroups/PolygroupSet.h b/Engine/Plugins/Runtime/GeometryProcessing/Source/DynamicMesh/Public/Polygroups/PolygroupSet.h
--- a/Engine/Plugins/Runtime/GeometryProcessing/Source/DynamicMesh/Public/Polygroups/PolygroupSet.h
+++ b/Engine/Plugins/Runtime/GeometryProcessing/Source/DynamicMesh/Public/Polygroups/PolygroupSet.h
@@ -6,6 +6,8 @@
 #include "DynamicMesh/DynamicMesh3.h"
 #include "DynamicMesh/DynamicMeshAttributeSet.h"
 
+#define UE_API DYNAMICMESH_API
+
 
 namespace UE
 {
@@ -17,7 +19,7 @@ namespace Geometry
  * group set stored on the mesh, and then N extended group layers stored in the mesh AttributeSet.
  * This struct can represent either.
  */
-struct DYNAMICMESH_API FPolygroupLayer
+struct FPolygroupLayer
 {
 	/** If true, layer is the default FDynamicMesh3 triangle groups layer */
 	bool bIsDefaultLayer = true;
@@ -42,7 +44,7 @@ struct DYNAMICMESH_API FPolygroupLayer
 	}
 
 	/** @return true if the specified layer (default or extended) exist and is initialized on the given Mesh */
-	bool CheckExists(const FDynamicMesh3* Mesh) const;
+	UE_API bool CheckExists(const FDynamicMesh3* Mesh) const;
 };
 
 
@@ -57,7 +59,7 @@ struct DYNAMICMESH_API FPolygroupLayer
  * For consistency with FDynamicMesh3, MaxGroupID is set such that all GroupIDs are less than MaxGroupID
  *
  */
-struct DYNAMICMESH_API FPolygroupSet
+struct FPolygroupSet
 {
 	const FDynamicMesh3* Mesh = nullptr;
 	const FDynamicMeshPolygroupAttribute* PolygroupAttrib = nullptr;
@@ -65,22 +67,22 @@ struct DYNAMICMESH_API FPolygroupSet
 	int32 MaxGroupID = 0; // Note: all group IDs are less than MaxGroupID
 
 	/** Initialize a PolygroupSet for the given Mesh, and standard triangle group layer */
-	explicit FPolygroupSet(const FDynamicMesh3* MeshIn);
+	UE_API explicit FPolygroupSet(const FDynamicMesh3* MeshIn);
 
 	/** Initialize a PolygroupSet for the given Mesh, and standard triangle group layer */
-	explicit FPolygroupSet(const FDynamicMesh3* MeshIn, FPolygroupLayer GroupLayer);
+	UE_API explicit FPolygroupSet(const FDynamicMesh3* MeshIn, FPolygroupLayer GroupLayer);
 
 	/** Initialize a PolygroupSet for given Mesh and specific Polygroup attribute layer */
-	explicit FPolygroupSet(const FDynamicMesh3* MeshIn, const FDynamicMeshPolygroupAttribute* PolygroupAttribIn);
+	UE_API explicit FPolygroupSet(const FDynamicMesh3* MeshIn, const FDynamicMeshPolygroupAttribute* PolygroupAttribIn);
 
 	/** Initialize a PolygroupSet for given Mesh and specific Polygroup attribute layer, found by index. If not valid, fall back to standard triangle group layer. */
-	explicit FPolygroupSet(const FDynamicMesh3* MeshIn, int32 PolygroupLayerIndex);
+	UE_API explicit FPolygroupSet(const FDynamicMesh3* MeshIn, int32 PolygroupLayerIndex);
 
 	/** Initialize a PolygroupSet for given Mesh and specific Polygroup attribute layer, found by name. If not valid, fall back to standard triangle group layer. */
-	explicit FPolygroupSet(const FDynamicMesh3* MeshIn, FName AttribName);
+	UE_API explicit FPolygroupSet(const FDynamicMesh3* MeshIn, FName AttribName);
 
 	/** Initialize a PolygroupSet by copying an existing PolygroupSet */
-	explicit FPolygroupSet(const FPolygroupSet* CopyIn);
+	UE_API explicit FPolygroupSet(const FPolygroupSet* CopyIn);
 
 
 	/** @return Mesh this PolygroupSet references  */
@@ -134,7 +136,7 @@ struct DYNAMICMESH_API FPolygroupSet
 	/**
 	 * Calculate the current maximum PolygroupID used in the active set and store in MaxGroupID member
 	 */
-	void RecalculateMaxGroupID();
+	UE_API void RecalculateMaxGroupID();
 
 	/**
 	 * Allocate a new unused PolygroupID by incrementing the MaxGroupID member
@@ -152,3 +154,5 @@ struct DYNAMICMESH_API FPolygroupSet
 
 }	// end namespace Geometry
 }	// end namespace UE
+
+#undef UE_API
diff --git a/Engine/Plugins/Runtime/GeometryProcessing/Source/DynamicMesh/Public/Polygroups/PolygroupsGenerator.h b/Engine/Plugins/Runtime/GeometryProcessing/Source/DynamicMesh/Public/Polygroups/PolygroupsGenerator.h
--- a/Engine/Plugins/Runtime/GeometryProcessing/Source/DynamicMesh/Public/Polygroups/PolygroupsGenerator.h
+++ b/Engine/Plugins/Runtime/GeometryProcessing/Source/DynamicMesh/Public/Polygroups/PolygroupsGenerator.h
@@ -6,6 +6,8 @@
 #include "DynamicMesh/DynamicMesh3.h"
 #include "Polygroups/PolygroupSet.h"
 
+#define UE_API DYNAMICMESH_API
+
 using UE::Geometry::FDynamicMesh3;
 
 
@@ -18,12 +20,12 @@ namespace Geometry
  * FPolygroupsGenerator generates (face/tri/poly)groups for an input mesh based on the geometry and attributes of the input mesh.
  * The various FindPolygroupsFromX() are the driving functions, each performs the full computation. 
  */
-class DYNAMICMESH_API FPolygroupsGenerator
+class FPolygroupsGenerator
 {
 public:
 
 	FPolygroupsGenerator() {}
-	FPolygroupsGenerator(FDynamicMesh3* MeshIn);
+	UE_API FPolygroupsGenerator(FDynamicMesh3* MeshIn);
 
 	// 
 	// Input parameters
@@ -59,7 +61,7 @@ public:
 	 * Find Polygroups by randomly picking initial seed triangles and then flood-filling outwards,
 	 * stopping when the opening angle at an edge is larger than the angle defined by the OneMinusCosAngleTolerance.
 	 */
-	bool FindPolygroupsFromFaceNormals(
+	UE_API bool FindPolygroupsFromFaceNormals(
 		double OneMinusCosAngleTolerance = 0.0001,
 		bool bRespectUVSeams = false,
 		bool bRespectNormalSeams = false,
@@ -68,29 +70,29 @@ public:
 	/**
 	 * Find Polygroups based on UV Islands, ie each UV Island becomes a Polygroup
 	 */
-	bool FindPolygroupsFromUVIslands(int32 UVLayer = 0);
+	UE_API bool FindPolygroupsFromUVIslands(int32 UVLayer = 0);
 
 	/**
 	 * Find PolyGroups based on Material IDs -- i.e., connected groups of triangles with the same Material ID become a PolyGroup
 	 */
-	bool FindPolygroupsFromMaterialIDs();
+	UE_API bool FindPolygroupsFromMaterialIDs();
 
 	/**
 	* Find Polygroups based on Seams in UV Overlay
 	*/
-	bool FindPolygroupsFromHardNormalSeams();
+	UE_API bool FindPolygroupsFromHardNormalSeams();
 
 	/**
 	 * Find PolyGroups based on mesh connectivity and optionally a custom indicator function
 	 * @param AreTrisConnected Optional custom function: Takes neighboring triangle IDs and returns true if they should be in the same PolyGroup
 	 * @return true if at least one polygroup was found
 	 */
-	bool FindPolygroupsFromConnectedTris(TFunction<bool(int32 TriA, int32 TriB)> AreTrisConnected = nullptr);
+	UE_API bool FindPolygroupsFromConnectedTris(TFunction<bool(int32 TriA, int32 TriB)> AreTrisConnected = nullptr);
 
 	/**
 	* Find Polygroups by trying to invert triangluation of a polygon mesh
 	*/
-	bool FindSourceMeshPolygonPolygroups(
+	UE_API bool FindSourceMeshPolygonPolygroups(
 		bool bRespectUVSeams = true,
 		bool bRespectNormalSeams = false,
 		double QuadAdjacencyWeight = 1.0,
@@ -114,7 +116,7 @@ public:
 	 * If StartingGroups are provided, then new groups are constrained to subdivide existing groups, ie
 	 * none of the new groups will cross the boundaries of the StartingGroups
 	 */
-	bool FindPolygroupsFromFurthestPointSampling(
+	UE_API bool FindPolygroupsFromFurthestPointSampling(
 		int32 NumPoints, 
 		EWeightingType WeightingType, 
 		FVector3d WeightingCoeffs = FVector3d::One(),
@@ -141,7 +143,7 @@ public:
 	* are part of a group.
 	* 
 	*/
-	bool FindPolygroupsFromAreaDensityPointSampling(
+	UE_API bool FindPolygroupsFromAreaDensityPointSampling(
 		int32 NumPoints,
 		EWeightingType WeightingType,
 		FVector3d WeightingCoeffs = FVector3d::One(),
@@ -151,27 +153,29 @@ public:
 	* Copy the computed Polygroups to the input Mesh. Will be automatically called by the 
 	* FindPolygroupsFromX() functions above if bCopyToMesh=true
 	*/
-	void CopyPolygroupsToMesh();
+	UE_API void CopyPolygroupsToMesh();
 
 	/**
 	* Copy the computed Polygroups to the given PolygroupSet and Mesh
 	*/
-	void CopyPolygroupsToPolygroupSet(FPolygroupSet& Polygroups, FDynamicMesh3& TargetMesh);
+	UE_API void CopyPolygroupsToPolygroupSet(FPolygroupSet& Polygroups, FDynamicMesh3& TargetMesh);
 
 	/**
 	 * Initialize the PolygroupEdges output member by finding all the mesh edges that are on polygroup borders.
 	 * Requires that polygroups have been computed and copied to the mesh
 	 */
-	bool FindPolygroupEdges();
+	UE_API bool FindPolygroupEdges();
 
 
 protected:
-	void PostProcessPolygroups(bool bApplyMerging, TFunctionRef<bool(int32, int32)> TrisConnectedPredicate = [](int, int) { return true; });
-	void OptimizePolygroups(TFunctionRef<bool(int32, int32)> TrisConnectedPredicate = [](int, int) { return true; });
+	UE_API void PostProcessPolygroups(bool bApplyMerging, TFunctionRef<bool(int32, int32)> TrisConnectedPredicate = [](int, int) { return true; });
+	UE_API void OptimizePolygroups(TFunctionRef<bool(int32, int32)> TrisConnectedPredicate = [](int, int) { return true; });
 
-	void GetSeamConstraintEdges(bool bUVSeams, bool bNormalSeams, TSet<int32>& ConstraintEdgesOut) const;
+	UE_API void GetSeamConstraintEdges(bool bUVSeams, bool bNormalSeams, TSet<int32>& ConstraintEdgesOut) const;
 };
 
 
 }	// end namespace Geometry
 }	// end namespace UE
+
+#undef UE_API
diff --git a/Engine/Plugins/Runtime/GeometryProcessing/Source/DynamicMesh/Public/Remesher.h b/Engine/Plugins/Runtime/GeometryProcessing/Source/DynamicMesh/Public/Remesher.h
--- a/Engine/Plugins/Runtime/GeometryProcessing/Source/DynamicMesh/Public/Remesher.h
+++ b/Engine/Plugins/Runtime/GeometryProcessing/Source/DynamicMesh/Public/Remesher.h
@@ -6,6 +6,8 @@
 
 #include "MeshRefinerBase.h"
 
+#define UE_API DYNAMICMESH_API
+
 
 namespace UE
 {
@@ -38,7 +40,7 @@ namespace Geometry
  * @todo port fast queue-based convergent remeshing
  * 
  */
-class DYNAMICMESH_API FRemesher : public FMeshRefinerBase
+class FRemesher : public FMeshRefinerBase
 {
 public:
 	//
@@ -134,7 +136,7 @@ public:
 	virtual ~FRemesher() {}
 
 	/** Set min/max edge-lengths to sane values for given target edge length */
-	void SetTargetEdgeLength(double fLength);
+	UE_API void SetTargetEdgeLength(double fLength);
 
 
 	/**
@@ -143,7 +145,7 @@ public:
 	 * between remesh passes, otherwise you will get weird results.
 	 * But you will probably still come out ahead, computation-time-wise
 	 */
-	virtual void Precompute();
+	UE_API virtual void Precompute();
 
 
 	/**
@@ -153,7 +155,7 @@ public:
 	 * - Projection pass if ProjectionMode == AfterRefinement
 	 * - number of modified edges returned in ModifiedEdgesLastPass
 	 */
-	virtual void BasicRemeshPass();
+	UE_API virtual void BasicRemeshPass();
 
 
 
@@ -242,7 +244,7 @@ protected:
 		Failed_NotAnEdge
 	};
 
-	virtual EProcessResult ProcessEdge(int edgeID);
+	UE_API virtual EProcessResult ProcessEdge(int edgeID);
 
 
 
@@ -250,47 +252,47 @@ protected:
 	// The edge needs to inherit the constraint on the other pre-existing edge that we kept.
 	// In addition, if the edge vertices were both constrained, then we /might/
 	// want to also constrain this new vertex, possibly project to constraint target. 
-	virtual void UpdateAfterSplit(int edgeID, int va, int vb, const FDynamicMesh3::FEdgeSplitInfo& splitInfo);
+	UE_API virtual void UpdateAfterSplit(int edgeID, int va, int vb, const FDynamicMesh3::FEdgeSplitInfo& splitInfo);
 
 	TDynamicVector<FVector3d> TempPosBuffer;	// this is a temporary buffer used by smoothing and projection
 	TArray<bool> TempFlagBuffer;				// list of which indices in TempPosBuffer were modified
 
-	virtual void InitializeVertexBufferForPass();
+	UE_API virtual void InitializeVertexBufferForPass();
 
-	virtual void ApplyVertexBuffer(bool bParallel);
+	UE_API virtual void ApplyVertexBuffer(bool bParallel);
 
 
 
 	// applies a smoothing pass to the mesh, storing intermediate positions in a buffer and then writing them at the end (so, no order effect)
-	virtual void FullSmoothPass_Buffer(bool bParallel);
+	UE_API virtual void FullSmoothPass_Buffer(bool bParallel);
 
 	// computes smoothed vertex position w/ proper constraints/etc. Does not modify mesh. 
-	virtual FVector3d ComputeSmoothedVertexPos(int VertexID,
+	UE_API virtual FVector3d ComputeSmoothedVertexPos(int VertexID,
 		TFunction<FVector3d(const FDynamicMesh3 &, int, double)> SmoothFunc, bool& bModified);
 
 	// FullSmoothPass_Buffer this to calls VertexSmoothFunc for each vertex of the mesh
-	virtual void ApplyToSmoothVertices(const TFunction<void(int)>& VertexSmoothFunc);
+	UE_API virtual void ApplyToSmoothVertices(const TFunction<void(int)>& VertexSmoothFunc);
 
 	// returns the function we want to use to compute a smoothed vertex position - will be CustomSmoothF if set, otherwise one of cotan/meanvalue/uniform
-	virtual TFunction<FVector3d(const FDynamicMesh3&, int, double)> GetSmoothFunction();
+	UE_API virtual TFunction<FVector3d(const FDynamicMesh3&, int, double)> GetSmoothFunction();
 
 
 	// Project vertices onto projection target.
-	virtual void FullProjectionPass(bool bParallel);
+	UE_API virtual void FullProjectionPass(bool bParallel);
 
 	// Project a single vertex using the given target
-	virtual void ProjectVertex(int VertexID, IProjectionTarget* UseTarget);
+	UE_API virtual void ProjectVertex(int VertexID, IProjectionTarget* UseTarget);
 
 	// used by collapse-edge to get projected position for new vertex
-	virtual FVector3d GetProjectedCollapsePosition(int VertexID, const FVector3d& vNewPos);
+	UE_API virtual FVector3d GetProjectedCollapsePosition(int VertexID, const FVector3d& vNewPos);
 
 	// Apply the given projection function to all vertices in the mesh
-	virtual void ApplyToProjectVertices(const TFunction<void(int)>& VertexProjectFunc);
+	UE_API virtual void ApplyToProjectVertices(const TFunction<void(int)>& VertexProjectFunc);
 
 	// Move all vertices in parallel. Update the timestamps.
 	// The NewVertexPosition function accepts a vertex index and returns a new vertex position, as well as setting the 
 	// bool out-param to indicate that the vertex has moved.
-	void MoveVerticesParallel(TFunction<FVector3d(int, bool&)> NewVertexPosition);
+	UE_API void MoveVerticesParallel(TFunction<FVector3d(int, bool&)> NewVertexPosition);
 
 
 	/*
@@ -387,4 +389,6 @@ protected:
 
 
 } // end namespace UE::Geometry
-} // end namespace UE
\ No newline at end of file
+} // end namespace UE
+
+#undef UE_API
diff --git a/Engine/Plugins/Runtime/GeometryProcessing/Source/DynamicMesh/Public/Sampling/MeshConstantMapEvaluator.h b/Engine/Plugins/Runtime/GeometryProcessing/Source/DynamicMesh/Public/Sampling/MeshConstantMapEvaluator.h
--- a/Engine/Plugins/Runtime/GeometryProcessing/Source/DynamicMesh/Public/Sampling/MeshConstantMapEvaluator.h
+++ b/Engine/Plugins/Runtime/GeometryProcessing/Source/DynamicMesh/Public/Sampling/MeshConstantMapEvaluator.h
@@ -5,6 +5,8 @@
 #include "Sampling/MeshMapEvaluator.h"
 #include "DynamicMesh/MeshTangents.h"
 
+#define UE_API DYNAMICMESH_API
+
 namespace UE
 {
 namespace Geometry
@@ -14,26 +16,26 @@ namespace Geometry
  * A mesh evaluator for constant data. This evaluator can be useful as a filler
  * when computing per channel color data.
  */
-class DYNAMICMESH_API FMeshConstantMapEvaluator : public FMeshMapEvaluator
+class FMeshConstantMapEvaluator : public FMeshMapEvaluator
 {
 public:
 	FMeshConstantMapEvaluator() = default;
-	explicit FMeshConstantMapEvaluator(float ValueIn);
+	UE_API explicit FMeshConstantMapEvaluator(float ValueIn);
 	
 	// Begin FMeshMapEvaluator interface
-	virtual void Setup(const FMeshBaseBaker& Baker, FEvaluationContext& Context) override;
+	UE_API virtual void Setup(const FMeshBaseBaker& Baker, FEvaluationContext& Context) override;
 
-	virtual const TArray<EComponents>& DataLayout() const override;
+	UE_API virtual const TArray<EComponents>& DataLayout() const override;
 
 	virtual EMeshMapEvaluatorType Type() const override { return EMeshMapEvaluatorType::Constant; }
 	// End FMeshMapEvaluator interface
 
-	static void EvaluateSample(float*& Out, const FCorrespondenceSample& Sample, void* EvalData);
+	static UE_API void EvaluateSample(float*& Out, const FCorrespondenceSample& Sample, void* EvalData);
 
-	static void EvaluateDefault(float*& Out, void* EvalData);
+	static UE_API void EvaluateDefault(float*& Out, void* EvalData);
 
-	static void EvaluateColor(const int DataIdx, float*& In, FVector4f& Out, void* EvalData);
-	static void EvaluateChannel(const int DataIdx, float*& In, float& Out, void* EvalData);
+	static UE_API void EvaluateColor(const int DataIdx, float*& In, FVector4f& Out, void* EvalData);
+	static UE_API void EvaluateChannel(const int DataIdx, float*& In, float& Out, void* EvalData);
 
 public:
 	float Value = 0.0f;
@@ -42,3 +44,4 @@ public:
 } // end namespace UE::Geometry
 } // end namespace UE
 
+#undef UE_API
diff --git a/Engine/Plugins/Runtime/GeometryProcessing/Source/DynamicMesh/Public/Sampling/MeshCurvatureMapEvaluator.h b/Engine/Plugins/Runtime/GeometryProcessing/Source/DynamicMesh/Public/Sampling/MeshCurvatureMapEvaluator.h
--- a/Engine/Plugins/Runtime/GeometryProcessing/Source/DynamicMesh/Public/Sampling/MeshCurvatureMapEvaluator.h
+++ b/Engine/Plugins/Runtime/GeometryProcessing/Source/DynamicMesh/Public/Sampling/MeshCurvatureMapEvaluator.h
@@ -4,6 +4,8 @@
 
 #include "Sampling/MeshMapEvaluator.h"
 
+#define UE_API DYNAMICMESH_API
+
 namespace UE
 {
 namespace Geometry
@@ -15,7 +17,7 @@ class FMeshVertexCurvatureCache;
 /**
  * A mesh evaluator for mesh curvatures.
  */
-class DYNAMICMESH_API FMeshCurvatureMapEvaluator : public FMeshMapEvaluator
+class FMeshCurvatureMapEvaluator : public FMeshMapEvaluator
 {
 public:
 	enum class ECurvatureType
@@ -55,23 +57,23 @@ public:
 
 public:
 	// Begin FMeshMapEvaluator interface
-	virtual void Setup(const FMeshBaseBaker& Baker, FEvaluationContext& Context) override;
+	UE_API virtual void Setup(const FMeshBaseBaker& Baker, FEvaluationContext& Context) override;
 
-	virtual const TArray<EComponents>& DataLayout() const override;
+	UE_API virtual const TArray<EComponents>& DataLayout() const override;
 
 	virtual EMeshMapEvaluatorType Type() const override { return EMeshMapEvaluatorType::Curvature; }
 	// End FMeshMapEvaluator interface
 
-	static void EvaluateSample(float*& Out, const FCorrespondenceSample& Sample, void* EvalData);
+	static UE_API void EvaluateSample(float*& Out, const FCorrespondenceSample& Sample, void* EvalData);
 
-	static void EvaluateDefault(float*& Out, void* EvalData);
+	static UE_API void EvaluateDefault(float*& Out, void* EvalData);
 
-	static void EvaluateColor(const int DataIdx, float*& In, FVector4f& Out, void* EvalData);
+	static UE_API void EvaluateColor(const int DataIdx, float*& In, FVector4f& Out, void* EvalData);
 
-	static void EvaluateChannel(const int DataIdx, float*& In, float& Out, void* EvalData);
+	static UE_API void EvaluateChannel(const int DataIdx, float*& In, float& Out, void* EvalData);
 
 	/** Populate Curvatures member if valid data has not been provided */
-	void CacheDetailCurvatures();
+	UE_API void CacheDetailCurvatures();
 
 protected:
 	// Cached data
@@ -84,13 +86,14 @@ protected:
 	FVector3f PositiveColor;
 
 protected:
-	double GetCurvature(int32 vid) const;
-	void GetColorMapRange(FVector3f& NegativeColorOut, FVector3f& ZeroColorOut, FVector3f& PositiveColorOut) const;
+	UE_API double GetCurvature(int32 vid) const;
+	UE_API void GetColorMapRange(FVector3f& NegativeColorOut, FVector3f& ZeroColorOut, FVector3f& PositiveColorOut) const;
 
 private:
-	double SampleFunction(const FCorrespondenceSample& SampleData) const;
+	UE_API double SampleFunction(const FCorrespondenceSample& SampleData) const;
 };
 
 } // end namespace UE::Geometry
 } // end namespace UE
 
+#undef UE_API
diff --git a/Engine/Plugins/Runtime/GeometryProcessing/Source/DynamicMesh/Public/Sampling/MeshGenericWorldPositionBaker.h b/Engine/Plugins/Runtime/GeometryProcessing/Source/DynamicMesh/Public/Sampling/MeshGenericWorldPositionBaker.h
--- a/Engine/Plugins/Runtime/GeometryProcessing/Source/DynamicMesh/Public/Sampling/MeshGenericWorldPositionBaker.h
+++ b/Engine/Plugins/Runtime/GeometryProcessing/Source/DynamicMesh/Public/Sampling/MeshGenericWorldPositionBaker.h
@@ -8,6 +8,8 @@
 #include "Image/ImageBuilder.h"
 #include "DynamicMesh/MeshTangents.h"
 
+#define UE_API DYNAMICMESH_API
+
 namespace UE
 {
 namespace Geometry
@@ -17,7 +19,7 @@ namespace Geometry
  * Bake Colors based on arbitrary Position/Normal sampling. Useful for (eg) baking
  * from an arbitrary function defined over space
  */
-class DYNAMICMESH_API FMeshGenericWorldPositionColorBaker : public FMeshImageBaker
+class FMeshGenericWorldPositionColorBaker : public FMeshImageBaker
 {
 public:
 	virtual ~FMeshGenericWorldPositionColorBaker() {}
@@ -35,7 +37,7 @@ public:
 	// Compute functions
 	//
 
-	virtual void Bake() override;
+	UE_API virtual void Bake() override;
 
 	//
 	// Output
@@ -56,7 +58,7 @@ protected:
  * Bake Tangent-Space Normals based on arbitrary Position/Normal sampling. Useful for (eg) 
  * baking from an arbitrary function defined over space.
  */
-class DYNAMICMESH_API FMeshGenericWorldPositionNormalBaker : public FMeshImageBaker
+class FMeshGenericWorldPositionNormalBaker : public FMeshImageBaker
 {
 public:
 	virtual ~FMeshGenericWorldPositionNormalBaker() {}
@@ -76,7 +78,7 @@ public:
 	// Compute functions
 	//
 
-	virtual void Bake() override;
+	UE_API virtual void Bake() override;
 
 	//
 	// Output
@@ -97,3 +99,5 @@ protected:
 
 } // end namespace UE::Geometry
 } // end namespace UE
+
+#undef UE_API
diff --git a/Engine/Plugins/Runtime/GeometryProcessing/Source/DynamicMesh/Public/Sampling/MeshHeightMapEvaluator.h b/Engine/Plugins/Runtime/GeometryProcessing/Source/DynamicMesh/Public/Sampling/MeshHeightMapEvaluator.h
--- a/Engine/Plugins/Runtime/GeometryProcessing/Source/DynamicMesh/Public/Sampling/MeshHeightMapEvaluator.h
+++ b/Engine/Plugins/Runtime/GeometryProcessing/Source/DynamicMesh/Public/Sampling/MeshHeightMapEvaluator.h
@@ -5,6 +5,8 @@
 #include "Sampling/MeshMapEvaluator.h"
 #include "Sampling/MeshBakerCommon.h"
 
+#define UE_API DYNAMICMESH_API
+
 namespace UE
 {
 namespace Geometry
@@ -15,7 +17,7 @@ class IMeshBakerDetailSampler;
 /**
  * A mesh evaluator for mesh height as color data.
  */
-class DYNAMICMESH_API FMeshHeightMapEvaluator : public FMeshMapEvaluator
+class FMeshHeightMapEvaluator : public FMeshMapEvaluator
 {
 public:
 	enum class EHeightRangeMode
@@ -34,19 +36,19 @@ public:
 	FMeshHeightMapEvaluator& operator=(const FMeshHeightMapEvaluator&) = default;
 
 	// Begin FMeshMapEvaluator interface
-	virtual void Setup(const FMeshBaseBaker& Baker, FEvaluationContext& Context) override;
+	UE_API virtual void Setup(const FMeshBaseBaker& Baker, FEvaluationContext& Context) override;
 
-	virtual const TArray<EComponents>& DataLayout() const override;
+	UE_API virtual const TArray<EComponents>& DataLayout() const override;
 
 	virtual EMeshMapEvaluatorType Type() const override { return EMeshMapEvaluatorType::Height; }
 	// End FMeshMapEvaluator interface
 
-	static void EvaluateSample(float*& Out, const FCorrespondenceSample& Sample, void* EvalData);
+	static UE_API void EvaluateSample(float*& Out, const FCorrespondenceSample& Sample, void* EvalData);
 
-	static void EvaluateDefault(float*& Out, void* EvalData);
+	static UE_API void EvaluateDefault(float*& Out, void* EvalData);
 
-	static void EvaluateColor(const int DataIdx, float*& In, FVector4f& Out, void* EvalData);
-	static void EvaluateChannel(const int DataIdx, float*& In, float& Out, void* EvalData);
+	static UE_API void EvaluateColor(const int DataIdx, float*& In, FVector4f& Out, void* EvalData);
+	static UE_API void EvaluateChannel(const int DataIdx, float*& In, float& Out, void* EvalData);
 
 protected:
 	// Cached data
@@ -54,9 +56,10 @@ protected:
 	FInterval1f CachedRange;
 	
 private:
-	float SampleFunction(const FCorrespondenceSample& SampleData) const;
+	UE_API float SampleFunction(const FCorrespondenceSample& SampleData) const;
 };
 
 } // end namespace UE::Geometry
 } // end namespace UE
 
+#undef UE_API
diff --git a/Engine/Plugins/Runtime/GeometryProcessing/Source/DynamicMesh/Public/Sampling/MeshMapBaker.h b/Engine/Plugins/Runtime/GeometryProcessing/Source/DynamicMesh/Public/Sampling/MeshMapBaker.h
--- a/Engine/Plugins/Runtime/GeometryProcessing/Source/DynamicMesh/Public/Sampling/MeshMapBaker.h
+++ b/Engine/Plugins/Runtime/GeometryProcessing/Source/DynamicMesh/Public/Sampling/MeshMapBaker.h
@@ -12,6 +12,8 @@
 #include "Image/BoxFilter.h"
 #include "Image/BCSplineFilter.h"
 
+#define UE_API DYNAMICMESH_API
+
 namespace UE
 {
 namespace Geometry
@@ -21,7 +23,7 @@ class FImageOccupancyMap;
 class FImageTile;
 class FMeshMapTileBuffer;
 
-class DYNAMICMESH_API FMeshMapBaker : public FMeshBaseBaker
+class FMeshMapBaker : public FMeshBaseBaker
 {
 public:
 
@@ -30,22 +32,22 @@ public:
 	//
 
 	/** Process all bakers to generate image results for each. */
-	void Bake();
+	UE_API void Bake();
 
 	/** Add a baker to be processed. */
-	int32 AddEvaluator(const TSharedPtr<FMeshMapEvaluator, ESPMode::ThreadSafe>& Eval);
+	UE_API int32 AddEvaluator(const TSharedPtr<FMeshMapEvaluator, ESPMode::ThreadSafe>& Eval);
 
 	/** @return the evaluator at the given index. */
-	FMeshMapEvaluator* GetEvaluator(int32 EvalIdx) const;
+	UE_API FMeshMapEvaluator* GetEvaluator(int32 EvalIdx) const;
 
 	/** @return the number of bake evaluators on this baker. */
-	int32 NumEvaluators() const;
+	UE_API int32 NumEvaluators() const;
 
 	/** Reset the list of bakers. */
-	void Reset();
+	UE_API void Reset();
 
 	/** @return the bake result image for a given baker index. */
-	const TArrayView<TUniquePtr<TImageBuilder<FVector4f>>> GetBakeResults(int32 EvalIdx);
+	UE_API const TArrayView<TUniquePtr<TImageBuilder<FVector4f>>> GetBakeResults(int32 EvalIdx);
 
 	/** @return true if we should abort calculation */
 	TFunction<bool(void)> CancelF = []() { return false; };
@@ -81,12 +83,12 @@ public:
 		MitchellNetravali
 	};
 	
-	void SetDimensions(FImageDimensions DimensionsIn);
-	void SetGutterEnabled(bool bEnabled);
-	void SetGutterSize(int32 GutterSizeIn);
-	void SetSamplesPerPixel(int32 SamplesPerPixelIn);
-	void SetFilter(EBakeFilterType FilterTypeIn);
-	void SetTileSize(int TileSizeIn);
+	UE_API void SetDimensions(FImageDimensions DimensionsIn);
+	UE_API void SetGutterEnabled(bool bEnabled);
+	UE_API void SetGutterSize(int32 GutterSizeIn);
+	UE_API void SetSamplesPerPixel(int32 SamplesPerPixelIn);
+	UE_API void SetFilter(EBakeFilterType FilterTypeIn);
+	UE_API void SetTileSize(int TileSizeIn);
 
 	FImageDimensions GetDimensions() const { return Dimensions; }
 	bool GetGutterEnabled() const { return bGutterEnabled; }
@@ -103,7 +105,7 @@ public:
 	 * @param Mesh the mesh to compute UV charts.
 	 * @param MeshUVCharts the triangle ID to UV Chart ID array.
 	 */
-	static void ComputeUVCharts(const FDynamicMesh3& Mesh, TArray<int32>& MeshUVCharts);
+	static UE_API void ComputeUVCharts(const FDynamicMesh3& Mesh, TArray<int32>& MeshUVCharts);
 
 	/**
 	 * Set an a Triangle ID to UV Chart ID array for TargetMesh.
@@ -149,7 +151,7 @@ public:
 
 protected:
 	/** Evaluate this sample. */
-	void BakeSample(
+	UE_API void BakeSample(
 		FMeshMapTileBuffer& TileBuffer,
 		const FMeshMapEvaluator::FCorrespondenceSample& Sample,
 		const FVector2d& QueryUVPosition,
@@ -157,13 +159,13 @@ protected:
 		const FImageOccupancyMap& OccupancyMap);
 
 	/** Initialize evaluation contexts and precompute data for bake evaluation. */
-	void InitBake();
+	UE_API void InitBake();
 
 	/** Initialize bake sample default floats and colors. */
-	void InitBakeDefaults();
+	UE_API void InitBakeDefaults();
 
 	/** Initialize filter */
-	void InitFilter();
+	UE_API void InitFilter();
 
 protected:
 	const bool bParallel = true;
@@ -204,9 +206,9 @@ protected:
 	EBakeFilterType FilterType = EBakeFilterType::BSpline;
 
 	/** Texture filters */
-	static FBoxFilter BoxFilter;
-	static FBSplineFilter BSplineFilter;
-	static FMitchellNetravaliFilter MitchellNetravaliFilter;
+	static UE_API FBoxFilter BoxFilter;
+	static UE_API FBSplineFilter BSplineFilter;
+	static UE_API FMitchellNetravaliFilter MitchellNetravaliFilter;
 
 	/** Texture filter function */
 	using TextureFilterFn = float(*)(const FVector2d& Dist);
@@ -262,3 +264,5 @@ protected:
 
 } // end namespace UE::Geometry
 } // end namespace UE
+
+#undef UE_API
diff --git a/Engine/Plugins/Runtime/GeometryProcessing/Source/DynamicMesh/Public/Sampling/MeshNormalMapBaker.h b/Engine/Plugins/Runtime/GeometryProcessing/Source/DynamicMesh/Public/Sampling/MeshNormalMapBaker.h
--- a/Engine/Plugins/Runtime/GeometryProcessing/Source/DynamicMesh/Public/Sampling/MeshNormalMapBaker.h
+++ b/Engine/Plugins/Runtime/GeometryProcessing/Source/DynamicMesh/Public/Sampling/MeshNormalMapBaker.h
@@ -8,12 +8,14 @@
 #include "Image/ImageBuilder.h"
 #include "DynamicMesh/MeshTangents.h"
 
+#define UE_API DYNAMICMESH_API
+
 namespace UE
 {
 namespace Geometry
 {
 
-class DYNAMICMESH_API FMeshNormalMapBaker : public FMeshImageBaker
+class FMeshNormalMapBaker : public FMeshImageBaker
 {
 public:
 	virtual ~FMeshNormalMapBaker() {}
@@ -30,7 +32,7 @@ public:
 	// Compute functions
 	//
 
-	virtual void Bake() override;
+	UE_API virtual void Bake() override;
 
 	//
 	// Output
@@ -48,3 +50,5 @@ protected:
 
 } // end namespace UE::Geometry
 } // end namespace UE
+
+#undef UE_API
diff --git a/Engine/Plugins/Runtime/GeometryProcessing/Source/DynamicMesh/Public/Sampling/MeshNormalMapEvaluator.h b/Engine/Plugins/Runtime/GeometryProcessing/Source/DynamicMesh/Public/Sampling/MeshNormalMapEvaluator.h
--- a/Engine/Plugins/Runtime/GeometryProcessing/Source/DynamicMesh/Public/Sampling/MeshNormalMapEvaluator.h
+++ b/Engine/Plugins/Runtime/GeometryProcessing/Source/DynamicMesh/Public/Sampling/MeshNormalMapEvaluator.h
@@ -8,6 +8,8 @@
 #include "Sampling/MeshBakerCommon.h"
 #include "DynamicMesh/MeshTangents.h"
 
+#define UE_API DYNAMICMESH_API
+
 namespace UE
 {
 namespace Geometry
@@ -16,19 +18,19 @@ namespace Geometry
 /**
  * A mesh evaluator for tangent space normals.
  */
-class DYNAMICMESH_API FMeshNormalMapEvaluator : public FMeshMapEvaluator
+class FMeshNormalMapEvaluator : public FMeshMapEvaluator
 {
 public:
 	PRAGMA_DISABLE_DEPRECATION_WARNINGS
 	FMeshNormalMapEvaluator() = default;
 	FMeshNormalMapEvaluator(const FMeshNormalMapEvaluator&) = default;
 	FMeshNormalMapEvaluator& operator=( const FMeshNormalMapEvaluator& ) = default;
-	PRAGMA_ENABLE_DEPRECATION_WARNINGS
+	UE_API PRAGMA_ENABLE_DEPRECATION_WARNINGS
 
 	// Begin FMeshMapEvaluator interface
 	virtual void Setup(const FMeshBaseBaker& Baker, FEvaluationContext& Context) override;
 
-	virtual const TArray<EComponents>& DataLayout() const override;
+	UE_API virtual const TArray<EComponents>& DataLayout() const override;
 
 	virtual EMeshMapEvaluatorType Type() const override { return EMeshMapEvaluatorType::Normal; }
 	// End FMeshMapEvaluator interface
@@ -36,10 +38,10 @@ public:
 	template <bool bUseDetailNormalMap>
 	static void EvaluateSample(float*& Out, const FCorrespondenceSample& Sample, void* EvalData);
 
-	static void EvaluateDefault(float*& Out, void* EvalData);
+	static UE_API void EvaluateDefault(float*& Out, void* EvalData);
 
-	static void EvaluateColor(const int DataIdx, float*& In, FVector4f& Out, void* EvalData);
-	static void EvaluateChannel(const int DataIdx, float*& In, float& Out, void* EvalData);
+	static UE_API void EvaluateColor(const int DataIdx, float*& In, FVector4f& Out, void* EvalData);
+	static UE_API void EvaluateChannel(const int DataIdx, float*& In, float& Out, void* EvalData);
 
 protected:
 	// Cached data
@@ -62,3 +64,4 @@ private:
 } // end namespace UE::Geometry
 } // end namespace UE
 
+#undef UE_API
diff --git a/Engine/Plugins/Runtime/GeometryProcessing/Source/DynamicMesh/Public/Sampling/MeshOcclusionMapEvaluator.h b/Engine/Plugins/Runtime/GeometryProcessing/Source/DynamicMesh/Public/Sampling/MeshOcclusionMapEvaluator.h
--- a/Engine/Plugins/Runtime/GeometryProcessing/Source/DynamicMesh/Public/Sampling/MeshOcclusionMapEvaluator.h
+++ b/Engine/Plugins/Runtime/GeometryProcessing/Source/DynamicMesh/Public/Sampling/MeshOcclusionMapEvaluator.h
@@ -7,6 +7,8 @@
 #include "Misc/EnumClassFlags.h"
 #include "DynamicMesh/MeshTangents.h"
 
+#define UE_API DYNAMICMESH_API
+
 namespace UE
 {
 namespace Geometry
@@ -26,7 +28,7 @@ ENUM_CLASS_FLAGS(EMeshOcclusionMapType);
 /**
  * A mesh evaluator for occlusion data (Ambient Occlusion & Bent Normals).
  */
-class DYNAMICMESH_API FMeshOcclusionMapEvaluator : public FMeshMapEvaluator
+class FMeshOcclusionMapEvaluator : public FMeshMapEvaluator
 {
 public:
 	enum class EDistribution
@@ -57,9 +59,9 @@ public:
 
 public:
 	// Begin FMeshMapEvaluator interface
-	virtual void Setup(const FMeshBaseBaker& Baker, FEvaluationContext& Context) override;
+	UE_API virtual void Setup(const FMeshBaseBaker& Baker, FEvaluationContext& Context) override;
 
-	virtual const TArray<EComponents>& DataLayout() const override;
+	UE_API virtual const TArray<EComponents>& DataLayout() const override;
 
 	virtual EMeshMapEvaluatorType Type() const override { return EMeshMapEvaluatorType::Occlusion; }
 	// End FMeshMapEvaluator interface
@@ -83,8 +85,8 @@ protected:
 	double BiasDotThreshold = 0.25;
 	TArray<FVector3d> RayDirections;
 
-	static const FVector3d DefaultTangentNormal;
-	static const FVector3d DefaultObjectNormal;
+	static UE_API const FVector3d DefaultTangentNormal;
+	static UE_API const FVector3d DefaultObjectNormal;
 
 private:
 	template <EMeshOcclusionMapType OcclusionType, ESpace NormalSpace>
@@ -93,3 +95,5 @@ private:
 
 } // end namespace UE::Geometry
 } // end namespace UE
+
+#undef UE_API
diff --git a/Engine/Plugins/Runtime/GeometryProcessing/Source/DynamicMesh/Public/Sampling/MeshPropertyMapBaker.h b/Engine/Plugins/Runtime/GeometryProcessing/Source/DynamicMesh/Public/Sampling/MeshPropertyMapBaker.h
--- a/Engine/Plugins/Runtime/GeometryProcessing/Source/DynamicMesh/Public/Sampling/MeshPropertyMapBaker.h
+++ b/Engine/Plugins/Runtime/GeometryProcessing/Source/DynamicMesh/Public/Sampling/MeshPropertyMapBaker.h
@@ -7,6 +7,8 @@
 #include "Sampling/MeshImageBakingCache.h"
 #include "Image/ImageBuilder.h"
 
+#define UE_API DYNAMICMESH_API
+
 namespace UE
 {
 namespace Geometry
@@ -22,7 +24,7 @@ enum class EMeshPropertyBakeType
 };
 
 
-class DYNAMICMESH_API FMeshPropertyMapBaker : public FMeshImageBaker
+class FMeshPropertyMapBaker : public FMeshImageBaker
 {
 public:
 	virtual ~FMeshPropertyMapBaker() {}
@@ -37,7 +39,7 @@ public:
 	// Compute functions
 	//
 
-	virtual void Bake() override;
+	UE_API virtual void Bake() override;
 
 	//
 	// Output
@@ -53,3 +55,5 @@ protected:
 
 } // end namespace UE::Geometry
 } // end namespace UE
+
+#undef UE_API
diff --git a/Engine/Plugins/Runtime/GeometryProcessing/Source/DynamicMesh/Public/Sampling/MeshPropertyMapEvaluator.h b/Engine/Plugins/Runtime/GeometryProcessing/Source/DynamicMesh/Public/Sampling/MeshPropertyMapEvaluator.h
--- a/Engine/Plugins/Runtime/GeometryProcessing/Source/DynamicMesh/Public/Sampling/MeshPropertyMapEvaluator.h
+++ b/Engine/Plugins/Runtime/GeometryProcessing/Source/DynamicMesh/Public/Sampling/MeshPropertyMapEvaluator.h
@@ -6,6 +6,8 @@
 #include "Sampling/MeshBakerCommon.h"
 #include "Image/ImageBuilder.h"
 
+#define UE_API DYNAMICMESH_API
+
 namespace UE
 {
 namespace Geometry
@@ -27,7 +29,7 @@ enum class EMeshPropertyMapType
 /**
  * A mesh evaluator for mesh properties as color data.
  */
-class DYNAMICMESH_API FMeshPropertyMapEvaluator : public FMeshMapEvaluator
+class FMeshPropertyMapEvaluator : public FMeshMapEvaluator
 {
 public:
 	EMeshPropertyMapType Property = EMeshPropertyMapType::Normal;
@@ -37,12 +39,12 @@ public:
 	FMeshPropertyMapEvaluator() = default;
 	FMeshPropertyMapEvaluator(const FMeshPropertyMapEvaluator&) = default;
 	FMeshPropertyMapEvaluator& operator=(const FMeshPropertyMapEvaluator&) = default;
-	PRAGMA_ENABLE_DEPRECATION_WARNINGS
+	UE_API PRAGMA_ENABLE_DEPRECATION_WARNINGS
 
 	// Begin FMeshMapEvaluator interface
 	virtual void Setup(const FMeshBaseBaker& Baker, FEvaluationContext& Context) override;
 
-	virtual const TArray<EComponents>& DataLayout() const override;
+	UE_API virtual const TArray<EComponents>& DataLayout() const override;
 
 	virtual EMeshMapEvaluatorType Type() const override { return EMeshMapEvaluatorType::Property; }
 	// End FMeshMapEvaluator interface
@@ -50,13 +52,13 @@ public:
 	template <bool bUseDetailNormalMap>
 	static void EvaluateSample(float*& Out, const FCorrespondenceSample& Sample, void* EvalData);
 
-	static void EvaluateDefault(float*& Out, void* EvalData);
+	static UE_API void EvaluateDefault(float*& Out, void* EvalData);
 
-	static void EvaluateColor(const int DataIdx, float*& In, FVector4f& Out, void* EvalData);
-	static void EvaluateChannel(const int DataIdx, float*& In, float& Out, void* EvalData);
+	static UE_API void EvaluateColor(const int DataIdx, float*& In, FVector4f& Out, void* EvalData);
+	static UE_API void EvaluateChannel(const int DataIdx, float*& In, float& Out, void* EvalData);
 
 protected:
-	FVector4f GetDefaultValue(EMeshPropertyMapType InProperty) const;
+	UE_API FVector4f GetDefaultValue(EMeshPropertyMapType InProperty) const;
 	
 	// Cached data
 	const IMeshBakerDetailSampler* DetailSampler = nullptr;
@@ -101,3 +103,4 @@ private:
 } // end namespace UE::Geometry
 } // end namespace UE
 
+#undef UE_API
diff --git a/Engine/Plugins/Runtime/GeometryProcessing/Source/DynamicMesh/Public/Sampling/MeshResampleImageBaker.h b/Engine/Plugins/Runtime/GeometryProcessing/Source/DynamicMesh/Public/Sampling/MeshResampleImageBaker.h
--- a/Engine/Plugins/Runtime/GeometryProcessing/Source/DynamicMesh/Public/Sampling/MeshResampleImageBaker.h
+++ b/Engine/Plugins/Runtime/GeometryProcessing/Source/DynamicMesh/Public/Sampling/MeshResampleImageBaker.h
@@ -7,12 +7,14 @@
 #include "Sampling/MeshImageBakingCache.h"
 #include "Image/ImageBuilder.h"
 
+#define UE_API DYNAMICMESH_API
+
 namespace UE
 {
 namespace Geometry
 {
 
-class DYNAMICMESH_API FMeshResampleImageBaker : public FMeshImageBaker
+class FMeshResampleImageBaker : public FMeshImageBaker
 {
 public:
 	virtual ~FMeshResampleImageBaker() {}
@@ -30,7 +32,7 @@ public:
 	// Compute functions
 	//
 
-	virtual void Bake() override;
+	UE_API virtual void Bake() override;
 
 	//
 	// Output
@@ -46,19 +48,21 @@ protected:
 
 
 
-class DYNAMICMESH_API FMeshMultiResampleImageBaker : public FMeshResampleImageBaker
+class FMeshMultiResampleImageBaker : public FMeshResampleImageBaker
 {
 public:
 
 	TMap<int32, TSharedPtr<UE::Geometry::TImageBuilder<FVector4f>, ESPMode::ThreadSafe>> MultiTextures;
 
-	virtual void Bake() override;
+	UE_API virtual void Bake() override;
 
 protected:
 
-	void InitResult();
-	void BakeMaterial(int32 MaterialID);
+	UE_API void InitResult();
+	UE_API void BakeMaterial(int32 MaterialID);
 };
 
 } // end namespace UE::Geometry
 } // end namespace UE
+
+#undef UE_API
diff --git a/Engine/Plugins/Runtime/GeometryProcessing/Source/DynamicMesh/Public/Sampling/MeshResampleImageEvaluator.h b/Engine/Plugins/Runtime/GeometryProcessing/Source/DynamicMesh/Public/Sampling/MeshResampleImageEvaluator.h
--- a/Engine/Plugins/Runtime/GeometryProcessing/Source/DynamicMesh/Public/Sampling/MeshResampleImageEvaluator.h
+++ b/Engine/Plugins/Runtime/GeometryProcessing/Source/DynamicMesh/Public/Sampling/MeshResampleImageEvaluator.h
@@ -6,6 +6,8 @@
 #include "Sampling/MeshBaseBaker.h"
 #include "Image/ImageBuilder.h"
 
+#define UE_API DYNAMICMESH_API
+
 namespace UE
 {
 namespace Geometry
@@ -16,33 +18,33 @@ class IMeshBakerDetailSampler;
 /**
  * A mesh evaluator for sampling 2D texture data. 
  */
-class DYNAMICMESH_API FMeshResampleImageEvaluator : public FMeshMapEvaluator
+class FMeshResampleImageEvaluator : public FMeshMapEvaluator
 {
 public:
 	FVector4f DefaultColor = FVector4f(0, 0, 0, 1);
 
 public:
 	// Begin FMeshMapEvaluator interface
-	virtual void Setup(const FMeshBaseBaker& Baker, FEvaluationContext& Context) override;
+	UE_API virtual void Setup(const FMeshBaseBaker& Baker, FEvaluationContext& Context) override;
 
-	virtual const TArray<EComponents>& DataLayout() const override;
+	UE_API virtual const TArray<EComponents>& DataLayout() const override;
 
 	virtual EMeshMapEvaluatorType Type() const override { return EMeshMapEvaluatorType::ResampleImage; }
 	// End FMeshMapEvaluator interface
 
-	static void EvaluateSample(float*& Out, const FCorrespondenceSample& Sample, void* EvalData);
+	static UE_API void EvaluateSample(float*& Out, const FCorrespondenceSample& Sample, void* EvalData);
 
-	static void EvaluateDefault(float*& Out, void* EvalData);
+	static UE_API void EvaluateDefault(float*& Out, void* EvalData);
 
-	static void EvaluateColor(const int DataIdx, float*& In, FVector4f& Out, void* EvalData);
-	static void EvaluateChannel(const int DataIdx, float*& In, float& Out, void* EvalData);
+	static UE_API void EvaluateColor(const int DataIdx, float*& In, FVector4f& Out, void* EvalData);
+	static UE_API void EvaluateChannel(const int DataIdx, float*& In, float& Out, void* EvalData);
 
 protected:
 	// Cached data
 	const IMeshBakerDetailSampler* DetailSampler = nullptr;
 
 private:
-	FVector4f ImageSampleFunction(const FCorrespondenceSample& Sample) const;
+	UE_API FVector4f ImageSampleFunction(const FCorrespondenceSample& Sample) const;
 };
 
 
@@ -50,7 +52,7 @@ private:
 /**
  * A mesh evaluator for sampling multiple 2D textures by material ID
  */
-class DYNAMICMESH_API FMeshMultiResampleImageEvaluator : public FMeshResampleImageEvaluator
+class FMeshMultiResampleImageEvaluator : public FMeshResampleImageEvaluator
 {
 public:
 	/** List of textures indexed by material ID. Entries can be null. */
@@ -61,15 +63,15 @@ public:
 
 public:
 	// Begin FMeshMapEvaluator interface
-	virtual void Setup(const FMeshBaseBaker& Baker, FEvaluationContext& Context) override;
+	UE_API virtual void Setup(const FMeshBaseBaker& Baker, FEvaluationContext& Context) override;
 
 	virtual EMeshMapEvaluatorType Type() const override { return EMeshMapEvaluatorType::MultiResampleImage; }
 	// End FMeshMapEvaluator interface
 
-	static void EvaluateSampleMulti(float*& Out, const FCorrespondenceSample& Sample, void* EvalData);
+	static UE_API void EvaluateSampleMulti(float*& Out, const FCorrespondenceSample& Sample, void* EvalData);
 
 private:
-	FVector4f ImageSampleFunction(const FCorrespondenceSample& Sample) const;
+	UE_API FVector4f ImageSampleFunction(const FCorrespondenceSample& Sample) const;
 
 	// Cached data
 	int NumMultiTextures = 0;
@@ -78,3 +80,4 @@ private:
 } // end namespace UE::Geometry
 } // end namespace UE
 
+#undef UE_API
diff --git a/Engine/Plugins/Runtime/GeometryProcessing/Source/DynamicMesh/Public/Sampling/MeshSurfaceSampler.h b/Engine/Plugins/Runtime/GeometryProcessing/Source/DynamicMesh/Public/Sampling/MeshSurfaceSampler.h
--- a/Engine/Plugins/Runtime/GeometryProcessing/Source/DynamicMesh/Public/Sampling/MeshSurfaceSampler.h
+++ b/Engine/Plugins/Runtime/GeometryProcessing/Source/DynamicMesh/Public/Sampling/MeshSurfaceSampler.h
@@ -6,6 +6,8 @@
 #include "DynamicMesh/DynamicMeshAttributeSet.h"
 #include "DynamicMesh/DynamicMeshAABBTree3.h"
 
+#define UE_API DYNAMICMESH_API
+
 namespace UE
 {
 namespace Geometry
@@ -138,20 +140,20 @@ public:
 	 * @param UVOverlay UV overlay of Mesh
 	 * @param QueryType Type of query functions we will call. Some queries are simpler and do not require spatial data structures
 	 */
-	virtual void Initialize(const FDynamicMesh3* Mesh, const FDynamicMeshUVOverlay* UVOverlay, EMeshSurfaceSamplerQueryType QueryType);
+	UE_API virtual void Initialize(const FDynamicMesh3* Mesh, const FDynamicMeshUVOverlay* UVOverlay, EMeshSurfaceSamplerQueryType QueryType);
 
 	/**
 	 * Fill in the given SampleInfo for the given UV location.
 	 * @pre The Initialize function must have been called with EMeshSurfaceSamplerQueryType::UVOnly
 	 * @return If the query location is not on a triangle in the Mesh then return false (SampleInfo will be invalid)
 	 */
-	virtual bool QuerySampleInfo(const FVector2d& UV, FMeshUVSampleInfo& SampleInfo);
+	UE_API virtual bool QuerySampleInfo(const FVector2d& UV, FMeshUVSampleInfo& SampleInfo);
 
 	/**
 	 * Fill in the given SampleInfo for the given UV location in the given UVTriangleID.
 	 * @return If the query location is not on a triangle in the Mesh then return false (SampleInfo will be invalid)
 	 */
-	virtual bool QuerySampleInfo(int32 UVTriangleID, const FVector2d& UV, FMeshUVSampleInfo& SampleInfo);
+	UE_API virtual bool QuerySampleInfo(int32 UVTriangleID, const FVector2d& UV, FMeshUVSampleInfo& SampleInfo);
 
 protected:
 
@@ -334,3 +336,4 @@ bool TMeshSurfaceUVSampler<SampleType>::CachedSampleUV(const FMeshUVSampleInfo&
 } // end namespace UE::Geometry
 } // end namespace UE
 
+#undef UE_API 
\ No newline at end of file
diff --git a/Engine/Plugins/Runtime/GeometryProcessing/Source/DynamicMesh/Public/Sampling/MeshUVShellMapEvaluator.h b/Engine/Plugins/Runtime/GeometryProcessing/Source/DynamicMesh/Public/Sampling/MeshUVShellMapEvaluator.h
--- a/Engine/Plugins/Runtime/GeometryProcessing/Source/DynamicMesh/Public/Sampling/MeshUVShellMapEvaluator.h
+++ b/Engine/Plugins/Runtime/GeometryProcessing/Source/DynamicMesh/Public/Sampling/MeshUVShellMapEvaluator.h
@@ -5,6 +5,8 @@
 #include "Sampling/MeshMapEvaluator.h"
 #include "Sampling/MeshBakerCommon.h"
 
+#define UE_API DYNAMICMESH_API
+
 namespace UE::Geometry
 {
 
@@ -13,7 +15,7 @@ class IMeshBakerDetailSampler;
 /**
  * A mesh evaluator for mesh properties as color data.
  */
-class DYNAMICMESH_API FMeshUVShellMapEvaluator : public FMeshMapEvaluator
+class FMeshUVShellMapEvaluator : public FMeshMapEvaluator
 {
 public:
 	FVector2d TexelSize = FVector2d::One();
@@ -29,30 +31,31 @@ public:
 	FMeshUVShellMapEvaluator() = default;
 	FMeshUVShellMapEvaluator(const FMeshUVShellMapEvaluator&) = default;
 	FMeshUVShellMapEvaluator& operator=(const FMeshUVShellMapEvaluator&) = default;
-	PRAGMA_ENABLE_DEPRECATION_WARNINGS
+	UE_API PRAGMA_ENABLE_DEPRECATION_WARNINGS
 
 	// Begin FMeshMapEvaluator interface
 	virtual void Setup(const FMeshBaseBaker& Baker, FEvaluationContext& Context) override;
 
-	virtual const TArray<EComponents>& DataLayout() const override;
+	UE_API virtual const TArray<EComponents>& DataLayout() const override;
 
 	virtual EMeshMapEvaluatorType Type() const override { return EMeshMapEvaluatorType::UVShell; }
 	// End FMeshMapEvaluator interface
 
-	static void EvaluateSample(float*& Out, const FCorrespondenceSample& Sample, void* EvalData);
+	static UE_API void EvaluateSample(float*& Out, const FCorrespondenceSample& Sample, void* EvalData);
 
-	static void EvaluateDefault(float*& Out, void* EvalData);
+	static UE_API void EvaluateDefault(float*& Out, void* EvalData);
 
-	static void EvaluateColor(const int DataIdx, float*& In, FVector4f& Out, void* EvalData);
-	static void EvaluateChannel(const int DataIdx, float*& In, float& Out, void* EvalData);
+	static UE_API void EvaluateColor(const int DataIdx, float*& In, FVector4f& Out, void* EvalData);
+	static UE_API void EvaluateChannel(const int DataIdx, float*& In, float& Out, void* EvalData);
 
 protected:
 	// Cached data
 	const IMeshBakerDetailSampler* DetailSampler = nullptr;
 
 private:
-	FVector4f SampleFunction(const FCorrespondenceSample& SampleData) const;
+	UE_API FVector4f SampleFunction(const FCorrespondenceSample& SampleData) const;
 };
 
 } // end namespace UE::Geometry
 
+#undef UE_API
diff --git a/Engine/Plugins/Runtime/GeometryProcessing/Source/DynamicMesh/Public/Sampling/MeshVertexBaker.h b/Engine/Plugins/Runtime/GeometryProcessing/Source/DynamicMesh/Public/Sampling/MeshVertexBaker.h
--- a/Engine/Plugins/Runtime/GeometryProcessing/Source/DynamicMesh/Public/Sampling/MeshVertexBaker.h
+++ b/Engine/Plugins/Runtime/GeometryProcessing/Source/DynamicMesh/Public/Sampling/MeshVertexBaker.h
@@ -8,12 +8,14 @@
 #include "Image/ImageBuilder.h"
 #include "Image/ImageDimensions.h"
 
+#define UE_API DYNAMICMESH_API
+
 namespace UE
 {
 namespace Geometry
 {
 
-class DYNAMICMESH_API FMeshVertexBaker : public FMeshBaseBaker
+class FMeshVertexBaker : public FMeshBaseBaker
 {
 public:
 	enum class EBakeMode
@@ -41,10 +43,10 @@ public:
 	//
 
 	/** Process all bakers to generate the image result. */
-	void Bake();
+	UE_API void Bake();
 
 	/** @return the bake result image. */
-	const TImageBuilder<FVector4f>* GetBakeResult() const;
+	UE_API const TImageBuilder<FVector4f>* GetBakeResult() const;
 
 	/** if this function returns true, we should abort calculation */
 	TFunction<bool(void)> CancelF = []() { return false; };
@@ -84,9 +86,11 @@ protected:
 	BakeFn BakeInternal = nullptr;
 
 	/** Constant evaluator for empty channels. */
-	static FMeshConstantMapEvaluator ZeroEvaluator;
-	static FMeshConstantMapEvaluator OneEvaluator;
+	static UE_API FMeshConstantMapEvaluator ZeroEvaluator;
+	static UE_API FMeshConstantMapEvaluator OneEvaluator;
 };
 
 } // end namespace UE::Geometry
-} // end namespace UE
\ No newline at end of file
+} // end namespace UE
+
+#undef UE_API
diff --git a/Engine/Plugins/Runtime/GeometryProcessing/Source/DynamicMesh/Public/Selections/GeometrySelection.h b/Engine/Plugins/Runtime/GeometryProcessing/Source/DynamicMesh/Public/Selections/GeometrySelection.h
--- a/Engine/Plugins/Runtime/GeometryProcessing/Source/DynamicMesh/Public/Selections/GeometrySelection.h
+++ b/Engine/Plugins/Runtime/GeometryProcessing/Source/DynamicMesh/Public/Selections/GeometrySelection.h
@@ -8,6 +8,8 @@
 #include "BoxTypes.h"
 #include "DynamicMesh/InfoTypes.h"
 
+#define UE_API DYNAMICMESH_API
+
 
 namespace UE
 {
@@ -26,7 +28,7 @@ namespace Geometry
  * Note that since the ints are unsigned, IndexConstants::InvalidID
  * is not directly representable (-1 will become positive 0xFFFFFFFF). 
  */
-struct DYNAMICMESH_API FGeoSelectionID
+struct FGeoSelectionID
 {
 	/** Topology ID, stored in upper 32 bits when packed into 64-bits */
 	uint32 TopologyID;
@@ -136,7 +138,7 @@ enum class EGeometryTopologyType
  * for efficient unique adds and removals. No assumptions are made about the values,
  * they could be (eg) mesh indices, IDs of some type, or even pointer values.
  */
-struct DYNAMICMESH_API FGeometrySelection
+struct FGeometrySelection
 {
 	/** Type of geometric element represented by this selection, if applicable */
 	EGeometryElementType ElementType = EGeometryElementType::Face;
@@ -203,7 +205,7 @@ struct DYNAMICMESH_API FGeometrySelection
  * FGeometrySelectionHitQueryConfig defines the desired settings for a "hit query" on
  * selected meshes/objects. 
  */
-struct DYNAMICMESH_API FGeometrySelectionHitQueryConfig
+struct FGeometrySelectionHitQueryConfig
 {
 	/** Type of object topology to query */
 	EGeometryTopologyType TopologyType = EGeometryTopologyType::Polygroup;
@@ -241,7 +243,7 @@ enum class EGeometrySelectionChangeType
  * to indicate what type of change should be applied to a FGeometrySelection
  * (based on additional parameters, generally)
  */
-struct DYNAMICMESH_API FGeometrySelectionUpdateConfig
+struct FGeometrySelectionUpdateConfig
 {
 	EGeometrySelectionChangeType ChangeType = EGeometrySelectionChangeType::Add;
 };
@@ -254,7 +256,7 @@ struct DYNAMICMESH_API FGeometrySelectionUpdateConfig
  * 
  * (Currently there is no way to swap the order)
  */
-struct DYNAMICMESH_API FGeometrySelectionDelta
+struct FGeometrySelectionDelta
 {
 	/**
 	 * Elements removed from a FGeometrySelection during some selection edit
@@ -274,7 +276,7 @@ struct DYNAMICMESH_API FGeometrySelectionDelta
 /**
  * 3D Bounding information for a FGeometrySelection
  */
-struct DYNAMICMESH_API FGeometrySelectionBounds
+struct FGeometrySelectionBounds
 {
 	UE::Geometry::FAxisAlignedBox3d WorldBounds = UE::Geometry::FAxisAlignedBox3d::Empty();
 };
@@ -284,7 +286,7 @@ struct DYNAMICMESH_API FGeometrySelectionBounds
  * 3D Geometry representing a FGeometrySelection, for example
  * suitable for passing to rendering code, etc
  */
-struct DYNAMICMESH_API FGeometrySelectionElements
+struct FGeometrySelectionElements
 {
 	TArray<UE::Geometry::FTriangle3d> Triangles;
 	TArray<UE::Geometry::FSegment3d> Segments;
@@ -299,7 +301,7 @@ struct DYNAMICMESH_API FGeometrySelectionElements
 };
 
 
-struct DYNAMICMESH_API FGeometrySelectionUpdateResult
+struct FGeometrySelectionUpdateResult
 {
 	bool bSelectionMissed = false;
 	bool bSelectionModified = false;
@@ -326,14 +328,14 @@ struct DYNAMICMESH_API FGeometrySelectionUpdateResult
  * to determine whether an ID is selected, rather than Selection.Contains()
  * 
  */
-class DYNAMICMESH_API FGeometrySelectionEditor
+class FGeometrySelectionEditor
 {
 public:
 	/**
 	 * Initialize the Editor with the given Selection. The
 	 * TargetSelectionIn must live longer than the FGeometrySelectionEditor
 	 */
-	void Initialize(
+	UE_API void Initialize(
 		FGeometrySelection* TargetSelectionIn,
 		const FGeometrySelectionHitQueryConfig& QueryConfigIn,
 		bool bEnableTopologyIDFiltering);
@@ -342,7 +344,7 @@ public:
 	 * Initialize the Editor with the given Selection. The
 	 * TargetSelectionIn must live longer than the FGeometrySelectionEditor
 	 */
-	void Initialize(
+	UE_API void Initialize(
 		FGeometrySelection* TargetSelectionIn,
 		bool bEnableTopologyIDFiltering);
 
@@ -361,16 +363,16 @@ public:
 	 * the ElementType or TopologyType are modified. This perhaps should be revisited as they basically
 	 * always need to be the same...
 	 */
-	void UpdateQueryConfig(const FGeometrySelectionHitQueryConfig& NewConfig, bool bEnableTopologyIDFilteringIn);
+	UE_API void UpdateQueryConfig(const FGeometrySelectionHitQueryConfig& NewConfig, bool bEnableTopologyIDFilteringIn);
 
 	/** @return true if the given ID is currently selected in the Target Selection */
-	bool IsSelected(uint64 ID) const;
+	UE_API bool IsSelected(uint64 ID) const;
 
 	/** Clear the Target Selection and return change information in DeltaOut */
-	void ClearSelection(FGeometrySelectionDelta& DeltaOut);
+	UE_API void ClearSelection(FGeometrySelectionDelta& DeltaOut);
 
 	/** Remove ID from Selection */
-	bool RemoveFromSelection(uint64 ID);
+	UE_API bool RemoveFromSelection(uint64 ID);
 
 	/** Access the Selection object this Editor is modifying */
 	const FGeometrySelection& GetSelection() const { return *TargetSelection; }
@@ -420,7 +422,7 @@ public:
 	}
 
 	/** Replace the current set of selected items with those in the NewSelection, and return change information in DeltaOut */
-	bool Replace(const FGeometrySelection& NewSelection, FGeometrySelectionDelta& DeltaOut);
+	UE_API bool Replace(const FGeometrySelection& NewSelection, FGeometrySelectionDelta& DeltaOut);
 
 
 protected:
@@ -428,7 +430,7 @@ protected:
 	FGeometrySelectionHitQueryConfig QueryConfig;
 
 	bool bEnableTopologyIDFiltering = false;
-	uint64 RemapToExistingTopologyID(uint64 ID, bool& bFoundOut) const;
+	UE_API uint64 RemapToExistingTopologyID(uint64 ID, bool& bFoundOut) const;
 };
 
 
@@ -439,7 +441,7 @@ protected:
  * only be constructed based on a parent FGeometrySelectionEditor, from which it derives it's
  * configuration, which is then how the rest of the system knows what geometry to provide as a preview/etc
  */
-class DYNAMICMESH_API FGeometrySelectionPreview : public FGeometrySelectionEditor
+class FGeometrySelectionPreview : public FGeometrySelectionEditor
 {
 public:
 	FGeometrySelection PreviewSelection;
@@ -456,3 +458,5 @@ public:
 
 } // end namespace UE::Geometry
 } // end namespace UE
+
+#undef UE_API
diff --git a/Engine/Plugins/Runtime/GeometryProcessing/Source/DynamicMesh/Public/Selections/QuadGridPatch.h b/Engine/Plugins/Runtime/GeometryProcessing/Source/DynamicMesh/Public/Selections/QuadGridPatch.h
--- a/Engine/Plugins/Runtime/GeometryProcessing/Source/DynamicMesh/Public/Selections/QuadGridPatch.h
+++ b/Engine/Plugins/Runtime/GeometryProcessing/Source/DynamicMesh/Public/Selections/QuadGridPatch.h
@@ -5,6 +5,8 @@
 #include "IndexTypes.h"
 #include "Templates/Function.h"
 
+#define UE_API DYNAMICMESH_API
+
 namespace UE
 {
 namespace Geometry
@@ -24,7 +26,7 @@ class FDynamicMesh3;
  * Generally the quad triangles will not be assumed to form any consistent pattern. This complicates some of
  * the functions however, and perhaps this assumption could be optionally provided in the future
  */
-class DYNAMICMESH_API FQuadGridPatch
+class FQuadGridPatch
 {
 public:
 	int NumVertexColsU = 0;
@@ -47,19 +49,19 @@ public:
 	/**
 	 * Initialize an Nx1 grid from a list of quads (represented as a pair of triangles) and the row of vertices on the "bottom" (ie first span)
 	 */
-	void InitializeFromQuadStrip(const FDynamicMesh3& Mesh, const TArray<FIndex2i>& SequentialQuadsIn, const TArray<int32>& FirstVertexSpan );
+	UE_API void InitializeFromQuadStrip(const FDynamicMesh3& Mesh, const TArray<FIndex2i>& SequentialQuadsIn, const TArray<int32>& FirstVertexSpan );
 
 	/**
 	 * Initialize an NxM grid from a list of quads (represented as a pair of triangles) and corresponding rows of vertices. 
 	 * For N quad-rows, there must be N+1 vertex spans. 
 	 * @return false if the input data is detected to be invalid, in this case the patch is also cleared
 	 */
-	bool InitializeFromQuadPatch(const FDynamicMesh3& Mesh, const TArray<TArray<FIndex2i>>& QuadRowsIn, const TArray<TArray<int32>>& VertexSpansIn);
+	UE_API bool InitializeFromQuadPatch(const FDynamicMesh3& Mesh, const TArray<TArray<FIndex2i>>& QuadRowsIn, const TArray<TArray<int32>>& VertexSpansIn);
 
 	/**
 	 * Reverse the order of the vertex/quad rows, ie "flip" vertically
 	 */
-	void ReverseRows();
+	UE_API void ReverseRows();
 
 	/**
 	 * Append the rows of NextPatch to the rows of the current QuadGrid.
@@ -69,14 +71,14 @@ public:
 	 * @param bChecked if true, function verifies requirements and returns false if a mismatch is detected
 	 * @return true if append succeeded
 	 */
-	bool AppendQuadPatchRows(FQuadGridPatch&& NextPatch, bool bChecked);
+	UE_API bool AppendQuadPatchRows(FQuadGridPatch&& NextPatch, bool bChecked);
 
 
 	/** @return the total number of quads in the grid */
 	int32 GetNumQuads() const { return (NumVertexRowsV-1) * (NumVertexColsU-1); }
 
 	/** @return a list of all the triangles in the grid */
-	void GetAllTriangles(TArray<int32>& AllTrianglesOut) const;
+	UE_API void GetAllTriangles(TArray<int32>& AllTrianglesOut) const;
 
 	/** 
 	 * Map a row and column index pair into a single int32. 11 Bits are available for the row, and 20 for the column 
@@ -106,7 +108,7 @@ public:
 	 * @param TriIndex the index of the triangle in the quad, only 0 and 1 are valid
 	 * @return tuple of encoded (row,column) vertex span indices, in the same ordering as the identified triangle
 	 */
-	FIndex3i GetQuadTriMappedToSpanIndices(const FDynamicMesh3& Mesh, int QuadRow, int QuadIndex, int TriIndex) const;
+	UE_API FIndex3i GetQuadTriMappedToSpanIndices(const FDynamicMesh3& Mesh, int QuadRow, int QuadIndex, int TriIndex) const;
 
 
 	/**
@@ -126,26 +128,26 @@ public:
 	/**
 	 * Return vertices in specified Column
 	 */
-	bool GetVertexColumn(int32 ColumnIndex, TArray<int32>& VerticesOut) const;
+	UE_API bool GetVertexColumn(int32 ColumnIndex, TArray<int32>& VerticesOut) const;
 
 	/**
 	 * Find column that contains vertex, or InvalidID
 	 */
-	int32 FindColumnIndex(int32 VertexID) const;
+	UE_API int32 FindColumnIndex(int32 VertexID) const;
 
 	/**
 	 * Construct a new QuadGridPatch that is a subset of the current QuadGridPatch
 	 * The values (Start|End)Quad(Row|Col) define the sub-patch, the values are inclusive,
 	 * so passing (1,2,1,2) defines a 2x2 quad patch
 	 */
-	void GetSubPatchByQuadRange(int StartQuadRow, int EndQuadRow, int StartQuadCol, int EndQuadCol, FQuadGridPatch& PatchOut ) const;
+	UE_API void GetSubPatchByQuadRange(int StartQuadRow, int EndQuadRow, int StartQuadCol, int EndQuadCol, FQuadGridPatch& PatchOut ) const;
 	
 
 	/**
 	 * Split the Columns of the QuadGridPatch into new QuadGridPatches where the columns
 	 * are grouped via the PredicateFunc. 
 	 */
-	void SplitColumnsByPredicate( 
+	UE_API void SplitColumnsByPredicate( 
 		TFunctionRef<bool(int32 QuadColumn0, int32 QuadColumn1)> PredicateFunc,
 		TArray<FQuadGridPatch>& SplitPatchesOut) const;
 	
@@ -157,7 +159,7 @@ public:
 	 * so any non-planarity of the quads is ignored.
 	 * @return unsigned opening angle in degrees in range [0,180], where 0 means quads are parallel
 	 */
-	double GetQuadOpeningAngleDeg(const FDynamicMesh3& Mesh, int32 Column0, int32 Column1, int32 Row) const;
+	UE_API double GetQuadOpeningAngleDeg(const FDynamicMesh3& Mesh, int32 Column0, int32 Column1, int32 Row) const;
 
 };
 
@@ -166,3 +168,5 @@ public:
 
 } // end namespace UE::Geometry
 } // end namespace UE
+
+#undef UE_API
diff --git a/Engine/Plugins/Runtime/GeometryProcessing/Source/DynamicMesh/Public/ShapeApproximation/MeshSimpleShapeApproximation.h b/Engine/Plugins/Runtime/GeometryProcessing/Source/DynamicMesh/Public/ShapeApproximation/MeshSimpleShapeApproximation.h
--- a/Engine/Plugins/Runtime/GeometryProcessing/Source/DynamicMesh/Public/ShapeApproximation/MeshSimpleShapeApproximation.h
+++ b/Engine/Plugins/Runtime/GeometryProcessing/Source/DynamicMesh/Public/ShapeApproximation/MeshSimpleShapeApproximation.h
@@ -6,6 +6,8 @@
 #include "DynamicMesh/DynamicMesh3.h"
 #include "ShapeApproximation/SimpleShapeSet3.h"
 
+#define UE_API DYNAMICMESH_API
+
 class FProgressCancel;
 
 namespace UE {
@@ -41,7 +43,7 @@ enum class EDetectedSimpleShapeType
  * can be identified and used directly, skipping the fitting process.
  *
  */
-class DYNAMICMESH_API FMeshSimpleShapeApproximation
+class FMeshSimpleShapeApproximation
 {
 public:
 	
@@ -112,7 +114,7 @@ public:
 	 * The TSharedPtrs are stored, rather than making a copy of the input meshes
 	 * @param bIsParallelSafe if true
 	 */
-	void InitializeSourceMeshes(const TArray<const FDynamicMesh3*>& InputMeshSet);
+	UE_API void InitializeSourceMeshes(const TArray<const FDynamicMesh3*>& InputMeshSet);
 
 
 	//
@@ -122,35 +124,35 @@ public:
 	/**
 	 * Fit containing axis-aligned boxes to each input mesh and store in ShapeSetOut
 	 */
-	void Generate_AlignedBoxes(FSimpleShapeSet3d& ShapeSetOut);
+	UE_API void Generate_AlignedBoxes(FSimpleShapeSet3d& ShapeSetOut);
 
 	/**
 	 * Fit containing minimal-volume oriented boxes to each input mesh and store in ShapeSetOut
 	 */
-	void Generate_OrientedBoxes(FSimpleShapeSet3d& ShapeSetOut, FProgressCancel* Progress = nullptr);
+	UE_API void Generate_OrientedBoxes(FSimpleShapeSet3d& ShapeSetOut, FProgressCancel* Progress = nullptr);
 
 	/**
 	 * Fit containing minimal-volume spheres to each input mesh and store in ShapeSetOut
 	 */
-	void Generate_MinimalSpheres(FSimpleShapeSet3d& ShapeSetOut);
+	UE_API void Generate_MinimalSpheres(FSimpleShapeSet3d& ShapeSetOut);
 
 	/**
 	 * Fit containing approximate-minimum-volume capsules to each input mesh and store in ShapeSetOut
 	 * @warning the capsule is fit by first fitting a line to the vertices, and then containing the points, so the fit can deviate quite a bit from a truly "minimal" capsule
 	 */
-	void Generate_Capsules(FSimpleShapeSet3d& ShapeSetOut);
+	UE_API void Generate_Capsules(FSimpleShapeSet3d& ShapeSetOut);
 
 	/**
 	 * Calculate 3D Convex Hulls for each input mesh and store in ShapeSetOut.
 	 * Each convex hull is stored as a triangle mesh, and optionally simplified if bSimplifyHulls=true
 	 */
-	void Generate_ConvexHulls(FSimpleShapeSet3d& ShapeSetOut, FProgressCancel* Progress = nullptr);
+	UE_API void Generate_ConvexHulls(FSimpleShapeSet3d& ShapeSetOut, FProgressCancel* Progress = nullptr);
 
 	/**
 	 * Calculate multiple 3D Convex Hulls for each input mesh and store in ShapeSetOut.
 	 * Each convex hull is stored as a triangle mesh, and optionally simplified if bSimplifyHulls=true
 	 */
-	void Generate_ConvexHullDecompositions(FSimpleShapeSet3d& ShapeSetOut, FProgressCancel* Progress = nullptr);
+	UE_API void Generate_ConvexHullDecompositions(FSimpleShapeSet3d& ShapeSetOut, FProgressCancel* Progress = nullptr);
 
 
 	/** Type/Mode for deciding 3D axis to use in Generate_ProjectedHulls() */
@@ -175,16 +177,16 @@ public:
 	 * are optionally simplified if bSimplifyHulls=true.
 	 * @param AxisMode which axis to use for the planar projection
 	 */
-	void Generate_ProjectedHulls(FSimpleShapeSet3d& ShapeSetOut, EProjectedHullAxisMode AxisMode);
+	UE_API void Generate_ProjectedHulls(FSimpleShapeSet3d& ShapeSetOut, EProjectedHullAxisMode AxisMode);
 
-	void Generate_LevelSets(FSimpleShapeSet3d& ShapeSetOut, FProgressCancel* Progress = nullptr);
+	UE_API void Generate_LevelSets(FSimpleShapeSet3d& ShapeSetOut, FProgressCancel* Progress = nullptr);
 	
 
 	/**
 	 * Fit containing axis-aligned box, oriented box, capsule, and sphere to each input mesh, and
 	 * store the one with smallest volume in ShapeSetOut
 	 */
-	void Generate_MinVolume(FSimpleShapeSet3d& ShapeSetOut);
+	UE_API void Generate_MinVolume(FSimpleShapeSet3d& ShapeSetOut);
 
 
 
@@ -202,11 +204,13 @@ protected:
 	TArray<FSourceMeshCache> SourceMeshCaches;
 
 
-	void DetectAndCacheSimpleShapeType(const FDynamicMesh3* SourceMesh, FSourceMeshCache& CacheOut);
-	bool GetDetectedSimpleShape(const FSourceMeshCache& Cache,
+	UE_API void DetectAndCacheSimpleShapeType(const FDynamicMesh3* SourceMesh, FSourceMeshCache& CacheOut);
+	UE_API bool GetDetectedSimpleShape(const FSourceMeshCache& Cache,
 		FSimpleShapeSet3d& ShapeSetOut, FCriticalSection& ShapeSetLock);
 };
 
 
 } // end namespace UE::Geometry
 } // end namespace UE
+
+#undef UE_API
diff --git a/Engine/Plugins/Runtime/GeometryProcessing/Source/DynamicMesh/Public/ShapeApproximation/SimpleShapeSet3.h b/Engine/Plugins/Runtime/GeometryProcessing/Source/DynamicMesh/Public/ShapeApproximation/SimpleShapeSet3.h
--- a/Engine/Plugins/Runtime/GeometryProcessing/Source/DynamicMesh/Public/ShapeApproximation/SimpleShapeSet3.h
+++ b/Engine/Plugins/Runtime/GeometryProcessing/Source/DynamicMesh/Public/ShapeApproximation/SimpleShapeSet3.h
@@ -13,6 +13,8 @@
 #include "TransformTypes.h"
 #include "TransformSequence.h"
 
+#define UE_API DYNAMICMESH_API
+
 
 namespace UE {
 namespace Geometry {
@@ -35,7 +37,7 @@ ENUM_CLASS_FLAGS(ESimpleShapeType);
 /**
  * FSphereShape is a 3D sphere
  */
-class DYNAMICMESH_API FSphereShape3d
+class FSphereShape3d
 {
 public:
 	FSphere3d Sphere;
@@ -54,7 +56,7 @@ public:
 /**
  * FBoxShape is a 3D oriented box
  */
-struct DYNAMICMESH_API FBoxShape3d
+struct FBoxShape3d
 {
 	FOrientedBox3d Box;
 
@@ -72,7 +74,7 @@ struct DYNAMICMESH_API FBoxShape3d
 /**
  * FCapsuleShape is a 3D oriented capsule/sphyl
  */
-struct DYNAMICMESH_API FCapsuleShape3d
+struct FCapsuleShape3d
 {
 	FCapsule3d Capsule;
 
@@ -90,7 +92,7 @@ struct DYNAMICMESH_API FCapsuleShape3d
 /**
  * FConvexShape is a 3D convex hull, currently stored as a triangle mesh
  */
-struct DYNAMICMESH_API FConvexShape3d
+struct FConvexShape3d
 {
 	FDynamicMesh3 Mesh;
 
@@ -113,7 +115,7 @@ struct DYNAMICMESH_API FConvexShape3d
  * FLevelSetShape is a 3D signed distance function sampled on a regular grid
  */
 
-struct DYNAMICMESH_API FLevelSetShape3d
+struct FLevelSetShape3d
 {
 	FTransform3d GridTransform;
 	FDenseGrid3f Grid;
@@ -151,7 +153,7 @@ struct DYNAMICMESH_API FLevelSetShape3d
  * FSimpleShapeSet stores a set of simple geometry shapes useful for things like collision detection/etc.
  * Various functions set-processing operations are supported.
  */
-struct DYNAMICMESH_API FSimpleShapeSet3d
+struct FSimpleShapeSet3d
 {
 	TArray<FSphereShape3d> Spheres;
 	TArray<FBoxShape3d> Boxes;
@@ -165,34 +167,34 @@ struct DYNAMICMESH_API FSimpleShapeSet3d
 	/**
 	 * Append elements of another shape set
 	 */
-	void Append(const FSimpleShapeSet3d& OtherShapeSet);
+	UE_API void Append(const FSimpleShapeSet3d& OtherShapeSet);
 
 
 	/**
 	 * Append elements of another shape set with given transform applied
 	 */
-	void Append(const FSimpleShapeSet3d& OtherShapeSet, const FTransform3d& Transform);
+	UE_API void Append(const FSimpleShapeSet3d& OtherShapeSet, const FTransform3d& Transform);
 
 
 	/**
 	 * Append elements of another shape set with given transforms applied
 	 */
-	void Append(const FSimpleShapeSet3d& OtherShapeSet, const TArray<FTransform3d>& TransformSequence);
+	UE_API void Append(const FSimpleShapeSet3d& OtherShapeSet, const TArray<FTransform3d>& TransformSequence);
 
 	/**
 	 * Append elements of another shape set with given transforms applied
 	 */
-	void Append(const FSimpleShapeSet3d& OtherShapeSet, const FTransformSequence3d& TransformSequence);
+	UE_API void Append(const FSimpleShapeSet3d& OtherShapeSet, const FTransformSequence3d& TransformSequence);
 
 	/**
 	 * Remove any of the elements that are fully contained in larger elements
 	 */
-	void RemoveContainedGeometry();
+	UE_API void RemoveContainedGeometry();
 
 	/**
 	 * Sort the elements by volume and then discard all but the largest MaximumCount elements
 	 */
-	void FilterByVolume(int32 MaximumCount);
+	UE_API void FilterByVolume(int32 MaximumCount);
 
 	struct FMergeShapesSettings
 	{
@@ -207,14 +209,16 @@ struct DYNAMICMESH_API FSimpleShapeSet3d
 		bool bIgnoreInternalNegativeSpace = true;
 	};
 	// Attempt to reduce the number of simple collision shapes by merging neighboring shapes, while still optionally protecting negative space
-	bool MergeShapes(int32 MergeAboveCount, const FMergeShapesSettings& MergeSettings);
+	UE_API bool MergeShapes(int32 MergeAboveCount, const FMergeShapesSettings& MergeSettings);
 
 	/**
 	 * Transform shape elements. This will be a best-effort as if there is non-uniform scaling only Convexes can be transformed correctly
 	 */
-	void ApplyTransform(const FTransform3d& Transform);
+	UE_API void ApplyTransform(const FTransform3d& Transform);
 };
 
 
 } // end namespace UE::Geometry
-} // end namespace UE
\ No newline at end of file
+} // end namespace UE
+
+#undef UE_API
diff --git a/Engine/Plugins/Runtime/GeometryProcessing/Source/DynamicMesh/Public/Solvers/ConstrainedMeshSolver.h b/Engine/Plugins/Runtime/GeometryProcessing/Source/DynamicMesh/Public/Solvers/ConstrainedMeshSolver.h
--- a/Engine/Plugins/Runtime/GeometryProcessing/Source/DynamicMesh/Public/Solvers/ConstrainedMeshSolver.h
+++ b/Engine/Plugins/Runtime/GeometryProcessing/Source/DynamicMesh/Public/Solvers/ConstrainedMeshSolver.h
@@ -16,7 +16,7 @@ namespace UE
 		/**
 		 * Interface to a index-based deformation solver for a 3D mesh vertex set that supports weighted point constraints.
 		 */
-		class DYNAMICMESH_API IConstrainedMeshSolver
+		class IConstrainedMeshSolver
 		{
 		public:
 
@@ -57,7 +57,7 @@ namespace UE
 		 * Extension of IConstrainedMeshSolver that supports manipulating the underlying Laplacian vectors
 		 * used in Laplacian-based Deformation Solvers.
 		 */
-		class DYNAMICMESH_API IConstrainedLaplacianMeshSolver : public IConstrainedMeshSolver
+		class IConstrainedLaplacianMeshSolver : public IConstrainedMeshSolver
 		{
 		public:
 			virtual ~IConstrainedLaplacianMeshSolver() {}
@@ -70,7 +70,7 @@ namespace UE
 		/**
 		 * Interface to a index-based UV solver for a 3D triangle mesh that supports weighted point constraints.
 		 */
-		class DYNAMICMESH_API IConstrainedMeshUVSolver
+		class IConstrainedMeshUVSolver
 		{
 		public:
 
diff --git a/Engine/Plugins/Runtime/GeometryProcessing/Source/DynamicMesh/Public/Solvers/PrecomputedMeshWeightData.h b/Engine/Plugins/Runtime/GeometryProcessing/Source/DynamicMesh/Public/Solvers/PrecomputedMeshWeightData.h
--- a/Engine/Plugins/Runtime/GeometryProcessing/Source/DynamicMesh/Public/Solvers/PrecomputedMeshWeightData.h
+++ b/Engine/Plugins/Runtime/GeometryProcessing/Source/DynamicMesh/Public/Solvers/PrecomputedMeshWeightData.h
@@ -5,6 +5,8 @@
 #include "DynamicMesh/DynamicMesh3.h"
 #include "Solvers/MeshLinearization.h"
 
+#define UE_API DYNAMICMESH_API
+
 namespace UE
 {
 	namespace MeshDeformation
@@ -45,7 +47,7 @@ namespace UE
 		* The per-triangle data used in constructing the cotangent weighted laplacian.
 		*
 		*/
-		class DYNAMICMESH_API CotanTriangleData
+		class CotanTriangleData
 		{
 		public:
 
@@ -58,7 +60,7 @@ namespace UE
 				Initialize(DynamicMesh, TriId);
 			}
 
-			void Initialize(const FDynamicMesh3& DynamicMesh, int32 SrcTriId);
+			UE_API void Initialize(const FDynamicMesh3& DynamicMesh, int32 SrcTriId);
 
 			int32 GetLocalEdgeIdx(const int32 DynamicsMeshEdgeId) const
 			{
@@ -148,12 +150,12 @@ namespace UE
 		* The per-triangle data used in constructing the mean-value weighted laplacian.
 		*
 		*/
-		class DYNAMICMESH_API MeanValueTriangleData
+		class MeanValueTriangleData
 		{
 
 		public:
 
-			void Initialize(const FDynamicMesh3& DynamicMesh, int32 SrcTriId);
+			UE_API void Initialize(const FDynamicMesh3& DynamicMesh, int32 SrcTriId);
 
 			// return Tan(angle / 2) for the corner indicated by this vert id.
 			double GetTanHalfAngle(int32 VtxId) const
@@ -228,4 +230,6 @@ namespace UE
 		void DYNAMICMESH_API ConstructEdgeCotanWeightsDataArray(const FDynamicMesh3& Mesh, TArray<double>& EdgeWeightsDataArray, double ClampMin = -1.e5, double ClampMax = 1.e5);
 		
 	}
-}
\ No newline at end of file
+}
+
+#undef UE_API
diff --git a/Engine/Plugins/Runtime/GeometryProcessing/Source/GeometryAlgorithms/Public/Curve/PolygonIntersectionUtils.h b/Engine/Plugins/Runtime/GeometryProcessing/Source/GeometryAlgorithms/Public/Curve/PolygonIntersectionUtils.h
--- a/Engine/Plugins/Runtime/GeometryProcessing/Source/GeometryAlgorithms/Public/Curve/PolygonIntersectionUtils.h
+++ b/Engine/Plugins/Runtime/GeometryProcessing/Source/GeometryAlgorithms/Public/Curve/PolygonIntersectionUtils.h
@@ -17,7 +17,7 @@ namespace UE::Geometry
 
 	/** Produces the boolean'd result of PolygonA and PolygonB depending on the OperationType. Both polygons must be closed. */
 	template <EPolygonBooleanOp OperationType, typename GeometryType, typename RealType>
-	class GEOMETRYALGORITHMS_API TBooleanPolygon2Polygon2
+	class TBooleanPolygon2Polygon2
 	{
 	public:
 		// Input
diff --git a/Engine/Plugins/Runtime/GeometryProcessing/Source/GeometryAlgorithms/Public/Parameterization/UVMetrics.h b/Engine/Plugins/Runtime/GeometryProcessing/Source/GeometryAlgorithms/Public/Parameterization/UVMetrics.h
--- a/Engine/Plugins/Runtime/GeometryProcessing/Source/GeometryAlgorithms/Public/Parameterization/UVMetrics.h
+++ b/Engine/Plugins/Runtime/GeometryProcessing/Source/GeometryAlgorithms/Public/Parameterization/UVMetrics.h
@@ -3,17 +3,21 @@
 #pragma once
 #include "CoreMinimal.h"
 
+#define UE_API GEOMETRYALGORITHMS_API
+
 namespace UE {
 namespace Geometry {
 
 	class FDynamicMesh3;
 
-	class GEOMETRYALGORITHMS_API FUVMetrics
+	class FUVMetrics
 	{
 		public:
-			static double ReedBeta(const UE::Geometry::FDynamicMesh3& Mesh, int32 UVChannel, int32 Tid);
-			static double Sander(const UE::Geometry::FDynamicMesh3& Mesh, int32 UVChannel, int32 Tid, bool bUseL2);
-			static double TexelDensity(const UE::Geometry::FDynamicMesh3& Mesh, int32 UVChannel, int32 Tid, int32 MapSize);		
+			static UE_API double ReedBeta(const UE::Geometry::FDynamicMesh3& Mesh, int32 UVChannel, int32 Tid);
+			static UE_API double Sander(const UE::Geometry::FDynamicMesh3& Mesh, int32 UVChannel, int32 Tid, bool bUseL2);
+			static UE_API double TexelDensity(const UE::Geometry::FDynamicMesh3& Mesh, int32 UVChannel, int32 Tid, int32 MapSize);		
 	};
 }
-}
\ No newline at end of file
+}
+
+#undef UE_API
