add new cluster-based simplifier, based on mesh edge-graph voronoi connectivity, for quick simplification of large meshes

tends to produce very regular meshes if input mesh had sufficient resolution, but may lose small details and introduce more artifacts than the edge-collapse-based simplifiers

diff --git a/Engine/Plugins/Runtime/GeometryProcessing/Source/DynamicMesh/Private/Operations/MeshClusterSimplifier.cpp b/Engine/Plugins/Runtime/GeometryProcessing/Source/DynamicMesh/Private/Operations/MeshClusterSimplifier.cpp
new file mode 100644
--- /dev/null
+++ b/Engine/Plugins/Runtime/GeometryProcessing/Source/DynamicMesh/Private/Operations/MeshClusterSimplifier.cpp
@@ -0,0 +1,733 @@
+// Copyright Epic Games, Inc. All Rights Reserved.
+
+//#include "Operations/UniformTessellate.h"
+#include "Operations/MeshClusterSimplifier.h"
+
+#include "Async/ParallelFor.h"
+#include "DynamicMesh/DynamicMesh3.h"
+#include "DynamicMesh/DynamicMeshAttributeSet.h"
+#include "VectorTypes.h"
+
+
+namespace UE::MeshClusterSimplifyLocals
+{
+	template<typename ElemType, int Dim, typename AttributeType>
+	void CopyAttribs(AttributeType* Result, const AttributeType* Source, TConstArrayView<int32> ResToSource, int32 Num)
+	{
+		ParallelFor(Num, [&ResToSource, &Source, &Result](int32 ResID)
+			{
+				int32 SourceID = ResToSource[ResID];
+				ElemType ToCopy[Dim];
+				Source->GetValue(SourceID, ToCopy);
+				Result->SetValue(ResID, ToCopy);
+			}
+		);
+	}
+}
+
+namespace UE::Geometry::MeshClusterSimplify
+{
+
+bool Simplify(const FDynamicMesh3& InMesh, FDynamicMesh3& ResultMesh, const FSimplifyOptions& SimplifyOptions)
+{
+	TRACE_CPUPROFILER_EVENT_SCOPE(MeshClusterSimplify::Simplify);
+
+	// we build the result mesh by incrementally copying from the input mesh, so they shouldn't be the same mesh
+	if (!ensure(&ResultMesh != &InMesh))
+	{
+		return false;
+	}
+
+	ResultMesh.Clear();
+	
+	const FDynamicMeshAttributeSet* InAttribs = InMesh.Attributes();
+
+	// We tag edges and vertices w/ the constraint level, abbreviated to EElemTag for convenience
+	using EElemTag = FSimplifyOptions::EConstraintLevel;
+	constexpr int32 NUM_TAGS = 3;
+
+	// TODO: optionally also compute some vertex curvature feature & sort by it, to favor capturing less flat parts of the input shape?
+
+	///
+	/// Step 1, Data Prep: Translate all mesh constraint options to simple per-edge and per-vertex tags, so we know what to try to especially preserve in the result
+	///
+	
+	// Compute an Edge ID -> Constraint Level mapping
+	TArray<EElemTag> EdgeTags;
+	EdgeTags.SetNumUninitialized(InMesh.MaxEdgeID());
+	ParallelFor(InMesh.MaxEdgeID(), [&InMesh, &InAttribs, &SimplifyOptions, &EdgeTags](int32 EID)
+	{
+		if (!InMesh.IsEdge(EID))
+		{
+			return;
+		}
+
+		EElemTag UseTag = EElemTag::Free;
+		if ((uint8)SimplifyOptions.PreserveEdges.Boundary < (uint8)UseTag && InMesh.IsBoundaryEdge(EID))
+		{
+			UseTag = SimplifyOptions.PreserveEdges.Boundary;
+		}
+		if ((uint8)SimplifyOptions.PreserveEdges.PolyGroup < (uint8)UseTag && InMesh.IsGroupBoundaryEdge(EID))
+		{
+			UseTag = SimplifyOptions.PreserveEdges.PolyGroup;
+		}
+
+		if (InAttribs)
+		{
+			if ((uint8)SimplifyOptions.PreserveEdges.Material < (uint8)UseTag && InAttribs->IsMaterialBoundaryEdge(EID))
+			{
+				UseTag = SimplifyOptions.PreserveEdges.Material;
+			}
+
+			if ((uint8)SimplifyOptions.PreserveEdges.UVSeam < (uint8)UseTag)
+			{
+				for (int32 UVLayer = 0; UVLayer < InAttribs->NumUVLayers(); ++UVLayer)
+				{
+					if (InAttribs->GetUVLayer(UVLayer)->IsSeamEdge(EID))
+					{
+						UseTag = SimplifyOptions.PreserveEdges.UVSeam;
+						break;
+					}
+				}
+			}
+			if ((uint8)SimplifyOptions.PreserveEdges.TangentSeam < (uint8)UseTag)
+			{
+				for (int32 NormalLayer = 1; NormalLayer < InAttribs->NumNormalLayers(); ++NormalLayer)
+				{
+					if (InAttribs->GetNormalLayer(NormalLayer)->IsSeamEdge(EID))
+					{
+						UseTag = SimplifyOptions.PreserveEdges.TangentSeam;
+						break;
+					}
+				}
+			}
+			if ((uint8)SimplifyOptions.PreserveEdges.NormalSeam < (uint8)UseTag)
+			{
+				if (const FDynamicMeshNormalOverlay* Normals = InAttribs->PrimaryNormals())
+				{
+					if (Normals->IsSeamEdge(EID))
+					{
+						UseTag = SimplifyOptions.PreserveEdges.NormalSeam;
+					}
+				}
+			}
+			if ((uint8)SimplifyOptions.PreserveEdges.ColorSeam < (uint8)UseTag)
+			{
+				if (const FDynamicMeshColorOverlay* Colors = InAttribs->PrimaryColors())
+				{
+					if (Colors->IsSeamEdge(EID))
+					{
+						UseTag = SimplifyOptions.PreserveEdges.ColorSeam;
+					}
+				}
+			}
+		}
+
+		EdgeTags[EID] = UseTag;
+	});
+
+	TArray<EElemTag> VertexTags; 
+	VertexTags.SetNumUninitialized(InMesh.MaxVertexID());
+
+	double CosBoundaryEdgeAngleTolerance = FMath::Cos(FMath::DegreesToRadians(FMath::Clamp(SimplifyOptions.FixBoundaryAngleTolerance, 0, 180)));
+
+	auto IsSeamIntersectionVertex = 
+		[&InAttribs, &SimplifyOptions]
+		(int32 VID)
+	{
+		if (SimplifyOptions.PreserveEdges.UVSeam == EElemTag::Constrained)
+		{
+			for (int32 Layer = 0; Layer < InAttribs->NumUVLayers(); ++Layer)
+			{
+				if (InAttribs->GetUVLayer(Layer)->IsSeamIntersectionVertex(VID))
+				{
+					return true;
+				}
+			}
+		}
+		if (SimplifyOptions.PreserveEdges.NormalSeam == EElemTag::Constrained)
+		{
+			if (const FDynamicMeshNormalOverlay* Normals = InAttribs->PrimaryNormals())
+			{
+				if (Normals->IsSeamIntersectionVertex(VID))
+				{
+					return true;
+				}
+			}
+		}
+		if (SimplifyOptions.PreserveEdges.TangentSeam == EElemTag::Constrained)
+		{
+			for (int32 Layer = 1; Layer < InAttribs->NumNormalLayers(); ++Layer)
+			{
+				if (InAttribs->GetNormalLayer(Layer)->IsSeamIntersectionVertex(VID))
+				{
+					return true;
+				}
+			}
+		}
+		if (SimplifyOptions.PreserveEdges.ColorSeam == EElemTag::Constrained)
+		{
+			if (const FDynamicMeshColorOverlay* Colors = InAttribs->PrimaryColors())
+			{
+				if (Colors->IsSeamIntersectionVertex(VID))
+				{
+					return true;
+				}
+			}
+		}
+
+		return false;
+	};
+
+	ParallelFor(InMesh.MaxVertexID(), 
+		[&InMesh, &InAttribs, &SimplifyOptions, 
+		&VertexTags, &EdgeTags, 
+		&IsSeamIntersectionVertex,
+		CosBoundaryEdgeAngleTolerance]
+		(int32 VID)
+		{
+			if (!InMesh.IsVertex(VID))
+			{
+				return;
+			}
+			int32 FixedCount = 0;
+			int32 ConstrainedCount = 0;
+			
+			FVector3d BoundaryEdgeVert[2];
+			int32 FoundBoundaryEdgeVerts = 0;
+			InMesh.EnumerateVertexEdges(VID, 
+				[&EdgeTags, VID, &FixedCount, &ConstrainedCount, &BoundaryEdgeVert,
+				&FoundBoundaryEdgeVerts, &SimplifyOptions, &InMesh]
+				(int32 EID)
+				{
+					FixedCount += int32(EdgeTags[EID] == EElemTag::Fixed);
+					if (EdgeTags[EID] == EElemTag::Constrained)
+					{
+						ConstrainedCount++;
+						if (SimplifyOptions.FixBoundaryAngleTolerance > 0)
+						{
+							if (InMesh.IsBoundaryEdge(EID))
+							{
+								if (FoundBoundaryEdgeVerts < 2)
+								{
+									FIndex2i EdgeV = InMesh.GetEdgeV(EID);
+									int32 OtherV = EdgeV.A == VID ? EdgeV.B : EdgeV.A;
+
+									BoundaryEdgeVert[FoundBoundaryEdgeVerts] = InMesh.GetVertex(OtherV);
+								}
+								FoundBoundaryEdgeVerts++;
+							}
+						}
+					}
+				}
+			);
+
+
+			if (FixedCount > 0)
+			{
+				VertexTags[VID] = EElemTag::Fixed;
+				return;
+			}
+			if (FoundBoundaryEdgeVerts == 2)
+			{
+				FVector3d CenterV = InMesh.GetVertex(VID);
+				FVector3d E1 = Normalized(BoundaryEdgeVert[0] - CenterV);
+				FVector3d E2 = Normalized(CenterV - BoundaryEdgeVert[1]);
+				if (E1.Dot(E2) < CosBoundaryEdgeAngleTolerance)
+				{
+					VertexTags[VID] = EElemTag::Fixed;
+					return;
+				}
+			}
+			if (ConstrainedCount > 0)
+			{
+				if (ConstrainedCount == 2 &&
+					// seams are a special case where we can have two constrained edges but still be at a seam intersection 
+					// (e.g. at a vertex that joins two different types of seam)
+					(!InAttribs || !IsSeamIntersectionVertex(VID))
+					)
+				{
+					// constrain vertices along contiguous constrained edge paths
+					VertexTags[VID] = EElemTag::Constrained;
+					return;
+				}
+				else
+				{
+					// fix vertices at constraint intersections
+					VertexTags[VID] = EElemTag::Fixed;
+					return;
+				}
+			}
+
+			VertexTags[VID] = EElemTag::Free;
+		});
+
+
+	///
+	/// Step 2. Clustering: Grow vertex clusters out to the target edge length size
+	///
+	
+	// Buckets of vertices to process -- vertices that are processed sooner are more likely to be directly included in the output
+	TStaticArray<TArray<int32>, NUM_TAGS> ProcessBuckets;
+	for (int32 VID : InMesh.VertexIndicesItr())
+	{
+		ProcessBuckets[(int32)VertexTags[VID]].Add(VID);
+	}
+
+	TArray<float> SourceDist;
+	TArray<int32> Source;
+	Source.Init(INDEX_NONE, InMesh.MaxVertexID());
+	SourceDist.Init(FMathf::MaxReal, InMesh.MaxVertexID());
+
+	auto TagVerticesByRegionGrowth = 
+		[&Source, &SourceDist, &InMesh, &SimplifyOptions, &EdgeTags, &VertexTags]
+		(const TStaticArray<TArray<int32>, NUM_TAGS>& VertexIDBuckets)
+	{
+		// add all the fixed vertices as sources first, so they can't be claimed by other verts
+		for (int32 VID : VertexIDBuckets[(int32)EElemTag::Fixed])
+		{
+			Source[VID] = VID;
+			SourceDist[VID] = 0.f;
+		}
+
+		struct FWalk
+		{
+			int32 VID;
+			float Dist;
+
+			bool operator<(const FWalk& Other) const
+			{
+				return Dist < Other.Dist;
+			}
+		};
+		TArray<FWalk> HeapV;
+
+		// for the non-fixed vertices, progressively grow from vertices, in passes from more-constrained to less-constrained edges
+		for (uint8 TagIdx = 1; TagIdx < (uint8)NUM_TAGS; ++TagIdx)
+		{
+			for (uint8 BucketIdx = 0; BucketIdx <= TagIdx; ++BucketIdx)
+			{
+				const TArray<int32>& CurBucket = VertexIDBuckets[BucketIdx];
+				for (int32 InBucketIdx = 0; InBucketIdx < CurBucket.Num(); ++InBucketIdx)
+				{
+					int32 GrowFromVID = CurBucket[InBucketIdx];
+
+					int32& CurSourceVID = Source[GrowFromVID];
+					float& CurSourceDist = SourceDist[GrowFromVID];
+					// the vertex is unclaimed, claim it as a new source/kept vertex
+					if (CurSourceVID == INDEX_NONE)
+					{
+						CurSourceVID = GrowFromVID;
+						CurSourceDist = 0.f;
+					}
+					// if the vertex was claimed by another source in the current tag pass, no need to process it further
+					else if (CurSourceVID != GrowFromVID && (uint8)VertexTags[GrowFromVID] == TagIdx)
+					{
+						continue;
+					}
+
+					// vertex is either a new source, or previously claimed but we need to consider growing via less-constrained edges
+
+					// helper to add candidate verts to a heap
+					HeapV.Reset();
+					auto AddCandidates = [MaxDist = SimplifyOptions.TargetEdgeLength,
+						&HeapV, &InMesh, &SourceDist, &Source, &EdgeTags, &VertexTags, TagIdx]
+						(const FWalk& From)
+						{
+							// expand to one-ring
+							InMesh.EnumerateVertexEdges(From.VID,
+								[&From, MaxDist,
+								&HeapV, &InMesh, &SourceDist, &Source, &EdgeTags, &VertexTags, TagIdx]
+								(int32 EID)
+								{
+									if ((uint8)EdgeTags[EID] != TagIdx)
+									{
+										return;
+									}
+
+									FIndex2i EdgeV = InMesh.GetEdgeV(EID);
+									int32 ToVID = EdgeV.A == From.VID ? EdgeV.B : EdgeV.A;
+
+									if ((uint8)VertexTags[ToVID] < TagIdx || From.Dist >= SourceDist[ToVID])
+									{
+										// vertex was already claimed by more-constrained context, or is already as close (or closer) to another source
+										return;
+									}
+									// possible candidate, compute the actual distance and grow if close enough
+									FVector3d Pos = InMesh.GetVertex(ToVID);
+									FVector3d FromPos = InMesh.GetVertex(From.VID);
+									float NewDist = From.Dist + (float)FVector3d::Dist(Pos, FromPos);
+									if (NewDist < MaxDist && NewDist < SourceDist[ToVID])
+									{
+										// Viable candidate distance; add to heap
+										HeapV.HeapPush(FWalk{ ToVID, NewDist });
+									}
+								}
+							);
+						};
+
+					// initialize the heap w/ the neighbors of the initial grow-from vertex
+					FWalk Start{ GrowFromVID, CurSourceDist };
+					AddCandidates(Start);
+
+					while (!HeapV.IsEmpty())
+					{
+						FWalk CurWalk;
+						HeapV.HeapPop(CurWalk, EAllowShrinking::No);
+
+						// we already got to this vert from another place
+						if (SourceDist[CurWalk.VID] <= CurWalk.Dist)
+						{
+							continue;
+						}
+
+						// claim the vertex
+						SourceDist[CurWalk.VID] = CurWalk.Dist;
+						Source[CurWalk.VID] = CurSourceVID;
+
+						// search its (current-tag-level) edges for more verts to claim
+						AddCandidates(CurWalk);
+					}
+				}
+			}
+		}
+	};
+
+	TagVerticesByRegionGrowth(ProcessBuckets);
+	for (int32 Idx = 0; Idx < ProcessBuckets.Num(); ++Idx)
+	{
+		ProcessBuckets[Idx].Empty();
+	}
+
+
+	///
+	/// Step 3: Copy the cluster connectivity out to our ResultMesh
+	/// 
+
+	TArray<int32> ToResVID, FromResVID;
+	TArray<int32> ResultToSourceTri;
+
+	// If simplification introduces non-manifold edges, we can often recover by fixing more vertices and re-attempting the build.
+	// After MeshBuildAttempts tries, if still failing, we stop adding vertices and just duplicate vertices to add the non-manifold triangles.
+	// TODO: We could potentially analyze the cluster connectivity more carefully handle more degenerate cluster connectivity, more robustly.
+	//		(if so -- it may be better to do so by analyzing the graph before building the ResultMesh, rather than this rebuilding approach!)
+	int32 MeshBuildAttempts = 2;
+	bool bResultHasDuplicateVertices = false;
+
+	while (MeshBuildAttempts-- > 0)
+	{
+		// clear mesh outputs
+		ToResVID.Reset();
+		FromResVID.Reset();
+		ResultToSourceTri.Reset();
+		ResultMesh.Clear();
+
+
+		bool bAllowDegenerate = MeshBuildAttempts <= 0;
+		// Array of vertex IDs to set to 'fixed' on a rebuild attempt
+		TArray<int32> SourceVIDToFix;
+
+		ToResVID.Init(INDEX_NONE, InMesh.MaxVertexID());
+		for (int32 VID = 0; VID < Source.Num(); ++VID)
+		{
+			if (Source[VID] == VID)
+			{
+				ToResVID[VID] = ResultMesh.AppendVertex(InMesh.GetVertex(VID));
+				// we need the reverse mapping if we're transferring seams
+				if (SimplifyOptions.bTransferAttributes)
+				{
+					FromResVID.Add(VID);
+				}
+			}
+		}
+
+		for (int32 TID : InMesh.TriangleIndicesItr())
+		{
+			FIndex3i Tri = InMesh.GetTriangle(TID);
+			FIndex3i SourceTri(Source[Tri.A], Source[Tri.B], Source[Tri.C]);
+			if (SourceTri.A != SourceTri.B && SourceTri.A != SourceTri.C && SourceTri.B != SourceTri.C)
+			{
+				FIndex3i ResTri(ToResVID[SourceTri.A], ToResVID[SourceTri.B], ToResVID[SourceTri.C]);
+				int32 ResultTID = ResultMesh.AppendTriangle(ResTri);
+				if (ResultTID == FDynamicMesh3::NonManifoldID)
+				{
+					if (bAllowDegenerate)
+					{
+						// TODO: only duplicate vertices on the non-manifold edge(s)
+						FIndex3i ExtraTri;
+						ExtraTri.A = ResultMesh.AppendVertex(ResultMesh.GetVertex(ResTri.A));
+						FromResVID.Add(SourceTri.A);
+						ExtraTri.B = ResultMesh.AppendVertex(ResultMesh.GetVertex(ResTri.B));
+						FromResVID.Add(SourceTri.B);
+						ExtraTri.C = ResultMesh.AppendVertex(ResultMesh.GetVertex(ResTri.C));
+						FromResVID.Add(SourceTri.C);
+						ResultTID = ResultMesh.AppendTriangle(ExtraTri);
+						bResultHasDuplicateVertices = true;
+					}
+					else
+					{
+						// Non-manifold edges can often be resolved by adding an extra vertex --
+						// mark the vertex with largest SourceDist for inclusion in the result mesh
+						int32 BestSubIdx = INDEX_NONE;
+						float BestDist = 0;
+						for (int32 SubIdx = 0; SubIdx < 3; ++SubIdx)
+						{
+							if (SourceDist[Tri[SubIdx]] > BestDist)
+							{
+								BestDist = SourceDist[Tri[SubIdx]];
+								BestSubIdx = SubIdx;
+							}
+						}
+						if (BestSubIdx != INDEX_NONE)
+						{
+							SourceVIDToFix.Add(Tri[BestSubIdx]);
+						}
+					}
+				}
+				if ((SimplifyOptions.bTransferAttributes || SimplifyOptions.bTransferGroups) && ResultTID >= 0)
+				{
+					checkSlow(ResultTID == ResultToSourceTri.Num()); // ResultMesh starts empty and should be compact
+					ResultToSourceTri.Add(TID);
+				}
+			}
+		}
+
+		// We marked some new vertices for inclusion in the result; tag them and re-try
+		if (!bAllowDegenerate && SourceVIDToFix.Num() > 0)
+		{
+			for (int32 VID : SourceVIDToFix)
+			{
+				VertexTags[VID] = EElemTag::Fixed;
+			}
+			ProcessBuckets[0] = MoveTemp(SourceVIDToFix);
+			TagVerticesByRegionGrowth(ProcessBuckets);
+			continue;
+		}
+
+		// Accept the result mesh triangulation
+		break;
+	}
+
+	///
+	/// Step 4: After accepting the final ResultMesh triangulation, copy the input mesh's attributes (UVs, materials, etc) over as well
+	/// 
+
+	if (SimplifyOptions.bTransferAttributes)
+	{
+		ResultMesh.EnableMatchingAttributes(InMesh);
+
+		if (InMesh.HasAttributes())
+		{
+			FDynamicMeshAttributeSet* ResultAttribs = ResultMesh.Attributes();
+
+			const bool bPreserveAnySeams =
+				SimplifyOptions.PreserveEdges.UVSeam != EElemTag::Free ||
+				SimplifyOptions.PreserveEdges.NormalSeam != EElemTag::Free ||
+				SimplifyOptions.PreserveEdges.TangentSeam != EElemTag::Free ||
+				SimplifyOptions.PreserveEdges.ColorSeam != EElemTag::Free;
+
+			// Seam mapping for overlays
+			{
+				// Compute a general wedge mapping that all the overlays can build from
+
+				// Map from ResultTID -> a source triangle per tri-vertex [aka wedge]
+				TArray<FIndex3i> ResultWedgeSourceTris;
+				// sub-indices per wedge
+				TArray<int8> SourceTriWedgeSubIndices;
+				ResultWedgeSourceTris.SetNumUninitialized(ResultMesh.MaxTriangleID());
+				SourceTriWedgeSubIndices.SetNumUninitialized(ResultMesh.MaxTriangleID() * 3);
+				ParallelFor(ResultMesh.MaxTriangleID(), 
+					[&ResultMesh, &ResultWedgeSourceTris, &SourceTriWedgeSubIndices,
+					bPreserveAnySeams, &FromResVID, &ResultToSourceTri, &Source,
+					&VertexTags, &EdgeTags, &InMesh]
+					(int32 ResultTID)
+					{
+						TArray<int32> TriQ;
+						TSet<int32> LocalSeenTris;
+						FIndex3i ResultVIDs = ResultMesh.GetTriangle(ResultTID);
+						for (int32 SubIdx = 0; SubIdx < 3; ++SubIdx)
+						{
+							int32 ResultVID = ResultVIDs[SubIdx];
+							int32 SourceVID = FromResVID[ResultVID];
+							bool bFound = false;
+
+							// we're on a seam vertex, do a local search (w/out crossing seam edges) to from the init triangle to the source vertex
+							// to try to find the best tri to use as a wedge reference
+							if (VertexTags[SourceVID] != EElemTag::Free && bPreserveAnySeams)
+							{
+								TriQ.Reset();
+								LocalSeenTris.Reset();
+								int32 SourceTID = ResultToSourceTri[ResultTID];
+								TriQ.Add(SourceTID);
+								while (!TriQ.IsEmpty())
+								{
+									int32 SearchTID = TriQ.Pop(EAllowShrinking::No);
+
+									if (LocalSeenTris.Contains(SearchTID))
+									{
+										continue;
+									}
+									LocalSeenTris.Add(SearchTID);
+
+									FIndex3i Tri = InMesh.GetTriangle(SearchTID);
+									int32 FoundSubIdx = Tri.IndexOf(SourceVID);
+									if (FoundSubIdx != INDEX_NONE)
+									{
+										bFound = true;
+										ResultWedgeSourceTris[ResultTID][SubIdx] = SearchTID;
+										SourceTriWedgeSubIndices[ResultTID * 3 + SubIdx] = (int8)FoundSubIdx;
+										break;
+									}
+
+									// check we're still on a valid triangle that has a vert tagged w/ our source VID
+									FIndex3i SourceTri(Source[Tri.A], Source[Tri.B], Source[Tri.C]);
+									if (!SourceTri.Contains(SourceVID))
+									{
+										continue;
+									}
+									FIndex3i TriEdges = InMesh.GetTriEdges(SearchTID);
+									for (int32 EdgeSubIdx = 0; EdgeSubIdx < 3; ++EdgeSubIdx)
+									{
+										int32 WalkSourceEID = TriEdges[EdgeSubIdx];
+										if (EdgeTags[WalkSourceEID] == EElemTag::Free)
+										{
+											FIndex2i EdgeT = InMesh.GetEdgeT(WalkSourceEID);
+											int32 WalkTID = EdgeT.A == SearchTID ? EdgeT.B : EdgeT.A;
+											if (WalkTID != INDEX_NONE)
+											{
+												TriQ.Add(WalkTID);
+											}
+										}
+									}
+								}
+							}
+
+							if (!bFound)
+							{
+								// no seams, or search failed; just grab any triangle
+								int32 NbrTID = *InMesh.VtxTrianglesItr(SourceVID).begin();
+								checkSlow(NbrTID != INDEX_NONE); // should not be possible for a vert w/ no neighbors to end up as a source VID
+								ResultWedgeSourceTris[ResultTID][SubIdx] = NbrTID;
+								SourceTriWedgeSubIndices[ResultTID * 3 + SubIdx] = (int8)InMesh.GetTriangle(NbrTID).IndexOf(SourceVID);
+							}
+						}
+					}
+				);
+
+				// Helper to use the general wedge mapping to copy elements for a given overlay
+				auto OverlayTransfer = 
+					[&ResultMesh, &ResultWedgeSourceTris, &SourceTriWedgeSubIndices,
+					bResultHasDuplicateVertices]
+					<typename OverlayType>
+					(OverlayType* ResultOverlay, const OverlayType* SourceOverlay)
+				{
+					TArray<int32> SourceToResElID;
+					SourceToResElID.Init(INDEX_NONE, SourceOverlay->MaxElementID());
+
+					// Note: Unfortunately can't parallelize this part easily; the overlay append and set both are not thread safe (due to ref counts)
+					for (int32 ResultTID : ResultMesh.TriangleIndicesItr())
+					{
+						FIndex3i ResultElemTri;
+						for (int32 ResultSubIdx = 0; ResultSubIdx < 3; ++ResultSubIdx)
+						{
+							int32 SourceTID = ResultWedgeSourceTris[ResultTID][ResultSubIdx];
+							int8 SourceSubIdx = SourceTriWedgeSubIndices[ResultTID * 3 + ResultSubIdx];
+							int32 SourceElemID = SourceOverlay->GetTriangle(SourceTID)[SourceSubIdx];
+							int32 UseElemID;
+							if (SourceToResElID[SourceElemID] == INDEX_NONE)
+							{
+								SourceToResElID[SourceElemID] = ResultOverlay->AppendElement(SourceOverlay->GetElement(SourceElemID));
+								UseElemID = SourceToResElID[SourceElemID];
+							}
+							else
+							{
+								UseElemID = SourceToResElID[SourceElemID];
+								// if we have duplicate vertices, may need to also duplicate the element
+								if (bResultHasDuplicateVertices)
+								{
+									if (ResultOverlay->GetParentVertex(UseElemID) != ResultMesh.GetTriangle(ResultTID)[ResultSubIdx])
+									{
+										UseElemID = ResultOverlay->AppendElement(SourceOverlay->GetElement(SourceElemID));
+									}
+								}
+							}
+							ResultElemTri[ResultSubIdx] = UseElemID;
+						}
+						ResultOverlay->SetTriangle(ResultTID, ResultElemTri);
+					}
+				};
+
+				for (int32 LayerIdx = 0; LayerIdx < InAttribs->NumUVLayers(); ++LayerIdx)
+				{
+					FDynamicMeshUVOverlay* ResultUVs = ResultAttribs->GetUVLayer(LayerIdx);
+					const FDynamicMeshUVOverlay* SourceUVs = InAttribs->GetUVLayer(LayerIdx);
+					OverlayTransfer(ResultUVs, SourceUVs);
+				}
+
+				for (int32 LayerIdx = 0; LayerIdx < InAttribs->NumNormalLayers(); ++LayerIdx)
+				{
+					OverlayTransfer(ResultAttribs->GetNormalLayer(LayerIdx), InAttribs->GetNormalLayer(LayerIdx));
+				}
+
+				if (InAttribs->HasPrimaryColors())
+				{
+					OverlayTransfer(ResultAttribs->PrimaryColors(), InAttribs->PrimaryColors());
+				}
+			}
+
+			for (int32 WeightLayerIdx = 0; WeightLayerIdx < InAttribs->NumWeightLayers(); ++WeightLayerIdx)
+			{
+				UE::MeshClusterSimplifyLocals::CopyAttribs<float, 1>(
+					ResultAttribs->GetWeightLayer(WeightLayerIdx),
+					InAttribs->GetWeightLayer(WeightLayerIdx),
+					FromResVID, ResultMesh.MaxVertexID()
+				);
+			}
+
+			for (int32 SculptLayerIdx = 0; SculptLayerIdx < InAttribs->NumSculptLayers(); ++SculptLayerIdx)
+			{
+				UE::MeshClusterSimplifyLocals::CopyAttribs<double, 3>(
+					ResultAttribs->GetSculptLayers()->GetLayer(SculptLayerIdx),
+					InAttribs->GetSculptLayers()->GetLayer(SculptLayerIdx),
+					FromResVID, ResultMesh.MaxVertexID()
+				);
+			}
+
+			for (int32 GroupLayerIdx = 0; GroupLayerIdx < InAttribs->NumPolygroupLayers(); ++GroupLayerIdx)
+			{
+				UE::MeshClusterSimplifyLocals::CopyAttribs<int32, 1>(
+					ResultAttribs->GetPolygroupLayer(GroupLayerIdx),
+					InAttribs->GetPolygroupLayer(GroupLayerIdx),
+					ResultToSourceTri, ResultMesh.MaxTriangleID()
+				);
+			}
+
+			if (const FDynamicMeshMaterialAttribute* InMats = InAttribs->GetMaterialID())
+			{
+				UE::MeshClusterSimplifyLocals::CopyAttribs<int32, 1>(
+					ResultAttribs->GetMaterialID(),
+					InMats,
+					ResultToSourceTri, ResultMesh.MaxTriangleID()
+				);
+			}
+		}
+	}
+
+	if (SimplifyOptions.bTransferGroups && InMesh.HasTriangleGroups())
+	{
+		ResultMesh.EnableTriangleGroups();
+		ParallelFor(ResultMesh.MaxTriangleID(),
+			[&ResultMesh, &InMesh, &ResultToSourceTri]
+			(int32 ResultTID)
+			{
+				checkSlow(ResultMesh.IsTriangle(ResultTID)); // ResultMesh is compact so all tris should be valid
+				int32 SourceTID = ResultToSourceTri[ResultTID];
+				ResultMesh.SetTriangleGroup(ResultTID, InMesh.GetTriangleGroup(SourceTID));
+			}
+		);
+	}
+
+	return true;
+}
+
+} // namespace UE::Geometry
diff --git a/Engine/Plugins/Runtime/GeometryProcessing/Source/DynamicMesh/Public/Operations/MeshClusterSimplifier.h b/Engine/Plugins/Runtime/GeometryProcessing/Source/DynamicMesh/Public/Operations/MeshClusterSimplifier.h
new file mode 100644
--- /dev/null
+++ b/Engine/Plugins/Runtime/GeometryProcessing/Source/DynamicMesh/Public/Operations/MeshClusterSimplifier.h
@@ -0,0 +1,98 @@
+// Copyright Epic Games, Inc. All Rights Reserved.
+
+#pragma once
+
+#include "GeometryTypes.h"
+#include "IndexTypes.h"
+#include "Clustering/FaceNormalClustering.h"
+
+class FProgressCancel;
+
+namespace UE::Geometry
+{
+
+class FDynamicMesh3;
+
+namespace MeshClusterSimplify
+{
+
+struct FSimplifyOptions
+{
+	// Vertices within this distance are allowed to be replaced with a single vertex
+	double TargetEdgeLength = 1;
+
+	// If > 0, boundary vertices w/ incident boundary edge angle greater than this (in degrees) will be kept in the output
+	double FixBoundaryAngleTolerance = 45;
+
+	// Constraint options control what simplifications are allowed
+	enum class EConstraintLevel : uint8
+	{
+		// Fixed vertices/edges will generally be preserved in the output, as they will each be given their own cluster
+		Fixed,
+		// Constrained vertices/edges may be simplified, but the edge flow should be preserved
+		// A vertex at an intersection of more than two constrained edges will be automatically preserved as 'Fixed'
+		Constrained,
+		// No constraints / ok to simplify as much as possible
+		Free
+	};
+
+	struct FPreserveFeatures
+	{
+		// Mesh boundaries
+		EConstraintLevel Boundary = EConstraintLevel::Constrained;
+		// Seam types
+		EConstraintLevel UVSeam = EConstraintLevel::Constrained;
+		EConstraintLevel NormalSeam = EConstraintLevel::Constrained;
+		EConstraintLevel TangentSeam = EConstraintLevel::Free;
+		EConstraintLevel ColorSeam = EConstraintLevel::Constrained;
+		// Material ID boundaries
+		EConstraintLevel Material = EConstraintLevel::Constrained;
+		// PolyGroup ID boundaries
+		EConstraintLevel PolyGroup = EConstraintLevel::Constrained;
+
+		// Helper to set all seam types to the same constraint level
+		void SetSeamConstraints(EConstraintLevel Level)
+		{
+			UVSeam = Level;
+			NormalSeam = Level;
+			TangentSeam = Level;
+			ColorSeam = Level;
+		}
+	};
+
+	// Manage which feature edge types we try to retain in the simplified result
+	FPreserveFeatures PreserveEdges{};
+
+	// Whether to attempt to transfer attributes to the result mesh
+	bool bTransferAttributes = true;
+
+	// Whether to attempt to transfer triangle groups (PolyGroups) to the result mesh
+	bool bTransferGroups = true;
+
+
+};
+
+/**
+ * Makes a simplified copy of the input mesh
+ * 
+ * This cluster simplify method first clusters vertices locally by distance (calculated along mesh edges), and creates new triangles
+ * from the connectivity of the clusters. i.e., it is a triangulation of the dual of the graph voronoi diagram over mesh edges.
+ * 
+ * To preserve feature edges:
+ * (1) constrained / feature-edge vertices are prioritized as cluster 'seeds,' and 
+ * (2) clusters are grown along feature edges first, then free edges after -- and growth over 'free' edges cannot claim 'constrained' vertices.
+ * This locks in clusters along 'constrained' feature edges.
+ * 
+ * Note that mesh features can be lost if the clusters are large enough that the graph becomes degenerate
+ *  -- e.g., if a mesh island has so few clusters that the graph connectivity does not contain triangles.
+ * 
+ * 
+ * @param InMesh The mesh to simplify
+ * @param OutSimplifiedMesh This mesh will store the simplified result mesh
+ * @param SimplifyOptions Options controlling simplification
+ * @return true on success
+ */
+bool DYNAMICMESH_API Simplify(const FDynamicMesh3& InMesh, FDynamicMesh3& OutSimplifiedMesh, const FSimplifyOptions& SimplifyOptions);
+
+} // end namespace UE::Geometry
+} // end namespace UE
