fix crash in cluster simplifier on partially-unset overlay
+ add test coverage

diff --git a/Engine/Plugins/Runtime/GeometryProcessing/Source/DynamicMesh/Private/Operations/MeshClusterSimplifier.cpp b/Engine/Plugins/Runtime/GeometryProcessing/Source/DynamicMesh/Private/Operations/MeshClusterSimplifier.cpp
--- a/Engine/Plugins/Runtime/GeometryProcessing/Source/DynamicMesh/Private/Operations/MeshClusterSimplifier.cpp
+++ b/Engine/Plugins/Runtime/GeometryProcessing/Source/DynamicMesh/Private/Operations/MeshClusterSimplifier.cpp
@@ -628,11 +628,19 @@ bool Simplify(const FDynamicMesh3& InMesh, FDynamicMesh3& ResultMesh, const FSim
 					for (int32 ResultTID : ResultMesh.TriangleIndicesItr())
 					{
 						FIndex3i ResultElemTri;
+						bool bHasUnsetSources = false;
 						for (int32 ResultSubIdx = 0; ResultSubIdx < 3; ++ResultSubIdx)
 						{
 							int32 SourceTID = ResultWedgeSourceTris[ResultTID][ResultSubIdx];
 							int8 SourceSubIdx = SourceTriWedgeSubIndices[ResultTID * 3 + ResultSubIdx];
 							int32 SourceElemID = SourceOverlay->GetTriangle(SourceTID)[SourceSubIdx];
+							if (SourceElemID == INDEX_NONE)
+							{
+								// if we mapped to an unset triangle in the source overlay, there is no element to copy
+								// we do not support partially-set triangles, so the whole result triangle will also be unset in this case
+								bHasUnsetSources = true;
+								break;
+							}
 							int32 UseElemID;
 							if (SourceToResElID[SourceElemID] == INDEX_NONE)
 							{
@@ -653,7 +661,11 @@ bool Simplify(const FDynamicMesh3& InMesh, FDynamicMesh3& ResultMesh, const FSim
 							}
 							ResultElemTri[ResultSubIdx] = UseElemID;
 						}
-						ResultOverlay->SetTriangle(ResultTID, ResultElemTri);
+
+						if (!bHasUnsetSources)
+						{
+							ResultOverlay->SetTriangle(ResultTID, ResultElemTri);
+						}
 					}
 				};
 
