add 'select occluded from outside view' jacketing method to geometry script, to have more controllable occlusion operations (e.g. from specific directions, for specific material IDs/parts of mesh, w/ transparent tris)

+ clean up the detect exterior visibility function logic to support the new options and be a bit faster / more robust

diff --git a/Engine/Plugins/Runtime/GeometryProcessing/Source/DynamicMesh/Private/Operations/DetectExteriorVisibility.cpp b/Engine/Plugins/Runtime/GeometryProcessing/Source/DynamicMesh/Private/Operations/DetectExteriorVisibility.cpp
--- a/Engine/Plugins/Runtime/GeometryProcessing/Source/DynamicMesh/Private/Operations/DetectExteriorVisibility.cpp
+++ b/Engine/Plugins/Runtime/GeometryProcessing/Source/DynamicMesh/Private/Operations/DetectExteriorVisibility.cpp
@@ -15,13 +15,60 @@ using namespace UE::Geometry;
 namespace UE::Private::DetectExteriorVisibilityInternal
 {
 
+TArray<FVector3d> GenerateCardinalDirectionsHelper()
+{
+	TArray<FVector3d> CardinalDirections;
+	// Small arbitrary offset vector, used to reduce chance of rays exactly fitting through small axis-aligned cracks
+	const FVector3d TinyPerturb(UE_DOUBLE_KINDA_SMALL_NUMBER, .17 * UE_DOUBLE_KINDA_SMALL_NUMBER, .3 * UE_DOUBLE_KINDA_SMALL_NUMBER);
+	for (int32 k = 0; k < 3; ++k)
+	{
+		FVector3d Direction = TinyPerturb;
+		Direction[k] = 1.0;
+		Direction.Normalize();
+		CardinalDirections.Add(Direction);
+		CardinalDirections.Add(-Direction);
+	}
+	return CardinalDirections;
+}
+
+TArray<FVector3d> GenerateTestDirectionsHelper(int32 NumVisibilityTestDirections, TConstArrayView<FVector3d> AdditionalDirections)
+{
+	// Small arbitrary offset vector, used to reduce chance of rays exactly fitting through small axis-aligned cracks
+	const FVector3d TinyPerturb(UE_DOUBLE_KINDA_SMALL_NUMBER, .17 * UE_DOUBLE_KINDA_SMALL_NUMBER, .3 * UE_DOUBLE_KINDA_SMALL_NUMBER);
+	//
+	// Construct set of exterior visibility test directions, below we will check if sample 
+	// points on the mesh triangles are visible from the exterior along these directions.
+	// Order is modulo-shuffled in hopes that for visible tris we don't waste a bunch
+	// of time on the 'far' side
+	TSphericalFibonacci<double> SphereSampler(NumVisibilityTestDirections);
+	TArray<FVector3d> VisibilityDirections;
+	VisibilityDirections.Reserve(NumVisibilityTestDirections + AdditionalDirections.Num());
+	FModuloIteration ModuloIter(NumVisibilityTestDirections);
+	uint32 DirectionIndex = 0;
+	while (ModuloIter.GetNextIndex(DirectionIndex))
+	{
+		VisibilityDirections.Add(Normalized(SphereSampler[DirectionIndex] + TinyPerturb));
+	}
+	VisibilityDirections.Append(AdditionalDirections);
+	return VisibilityDirections;
+}
+
 // Templated helper to compute occlusion with custom AABB tree (can be specialized for different mesh types)
 // Note that for the has-transparent case, we assume that Spatial *does not* include the transparent triangles
 template<typename SpatialType, bool bHasTransparent>
-void ComputePerTriangleOcclusionHelper(const FDynamicMesh3& TargetMesh, const SpatialType& Spatial, TArray<bool>& OutTriOccluded, bool bPerPolyGroup, const FExteriorVisibilitySampling& SamplingParameters, TArrayView<const bool> SkipTris, TFunctionRef<bool(int32)> IsTransparent)
+void ComputePerTriangleOcclusionHelper(
+	const FDynamicMesh3& TargetMesh, const SpatialType& Spatial, 
+	TArray<bool>& OutTriOccluded, bool bPerPolyGroup, 
+	const FExteriorVisibilitySampling& SamplingParameters, 
+	TArrayView<const bool> SkipTris, 
+	TFunctionRef<bool(int32)> IsTransparent,
+	TConstArrayView<FVector3d> FirstTestDirections,
+	TConstArrayView<FVector3d> FullTestDirections,
+	bool bSampleNormalDirections)
 {
 	const bool bDoubleSided = SamplingParameters.bDoubleSided;
 	const double SampleRadius = FMath::Max(UE_DOUBLE_KINDA_SMALL_NUMBER, SamplingParameters.SamplingDensity);
+	const FVector3d TinyPerturb(UE_DOUBLE_KINDA_SMALL_NUMBER, .17 * UE_DOUBLE_KINDA_SMALL_NUMBER, .3 * UE_DOUBLE_KINDA_SMALL_NUMBER);
 	int32 NumVisibilityTestDirections = FMath::Max(3, SamplingParameters.NumSearchDirections);
 
 	FAxisAlignedBox3d Bounds = Spatial.GetBoundingBox();
@@ -67,18 +114,6 @@ void ComputePerTriangleOcclusionHelper(const FDynamicMesh3& TargetMesh, const Sp
 		return HitTID;
 	};
 
-	// array of (+/-)X/Y/Z directions
-	TArray<FVector3d> CardinalDirections;
-	// Small arbitrary offset vector, used to reduce chance of rays exactly fitting through small axis-aligned cracks
-	const FVector3d TinyPerturb(UE_DOUBLE_KINDA_SMALL_NUMBER, .17 * UE_DOUBLE_KINDA_SMALL_NUMBER, .3*UE_DOUBLE_KINDA_SMALL_NUMBER);
-	for (int32 k = 0; k < 3; ++k)
-	{
-		FVector3d Direction = TinyPerturb;
-		Direction[k] = 1.0;
-		Direction.Normalize();
-		CardinalDirections.Add(Direction);
-		CardinalDirections.Add(-Direction);
-	}
 
 	// TODO: it seems like a quite common failure case is triangles that are somewhat deeply
 	// nested inside cavities/etc. A possible way to help here would be to essentially raytrace
@@ -130,7 +165,7 @@ void ComputePerTriangleOcclusionHelper(const FDynamicMesh3& TargetMesh, const Sp
 		}
 
 		FVector3d UnusedBaryCoords;
-		for (FVector3d Direction : CardinalDirections)
+		for (FVector3d Direction : FirstTestDirections)
 		{
 			// if direction is orthogonal to the triangle, hit-test is unstable, but even
 			// worse, on rectilinear shapes (eg imagine some stacked cubes or adjacent parts)
@@ -138,7 +173,7 @@ void ComputePerTriangleOcclusionHelper(const FDynamicMesh3& TargetMesh, const Sp
 			// and manage to hit the search triangle
 			if (FMathd::Abs(Direction.Dot(Normal)) > GlancingAngleDotTolerance)
 			{
-				if (FindHitTriangleTest(Centroid, Normal, Centroid + Radius * Direction, TID, UnusedBaryCoords) == TID)
+				if (FindHitTriangleTest(Centroid, Normal, Centroid - Radius * Direction, TID, UnusedBaryCoords) == TID)
 				{
 					TriStatusKnown[TID] = EVisStatus::Visible;
 					return;
@@ -148,7 +183,7 @@ void ComputePerTriangleOcclusionHelper(const FDynamicMesh3& TargetMesh, const Sp
 		
 		// Slightly nudge the normal direction to reduce the chance of edge cases for meshes w/ features aligned to the normals
 		const FVector3d NormalRayDir = Normal + TinyPerturb;
-		if (FindHitTriangleTest(Centroid, Normal, Centroid + (1.01 * Radius) * NormalRayDir, TID, UnusedBaryCoords) == TID)
+		if (bSampleNormalDirections && FindHitTriangleTest(Centroid, Normal, Centroid + (1.01 * Radius) * NormalRayDir, TID, UnusedBaryCoords) == TID)
 		{
 			TriStatusKnown[TID] = EVisStatus::Visible;
 			return;
@@ -203,33 +238,20 @@ void ComputePerTriangleOcclusionHelper(const FDynamicMesh3& TargetMesh, const Sp
 	// Thread safe triangle visibility, atomics can be updated on any thread
 	TArray<std::atomic<bool>> ThreadSafeTriVisible;
 	ThreadSafeTriVisible.SetNum(TargetMesh.MaxTriangleID());
-	MaxUnknownTID = -1; // re-compute max unknown
+	// Collect the unknown triangles, to run parallel sampling on only those triangles
+	TArray<int32> UnknownTIDs;
 	for (int32 TID : TargetMesh.TriangleIndicesItr())
 	{
 		EVisStatus Status = TriStatusKnown[TID];
 		ThreadSafeTriVisible[TID] = Status == EVisStatus::Visible;
 		if (Status == EVisStatus::Unknown)
 		{
-			MaxUnknownTID = TID;
+			UnknownTIDs.Add(TID);
 		}
 	}
 
-	//
-	// Construct set of exterior visibility test directions, below we will check if sample 
-	// points on the mesh triangles are visible from the exterior along these directions.
-	// Order is modulo-shuffled in hopes that for visible tris we don't waste a bunch
-	// of time on the 'far' side
-	TSphericalFibonacci<double> SphereSampler(NumVisibilityTestDirections);
-	TArray<FVector3d> VisibilityDirections;
-	FModuloIteration ModuloIter(NumVisibilityTestDirections);
-	uint32 DirectionIndex = 0;
-	while (ModuloIter.GetNextIndex(DirectionIndex))
-	{
-		VisibilityDirections.Add(Normalized(SphereSampler[DirectionIndex] + TinyPerturb));
-	}
-	// Fibonacci set generally does not include the cardinal directions, but they are highly useful to check
-	VisibilityDirections.Append(CardinalDirections);
-	NumVisibilityTestDirections = VisibilityDirections.Num();
+	// Test the requested number of directions + also the cardinal directions
+	NumVisibilityTestDirections = FullTestDirections.Num();
 
 	// For each triangle we will generate a set of sample points on the triangle surface,
 	// and then check if that point is visible along any of the sample directions.
@@ -241,13 +263,10 @@ void ComputePerTriangleOcclusionHelper(const FDynamicMesh3& TargetMesh, const Sp
 	// For each triangle, generate a set of sample points on the triangle surface,
 
 	// This is the expensive part!
-	ParallelFor(MaxUnknownTID + 1, [&](int32 TID)
+	ParallelFor(UnknownTIDs.Num(), [&](int32 UnknownIdx)
 	{
-		if (!TargetMesh.IsTriangle(TID))
-		{
-			return;
-		}
-
+		int32 TID = UnknownTIDs[UnknownIdx];
+		
 		int32 GID = 0;
 		bool bHasValidGroup = bPerPolyGroup;
 		if (bPerPolyGroup)
@@ -271,102 +290,105 @@ void ComputePerTriangleOcclusionHelper(const FDynamicMesh3& TargetMesh, const Sp
 		double TriArea;
 		FVector3d TriNormal = VectorUtil::NormalArea(A, B, C, TriArea);		// TriStatusKnown should skip degen tris, do not need to check here
 
-		FFrame3d TriFrame(Centroid, TriNormal);
-		FTriangle2d UVTriangle(TriFrame.ToPlaneUV(A), TriFrame.ToPlaneUV(B), TriFrame.ToPlaneUV(C));
-
 		// Slightly shrink the triangle, this helps to avoid spurious hits
 		// TODO obviously should scale by an actual dimension and not just a relative %...
-		FVector2d Center = (UVTriangle.V[0] + UVTriangle.V[1] + UVTriangle.V[2]) / 3.0;
+		FTriangle3d SampleTri(A, B, C);
 		for (int32 k = 0; k < 3; ++k)
 		{
-			UVTriangle.V[k] = (1 - TriScalingAlpha) * Center + (TriScalingAlpha)*UVTriangle.V[k];
+			SampleTri.V[k] = Lerp(Centroid, SampleTri.V[k], TriScalingAlpha);
 		}
 
 		double DiscArea = (FMathd::Pi * SampleRadius * SampleRadius);
 		int NumSamples = FMath::Max((int)(TriArea / DiscArea), MinTriSamplesPerSamplePoint);
-		FVector2d V1 = UVTriangle.V[1] - UVTriangle.V[0];
-		FVector2d V2 = UVTriangle.V[2] - UVTriangle.V[0];
+		FVector3d V1 = SampleTri.V[1] - SampleTri.V[0];
+		FVector3d V2 = SampleTri.V[2] - SampleTri.V[0];
 
 		TArray<int32> HitTris, HitGroups;		// re-use this array in inner loop to avoid hitting atomics so often
 
-		int NumTested = 0; int Iterations = 0;
+		// always sample the (already inset) corners and mid-edges
+		constexpr int32 NumFixedSamples = 6;
+		FVector2d FixedSamples[NumFixedSamples]{ {0,0}, {0,1}, {1,0}, {.5,0}, {0,.5}, {.5, .5} };
+
 		FRandomStream RandomStream(TID);
-		while (NumTested < NumSamples && Iterations++ < 10000)
+		for (int32 NumTested = 0; NumTested < NumSamples; ++NumTested)
 		{
-			double a1 = RandomStream.GetFraction();
-			double a2 = RandomStream.GetFraction();
-			if (a1 + a2 > 1) // if we would sample outside the triangle, reflect the sample
+			FVector2d Sample;
+			if (NumTested < NumFixedSamples)
 			{
-				a1 = 1 - a1;
-				a2 = 1 - a2;
+				Sample = FixedSamples[NumTested];
 			}
-			FVector2d PointUV = UVTriangle.V[0] + a1 * V1 + a2 * V2;
-			if (UVTriangle.IsInside(PointUV))
+			else // add random samples
 			{
-				NumTested++;
-				FVector3d Position = TriFrame.FromPlaneUV(PointUV, 2);
+				Sample.X = RandomStream.GetFraction();
+				Sample.Y = RandomStream.GetFraction();
+				if (Sample.X + Sample.Y > 1) // if we would sample outside the triangle, reflect the sample
+				{
+					Sample.X = 1 - Sample.X;
+					Sample.Y = 1 - Sample.Y;
+				}
+			}
+			FVector3d Position = SampleTri.V[0] + Sample.X * V1 + Sample.Y * V2;
 
-				// cast ray from all exterior sample locations for this triangle sample point
-				HitTris.Reset();
-				for (int32 k = 0; k < NumVisibilityTestDirections; ++k)
+			// cast ray from all exterior sample locations for this triangle sample point
+			HitTris.Reset();
+			for (int32 k = 0; k < NumVisibilityTestDirections; ++k)
+			{
+				FVector3d Direction = FullTestDirections[k];
+				if (FMathd::Abs(Direction.Dot(TriNormal)) < GlancingAngleDotTolerance)
 				{
-					FVector3d Direction = VisibilityDirections[k];
-					if (FMathd::Abs(Direction.Dot(TriNormal)) < GlancingAngleDotTolerance)
+					continue;
+				}
+
+				FVector3d RayFrom = Position - 2.0 * Radius * FullTestDirections[k];
+				FVector3d BaryCoords;
+				int32 HitTriID = FindHitTriangleTest(Position, TriNormal, RayFrom, TID, BaryCoords);
+				if (HitTriID != IndexConstants::InvalidID && TriStatusKnown[HitTriID] == EVisStatus::Unknown)
+				{
+					if (HitTriID == TID)
 					{
-						continue;
+						HitTris.Add(HitTriID);
+						break;
 					}
-
-					FVector3d RayFrom = Position + 2.0 * Radius * VisibilityDirections[k];
-					FVector3d BaryCoords;
-					int32 HitTriID = FindHitTriangleTest(Position, TriNormal, RayFrom, TID, BaryCoords);
-					if (HitTriID != IndexConstants::InvalidID && TriStatusKnown[HitTriID] == EVisStatus::Unknown)
+					// For tris other than the one that was sampled to create the ray, we want
+					// to filter out on-edge triangle hits, as they are generally spurious and
+					// will result in interior triangles remaining visible
+					else if (BaryCoords.GetMin() > BaryCoordsThreshold &&
+						BaryCoords.GetMax() < (1.0 - BaryCoordsThreshold))
 					{
-						if (HitTriID == TID)
-						{
-							HitTris.Add(HitTriID);
-							break;
-						}
-						// For tris other than the one that was sampled to create the ray, we want
-						// to filter out on-edge triangle hits, as they are generally spurious and
-						// will result in interior triangles remaining visible
-						else if (BaryCoords.GetMin() > BaryCoordsThreshold &&
-							BaryCoords.GetMax() < (1.0 - BaryCoordsThreshold))
-						{
-							HitTris.AddUnique(HitTriID);
-						}
+						HitTris.AddUnique(HitTriID);
 					}
 				}
+			}
 
-				// mark any hit tris
-				for (int32 HitTriID : HitTris)
+			// mark any hit tris
+			for (int32 HitTriID : HitTris)
+			{
+				ThreadSafeTriVisible[HitTriID] = true;
+			}
+			if (bPerPolyGroup)
+			{
+				// Note: can tune the threshold at which we reset the hitgroups array, to trade off atomic accesses vs AddUnique search cost 
+				constexpr int32 ResetGroupSize = 20;
+				if (HitGroups.Num() >= ResetGroupSize)
 				{
-					ThreadSafeTriVisible[HitTriID] = true;
+					HitGroups.Reset();
 				}
-				if (bPerPolyGroup)
+				for (int32 HitTriID : HitTris)
 				{
-					// Note: can tune the threshold at which we reset the hitgroups array, to trade off atomic accesses vs AddUnique search cost 
-					constexpr int32 ResetGroupSize = 20;
-					if (HitGroups.Num() >= ResetGroupSize)
-					{
-						HitGroups.Reset();
-					}
-					for (int32 HitTriID : HitTris)
+					int32 HitGID = TargetMesh.GetTriangleGroup(HitTriID);
+					if (GroupStatusKnown.IsValidIndex(HitGID) && GroupStatusKnown[HitGID] == EVisStatus::Unknown && HitGroups.Num() == HitGroups.AddUnique(HitGID))
 					{
-						int32 HitGID = TargetMesh.GetTriangleGroup(HitTriID);
-						if (GroupStatusKnown.IsValidIndex(HitGID) && GroupStatusKnown[HitGID] == EVisStatus::Unknown && HitGroups.Num() == HitGroups.AddUnique(HitGID))
-						{
-							// Note: setting/checking the group visibility flag is optional and just allow threads to stop early
-							// so it should be fine to do so w/ the relaxed memory ordering
-							ThreadSafeGroupVisible[HitGID].store(true, std::memory_order_relaxed);
-						}
+						// Note: setting/checking the group visibility flag is optional and just allow threads to stop early
+						// so it should be fine to do so w/ the relaxed memory ordering
+						ThreadSafeGroupVisible[HitGID].store(true, std::memory_order_relaxed);
 					}
 				}
+			}
 
-				// if our triangle has become visible (in this thread or another) we can terminate now
-				if (ThreadSafeTriVisible[TID] || (bHasValidGroup && ThreadSafeGroupVisible[GID].load(std::memory_order_relaxed)))
-				{
-					return;
-				}
+			// if our triangle has become visible (in this thread or another) we can terminate now
+			if (ThreadSafeTriVisible[TID] || (bHasValidGroup && ThreadSafeGroupVisible[GID].load(std::memory_order_relaxed)))
+			{
+				return;
 			}
 		}
 
@@ -417,6 +439,29 @@ void ComputePerTriangleOcclusionHelper(const FDynamicMesh3& TargetMesh, const Sp
 		OutTriOccluded[Idx] = !ThreadSafeTriVisible[Idx];
 	}
 }
+
+template<typename SpatialType, bool bHasTransparent>
+void ComputePerTriangleOcclusionHelper(
+	const FDynamicMesh3& TargetMesh, const SpatialType& Spatial,
+	TArray<bool>& OutTriOccluded, bool bPerPolyGroup,
+	const FExteriorVisibilitySampling& SamplingParameters,
+	TArrayView<const bool> SkipTris,
+	TFunctionRef<bool(int32)> IsTransparent)
+{
+	TArray<FVector3d> CardinalDirections = GenerateCardinalDirectionsHelper();
+	TArray<FVector3d> FullDirections = GenerateTestDirectionsHelper(SamplingParameters.NumSearchDirections, CardinalDirections);
+	ComputePerTriangleOcclusionHelper<SpatialType, bHasTransparent>(
+		TargetMesh, Spatial,
+		OutTriOccluded, bPerPolyGroup,
+		SamplingParameters,
+		SkipTris,
+		IsTransparent,
+		CardinalDirections,
+		FullDirections,
+		true
+	);
+}
+
 } // namespace UE::Private::DetectExteriorVisibilityInternal
 
 void FExteriorVisibilitySampling::ComputePerTriangleOcclusion(const FDynamicMesh3& TargetMesh, TArray<bool>& OutTriOccluded, bool bPerPolyGroup, const FExteriorVisibilitySampling& SamplingParameters, TArrayView<const bool> SkipTris)
@@ -507,6 +552,42 @@ void FExteriorVisibilitySampling::ComputePerTriangleOcclusion(const FDynamicMesh
 	return ComputePerTriangleOcclusionHelper< TMeshAABBTree3<FSkipTransparentMeshWrapper>, true >(TargetMesh, Spatial, OutTriOccluded, bPerPolyGroup, SamplingParameters, SkipTris, IsTriTransparent);
 }
 
+void FExteriorVisibilitySampling::ComputePerTriangleOcclusionFromDirections(const FDynamicMesh3& TargetMesh, TFunctionRef<bool(int32)> IsTriTransparent, TArray<bool>& OutTriOccluded, bool bPerPolyGroup,
+	TConstArrayView<FVector3d> InitialViewTestDirections, TConstArrayView<FVector3d> FullViewTestDirections, bool bAlsoSampleTriangleNormalDirections,
+	const FExteriorVisibilitySampling& SamplingParameters,
+	TArrayView<const bool> SkipTris)
+{
+	using namespace UE::Private::DetectExteriorVisibilityInternal;
+	FSkipTransparentMeshWrapper SkipTransparent(&TargetMesh, IsTriTransparent);
+	TMeshAABBTree3<FSkipTransparentMeshWrapper> Spatial(&SkipTransparent, true);
+	return ComputePerTriangleOcclusionHelper< TMeshAABBTree3<FSkipTransparentMeshWrapper>, true >(TargetMesh, Spatial, OutTriOccluded, bPerPolyGroup, SamplingParameters, SkipTris, IsTriTransparent,
+		InitialViewTestDirections, FullViewTestDirections,bAlsoSampleTriangleNormalDirections);
+}
+
+TArray<FVector3d> FExteriorVisibilitySampling::GetRestrictedViewDirections(TConstArrayView<FVector3d> RestrictViewDirections, double RestrictViewAngle, int32 NumSphereSamples)
+{
+	using namespace UE::Private::DetectExteriorVisibilityInternal;
+
+	TArray<FVector3d> SampleDirections = GenerateTestDirectionsHelper(NumSphereSamples, GenerateCardinalDirectionsHelper());
+	TArray<FVector3d> FoundDirections(RestrictViewDirections);
+	
+	double CosThresh = FMath::Cos(FMath::DegreesToRadians(RestrictViewAngle));
+	for (const FVector3d& SampleDir : SampleDirections)
+	{
+		bool bFoundInCone = false;
+		for (const FVector3d& RestrictDir : RestrictViewDirections)
+		{
+			if (SampleDir.Dot(RestrictDir) >= CosThresh)
+			{
+				FoundDirections.Add(SampleDir);
+				break;
+			}
+		}
+	}
+
+	return FoundDirections;
+}
+
 
 void FDetectPerDynamicMeshExteriorVisibility::ComputeHidden(TArray<bool>& OutIsOccluded, TArray<bool>* OutTransparentIsOccluded)
 {
diff --git a/Engine/Plugins/Runtime/GeometryProcessing/Source/DynamicMesh/Public/Operations/DetectExteriorVisibility.h b/Engine/Plugins/Runtime/GeometryProcessing/Source/DynamicMesh/Public/Operations/DetectExteriorVisibility.h
--- a/Engine/Plugins/Runtime/GeometryProcessing/Source/DynamicMesh/Public/Operations/DetectExteriorVisibility.h
+++ b/Engine/Plugins/Runtime/GeometryProcessing/Source/DynamicMesh/Public/Operations/DetectExteriorVisibility.h
@@ -15,6 +15,7 @@ class FDynamicMesh3;
 struct FExteriorVisibilitySampling
 {
 	// Approximate spacing between triangle samples used for visibility tests
+	// Set a smaller value to increase the number of visibility samples per triangle, which can give more accurate visibility results
 	double SamplingDensity = 1.0;
 
 	// Whether to treat faces as double-sided
@@ -42,6 +43,23 @@ struct FExteriorVisibilitySampling
 	// @param SkipTris				True for triangles do not need visibility testing (will be marked as not occluded). If empty, no triangles will be skipped.
 	DYNAMICMESH_API static void ComputePerTriangleOcclusion(const FDynamicMesh3& Mesh, TFunctionRef<bool(int32)> IsTriTransparent, TArray<bool>& OutTriOccluded, bool bPerPolyGroup, const FExteriorVisibilitySampling& SamplingParameters,
 		TArrayView<const bool> SkipTris = TArrayView<const bool>());
+
+	// Compute per-triangle visibility array for a triangle mesh, with support for transparent (non-occluding) triangles, and with custom sampling directions (ignoring the NumSearchDirections parameter)
+	// @param IsTriTransparent		Indicator function returns true if a triangle is 'transparent' (should not occlude, but can still be occluded)
+	// @param OutTriOccluded		For each valid triangle ID, OutTriOccluded[ID] will be true if that triangle is hidden, false if it is visible
+	// @param InitialViewTestDirections Initial directions to test for triangle visibility/occlusion. Note: directions are from mesh surface toward outside.
+	//									Use this to specify a small set of directions (e.g., cardinal directions) to test per triangle before using the full set, which can help improve the algorithm performance.
+	// @param FullViewTestDirections Full set of directions to test for triangle visibility/occlusion. Note: directions are from mesh surface toward outside
+	// @param bAlsoSampleTriangleNormalDirections Whether to also sample in the normal direction for each triangle. Can be disabled for greater control over the sampling directions.
+	// @param bPerPolyGroup			If true, visiblity will be determined on a per group basis: If any triangle in a group is visible, the whole group will be marked as not-occluded
+	// @param SamplingParameters	Parameters controlling how exterior visiblity is tested
+	// @param SkipTris				True for triangles do not need visibility testing (will be marked as not occluded). If empty, no triangles will be skipped.
+	DYNAMICMESH_API static void ComputePerTriangleOcclusionFromDirections(const FDynamicMesh3& Mesh, TFunctionRef<bool(int32)> IsTriTransparent, TArray<bool>& OutTriOccluded, bool bPerPolyGroup, 
+		TConstArrayView<FVector3d> InitialTestDirections, TConstArrayView<FVector3d> FullTestDirections, bool bAlsoSampleTriangleNormalDirections,
+		const FExteriorVisibilitySampling& SamplingParameters,
+		TArrayView<const bool> SkipTris = TArrayView<const bool>());
+
+	DYNAMICMESH_API static TArray<FVector3d> GetRestrictedViewDirections(TConstArrayView<FVector3d> RestrictViewDirections, double RestrictViewAngle, int32 NumSphereSamples);
 };
 
 // Determine which meshes are visible from exterior views of a list of meshes
diff --git a/Engine/Plugins/Runtime/GeometryScripting/Source/GeometryScriptingCore/Private/MeshRepairFunctions.cpp b/Engine/Plugins/Runtime/GeometryScripting/Source/GeometryScriptingCore/Private/MeshRepairFunctions.cpp
--- a/Engine/Plugins/Runtime/GeometryScripting/Source/GeometryScriptingCore/Private/MeshRepairFunctions.cpp
+++ b/Engine/Plugins/Runtime/GeometryScripting/Source/GeometryScriptingCore/Private/MeshRepairFunctions.cpp
@@ -12,6 +12,8 @@
 #include "DynamicMeshEditor.h"
 #include "Operations/MeshResolveTJunctions.h"
 #include "Operations/RemoveOccludedTriangles.h"
+#include "Operations/DetectExteriorVisibility.h"
+#include "Selections/GeometrySelection.h"
 #include "Selections/MeshConnectedComponents.h"
 #include "Selections/MeshFaceSelection.h"
 #include "UDynamicMesh.h"
@@ -301,6 +303,87 @@ UDynamicMesh* UGeometryScriptLibrary_MeshRepairFunctions::RemoveHiddenTriangles(
 }
 
 
+UDynamicMesh* UGeometryScriptLibrary_MeshRepairFunctions::SelectHiddenTrianglesFromOutside(
+	UDynamicMesh* TargetMesh,
+	FGeometryScriptSelectHiddenTrianglesFromOutsideOptions Options,
+	FGeometryScriptMeshSelection IgnoreSelection,
+	FGeometryScriptMeshSelection TransparentSelection,
+	FGeometryScriptMeshSelection& ResultSelection,
+	UGeometryScriptDebug* Debug)
+{
+	if (TargetMesh == nullptr)
+	{
+		UE::Geometry::AppendError(Debug, EGeometryScriptErrorType::InvalidInputs, LOCTEXT("SelectHiddenTrianglesFromOutside_InvalidInput", "SelectHiddenTrianglesFromOutside: TargetMesh is Null"));
+		return TargetMesh;
+	}
+
+	TargetMesh->ProcessMesh([&Options, &IgnoreSelection, &TransparentSelection, &ResultSelection](const FDynamicMesh3& Mesh)
+		{
+			FExteriorVisibilitySampling SamplingParameters;
+			SamplingParameters.bDoubleSided = Options.bDoubleSided;
+			SamplingParameters.NumSearchDirections = Options.NumSearchDirections;
+			SamplingParameters.SamplingDensity = Options.SampleSpacing;
+			const bool bPerPolyGroup = Options.bPerPolyGroup;
+
+			TSet<int32> TransparentTriangles, IgnoreTriangles;
+			TransparentSelection.ProcessByTriangleID(Mesh, [&TransparentTriangles](int32 TriangleID)
+				{
+					TransparentTriangles.Add(TriangleID);
+				});
+			IgnoreSelection.ProcessByTriangleID(Mesh, [&IgnoreTriangles](int32 TriangleID)
+				{
+					IgnoreTriangles.Add(TriangleID);
+				});
+			auto IsTriTransparent = [&TransparentTriangles](int32 TID)
+				{
+					return TransparentTriangles.Contains(TID);
+				};
+			TArray<bool> SkipTris;
+			if (!IgnoreTriangles.IsEmpty())
+			{
+				SkipTris.SetNumZeroed(Mesh.MaxTriangleID());
+				for (int32 IgnoreTID : IgnoreTriangles)
+				{
+					SkipTris[IgnoreTID] = true;
+				}
+			}
+
+			TArray<bool> ResultTriOccluded;
+
+			if (!Options.RestrictToViewDirections.IsEmpty())
+			{
+				TArray<FVector3d> NormalizedRestrictDirections = Options.RestrictToViewDirections;
+				for (FVector3d& V : NormalizedRestrictDirections)
+				{
+					V = Normalized(V);
+				}
+
+				TArray<FVector3d> FullTestDirections = FExteriorVisibilitySampling::GetRestrictedViewDirections(NormalizedRestrictDirections, Options.ViewDirectionAngleRange, Options.NumSearchDirections);
+				FExteriorVisibilitySampling::ComputePerTriangleOcclusionFromDirections(Mesh, IsTriTransparent, ResultTriOccluded, bPerPolyGroup, NormalizedRestrictDirections, FullTestDirections, false, SamplingParameters, SkipTris);
+			}
+			else
+			{
+				FExteriorVisibilitySampling::ComputePerTriangleOcclusion(Mesh, IsTriTransparent, ResultTriOccluded, bPerPolyGroup, SamplingParameters, SkipTris);
+			}
+
+			
+			// transfer boolean array to a geometry script selection result
+			FGeometrySelection NewSelection;
+			for (int32 TID = 0; TID < ResultTriOccluded.Num(); ++TID)
+			{
+				if (ResultTriOccluded[TID])
+				{
+					NewSelection.Selection.Add(FGeoSelectionID::MeshTriangle(TID).Encoded());
+				}
+			}
+			ResultSelection.SetSelection(MoveTemp(NewSelection));
+		}
+	);
+
+	return TargetMesh;
+}
+
+
 
 UDynamicMesh* UGeometryScriptLibrary_MeshRepairFunctions::SplitMeshBowties(
 	UDynamicMesh* TargetMesh,
diff --git a/Engine/Plugins/Runtime/GeometryScripting/Source/GeometryScriptingCore/Public/GeometryScript/MeshRepairFunctions.h b/Engine/Plugins/Runtime/GeometryScripting/Source/GeometryScriptingCore/Public/GeometryScript/MeshRepairFunctions.h
--- a/Engine/Plugins/Runtime/GeometryScripting/Source/GeometryScriptingCore/Public/GeometryScript/MeshRepairFunctions.h
+++ b/Engine/Plugins/Runtime/GeometryScripting/Source/GeometryScriptingCore/Public/GeometryScript/MeshRepairFunctions.h
@@ -5,6 +5,7 @@
 #include "CoreMinimal.h"
 #include "Kismet/BlueprintFunctionLibrary.h"
 #include "GeometryScript/GeometryScriptTypes.h"
+#include "GeometryScript/GeometryScriptSelectionTypes.h"
 #include "MeshRepairFunctions.generated.h"
 
 #define UE_API GEOMETRYSCRIPTINGCORE_API
@@ -138,11 +139,43 @@ public:
 	UPROPERTY(EditAnywhere, BlueprintReadWrite, Category = Options)
 	float NormalOffset = 1e-6f;
 
-	/**  */
+	/** Whether to compact the resulting mesh after triangles are deleted */
 	UPROPERTY(EditAnywhere, BlueprintReadWrite, Category = Options)
 	bool bCompactResult = true;
 };
 
+USTRUCT(BlueprintType)
+struct GEOMETRYSCRIPTINGCORE_API FGeometryScriptSelectHiddenTrianglesFromOutsideOptions
+{
+	GENERATED_BODY()
+public:
+
+	// Approximate spacing between samples on triangle faces used for determining visibility
+	// Set a smaller value to increase the number of visibility samples per triangle, which can give more accurate visibility results
+	UPROPERTY(EditAnywhere, BlueprintReadWrite, Category = Options)
+	double SampleSpacing = 1.0;
+
+	// Whether to treat faces as double-sided when determining visibility
+	UPROPERTY(EditAnywhere, BlueprintReadWrite, Category = Options)
+	bool bDoubleSided = false;
+
+	// Number of directions to test for visibility. Note: RestrictToViewDirections will filter the view directions after they are sampled from a full sphere, so the actual number of sampled directions will be smaller if restricted.
+	UPROPERTY(EditAnywhere, BlueprintReadWrite, Category = Options)
+	int32 NumSearchDirections = 128;
+
+	// Whether to decide visibility on a per-polygroup basis, rather than per-triangle
+	UPROPERTY(EditAnywhere, BlueprintReadWrite, Category = Options)
+	bool bPerPolyGroup = false;
+
+	// If not empty, only allow views along these directions (or a direction within the ViewDirectionAngleRange cone of these directions)
+	// To be specific: A mesh sample P is visible from given direction D if a ray starting at P-D*(Mesh Diameter), cast in direction D, reaches P
+	UPROPERTY(EditAnywhere, BlueprintReadWrite, Category = Options)
+	TArray<FVector> RestrictToViewDirections;
+
+	// Field of view to use for the RestrictViewDirections (if any), in degrees
+	UPROPERTY(EditAnywhere, BlueprintReadWrite, Category = Options)
+	double ViewDirectionAngleRange = 45;
+};
 
 
 UENUM(BlueprintType)
@@ -263,13 +296,30 @@ public:
 	* Removes any triangles in the mesh that are not visible from the exterior view, under various definitions of "visible" and "outside"
 	* as specified by the Options.
 	*/
-	UFUNCTION(BlueprintCallable, Category = "GeometryScript|Repair", meta=(ScriptMethod))
+	UFUNCTION(BlueprintCallable, Category = "GeometryScript|Repair", meta=(ScriptMethod, KeyWords = "Jacket Occluded"))
 	static UE_API UPARAM(DisplayName = "Target Mesh") UDynamicMesh* 
 	RemoveHiddenTriangles(  
 		UDynamicMesh* TargetMesh, 
 		FGeometryScriptRemoveHiddenTrianglesOptions Options,
 		UGeometryScriptDebug* Debug = nullptr);
 
+	/*
+	* Select mesh triangles hidden from outside views, with optional filtering on view directions
+	*
+	* @param IgnoreSelection Optionally indicate parts of the mesh that should not be included in the ResultSelection, so do not need to be tested for occlusion
+	* @param TransparentSelection Optionally indicate parts of the mesh that should not be considered occluding, e.g. transparent triangles
+	* @param ResultSelection The fully occluded triangles of the target mesh (based on the options / transparencies, and excluding any IgnoreSelection triangles)
+	*/
+	UFUNCTION(BlueprintCallable, Category = "GeometryScript|Repair", meta = (ScriptMethod, KeyWords = "Remove Jacket Occluded"))
+	static UE_API UPARAM(DisplayName = "Target Mesh") UDynamicMesh* 
+	SelectHiddenTrianglesFromOutside(
+		UDynamicMesh* TargetMesh, 
+		FGeometryScriptSelectHiddenTrianglesFromOutsideOptions Options,
+		FGeometryScriptMeshSelection IgnoreSelection,
+		FGeometryScriptMeshSelection TransparentSelection,
+		FGeometryScriptMeshSelection& ResultSelection,
+		UGeometryScriptDebug* Debug = nullptr);
+
 	/**
 	* Splits Bowties in the mesh and/or the attributes.  A Bowtie is formed when a single vertex is connected to more than two boundary edges, 
 	* and splitting duplicates the shared vertex so each triangle will have a unique copy.
