Optimize Transfet skin weight node
Adding Vertices to smooth  within the flood distance has been multithreaded
From 18ms to 3.5ms

diff --git a/Engine/Plugins/Runtime/GeometryProcessing/Source/DynamicMesh/Private/Operations/SmoothBoneWeights.cpp b/Engine/Plugins/Runtime/GeometryProcessing/Source/DynamicMesh/Private/Operations/SmoothBoneWeights.cpp
--- a/Engine/Plugins/Runtime/GeometryProcessing/Source/DynamicMesh/Private/Operations/SmoothBoneWeights.cpp
+++ b/Engine/Plugins/Runtime/GeometryProcessing/Source/DynamicMesh/Private/Operations/SmoothBoneWeights.cpp
@@ -9,6 +9,8 @@
 #include "Math/UnrealMathUtility.h"
 #include "Parameterization/MeshLocalParam.h"
 #include "DynamicMesh/MeshNormals.h"
+#include "Tasks/Task.h"
+#include "Async/TaskGraphInterfaces.h"
 
 using namespace UE::AnimationCore;
 using namespace UE::Geometry;
@@ -286,41 +288,90 @@ bool FSmoothDynamicMeshVertexSkinWeights::SmoothWeightsAtVerticesWithinDistance(
 	// Vertices array
 	if (FloodFillUpToDistance > 0)
 	{
-		for (const int32 SeedVID : Vertices)
-		{	
-			if (Cancelled()) 
-			{
-				return false;
-			}
+		// Now this process is quite fast, so cancellation can be check at the beginning
+		if (Cancelled())
+		{
+			return false;
+		}
 
-			// If at least one neighboring vertex is not part of the set of vertices to smooth then we need to flood
-			bool bNeedToFlood = false;
-			for (const int32 NeighborVertexID : SourceMesh->VtxVerticesItr(SeedVID))
+		const int32 NumVertices = Vertices.Num();
+		const int32 NumTasks = FMath::Max(FMath::Min(FTaskGraphInterface::Get().GetNumWorkerThreads(), NumVertices), 1);
+		constexpr int32 MinVerticesByTask = 20;
+		const int32 VerticesByTask = FMath::Max(FMath::DivideAndRoundUp(NumVertices, NumTasks), MinVerticesByTask);
+		const int32 NumBatches = FMath::DivideAndRoundUp(NumVertices, VerticesByTask);
+		TArray<UE::Tasks::FTask> PendingTasks;
+		std::atomic<int32> NumAccessSet = 0;
+		for (int32 BatchIndex = 0; BatchIndex < NumBatches; BatchIndex++)
+		{
+			const int32 StartIndex = BatchIndex * VerticesByTask;
+			int32 EndIndex = (BatchIndex + 1) * VerticesByTask;
+			EndIndex = BatchIndex == NumBatches - 1 ? FMath::Min(NumVertices, EndIndex) : EndIndex;
+			UE::Tasks::FTask PendingTask = UE::Tasks::Launch(UE_SOURCE_LOCATION, [this, StartIndex, EndIndex, &Vertices, &VerticesToSmooth, FloodFillUpToDistance, FloodFillUpToDistanceSquared, &NumAccessSet]()
 			{
-				if (!VerticesToSmooth.Contains(NeighborVertexID) && 
-					DistanceSquared(SourceMesh->GetVertex(SeedVID), SourceMesh->GetVertex(NeighborVertexID)) < FloodFillUpToDistanceSquared)
+				for (int32 Index = StartIndex; Index < EndIndex; ++Index)
 				{
-					bNeedToFlood = true;
-					break;
+					const int32 SeedVID = Vertices[Index];
+					
+					// If at least one neighboring vertex is not part of the set of vertices to smooth then we need to flood
+					bool bNeedToFlood = false;
+					for (const int32 NeighborVertexID : SourceMesh->VtxVerticesItr(SeedVID))
+					{
+						bool bVerticesToSmoothContainsNeighborVertexID;
+						while (true)
+						{
+							// Make sure to read in the set only when we are not writing in it, INDEX_NONE is used when are are writing in the set
+							// Store the number of thread which are reading from the set in NumAccessSet
+							if (int32 NumAccessSetExpected = NumAccessSet.load(std::memory_order_relaxed); NumAccessSetExpected != INDEX_NONE)
+							{
+								// Increment the NumAccessSet to read
+								if (NumAccessSet.compare_exchange_weak(NumAccessSetExpected, NumAccessSetExpected+1, std::memory_order_relaxed))
+								{
+									bVerticesToSmoothContainsNeighborVertexID = VerticesToSmooth.Contains(NeighborVertexID);
+									NumAccessSet.fetch_sub(1, std::memory_order_relaxed); // Decrement the NumAccessSet to release
+									break;
+								}
+							}
+						}
+						if (!bVerticesToSmoothContainsNeighborVertexID &&
+							DistanceSquared(SourceMesh->GetVertex(SeedVID), SourceMesh->GetVertex(NeighborVertexID)) < FloodFillUpToDistanceSquared)
+						{
+							bNeedToFlood = true;
+							break;
+						}
+					}
+
+					if (bNeedToFlood)
+					{
+						FVector3d Normal = FMeshNormals::ComputeVertexNormal(*SourceMesh, SeedVID);
+						FFrame3d SeedFrame = SourceMesh->GetVertexFrame(SeedVID, false, &Normal);
+
+						TMeshLocalParam<FDynamicMesh3> Param(SourceMesh);
+						Param.ParamMode = ELocalParamTypes::ExponentialMapUpwindAvg;
+						Param.ComputeToMaxDistance(SeedVID, SeedFrame, FloodFillUpToDistance);
+
+						// Only points within FloodFillUpToDistance should have UVs set
+						TArray<int32> PointsWithinDistance;
+						Param.GetPointsWithUV(PointsWithinDistance);
+
+						while (true)
+						{
+							// If none set are reading from the set try to write new data and set the atomic to INDEX_NONE
+							if (int32 Expected = NumAccessSet.load(std::memory_order_relaxed); Expected == 0)
+							{
+								if (NumAccessSet.compare_exchange_weak(Expected, INDEX_NONE, std::memory_order_relaxed))
+								{
+									VerticesToSmooth.Append(PointsWithinDistance);
+									NumAccessSet.store(0, std::memory_order_relaxed); // Release the atomic for reading or writing
+									break;
+								}
+							}
+						}
+					}
 				}
-			}
-
-			if (bNeedToFlood)
-			{
-				FVector3d Normal = FMeshNormals::ComputeVertexNormal(*SourceMesh, SeedVID);
-				FFrame3d SeedFrame = SourceMesh->GetVertexFrame(SeedVID, false, &Normal);
-
-				TMeshLocalParam<FDynamicMesh3> Param(SourceMesh);
-				Param.ParamMode = ELocalParamTypes::ExponentialMapUpwindAvg;
-				Param.ComputeToMaxDistance(SeedVID, SeedFrame, FloodFillUpToDistance);
-
-				// Only points within FloodFillUpToDistance should have UVs set
-				TArray<int32> PointsWithinDistance;
-				Param.GetPointsWithUV(PointsWithinDistance);
-
-				VerticesToSmooth.Append(PointsWithinDistance);
-			}
+			});
+			PendingTasks.Add(PendingTask);
 		}
+		UE::Tasks::Wait(PendingTasks);
 	}
 
 	for (int32 Itr = 0; Itr < NumIterations; ++Itr)
