Optimize Laplacian Matrix construction
Optimize from 16ms to 3.8ms

diff --git a/Engine/Plugins/Runtime/GeometryProcessing/Source/DynamicMesh/Private/Operations/TransferBoneWeights.cpp b/Engine/Plugins/Runtime/GeometryProcessing/Source/DynamicMesh/Private/Operations/TransferBoneWeights.cpp
--- a/Engine/Plugins/Runtime/GeometryProcessing/Source/DynamicMesh/Private/Operations/TransferBoneWeights.cpp
+++ b/Engine/Plugins/Runtime/GeometryProcessing/Source/DynamicMesh/Private/Operations/TransferBoneWeights.cpp
@@ -495,32 +495,56 @@ bool FTransferBoneWeights::TransferWeightsToMesh(FDynamicMesh3& InOutTargetMesh,
 		FEigenSparseMatrixAssembler CotangentAssembler(NumVerts, NumVerts);
 		FEigenSparseMatrixAssembler LaplacianAssembler(NumVerts, NumVerts);
 
+		UE::Tasks::FTask CotangentTask, LaplacianTask;
 		if (bUseIntrinsicLaplacian)
 		{
-			// Construct the Cotangent weights matrix
-			UE::MeshDeformation::ConstructFullIDTCotangentLaplacian<double>(InOutTargetMesh, VtxLinearization, CotangentAssembler,
-																			UE::MeshDeformation::ECotangentWeightMode::Default, 
-																			UE::MeshDeformation::ECotangentAreaMode::NoArea);
-
-			// Construct the Laplacian with cotangent weights scaled by the voronoi area (i.e. M^(-1)*L matrix where M is the mass/stiffness matrix)
-			UE::MeshDeformation::ConstructFullIDTCotangentLaplacian<double>(InOutTargetMesh, VtxLinearization, LaplacianAssembler,
-																			UE::MeshDeformation::ECotangentWeightMode::Default,
-																			UE::MeshDeformation::ECotangentAreaMode::VoronoiArea);
+			CotangentTask = UE::Tasks::Launch(UE_SOURCE_LOCATION, [&InOutTargetMesh, &VtxLinearization, &CotangentAssembler]()
+			{
+				// Construct the Cotangent weights matrix
+				UE::MeshDeformation::ConstructFullIDTCotangentLaplacian<double>(InOutTargetMesh, VtxLinearization, CotangentAssembler,
+					UE::MeshDeformation::ECotangentWeightMode::Default,
+					UE::MeshDeformation::ECotangentAreaMode::NoArea);
+			});
+
+			
+			LaplacianTask = UE::Tasks::Launch(UE_SOURCE_LOCATION, [&InOutTargetMesh, &VtxLinearization, &LaplacianAssembler]()
+			{
+				// Construct the Laplacian with cotangent weights scaled by the voronoi area (i.e. M^(-1)*L matrix where M is the mass/stiffness matrix)
+				UE::MeshDeformation::ConstructFullIDTCotangentLaplacian<double>(InOutTargetMesh, VtxLinearization, LaplacianAssembler,
+																				UE::MeshDeformation::ECotangentWeightMode::Default,
+																				UE::MeshDeformation::ECotangentAreaMode::VoronoiArea);
+			});
 		}
 		else 
 		{
-			UE::MeshDeformation::ConstructFullCotangentLaplacian<double>(InOutTargetMesh, VtxLinearization, CotangentAssembler,
-																		 UE::MeshDeformation::ECotangentWeightMode::Default,
-																		 UE::MeshDeformation::ECotangentAreaMode::NoArea);
-
-			UE::MeshDeformation::ConstructFullCotangentLaplacian<double>(InOutTargetMesh, VtxLinearization, LaplacianAssembler,
-																		 UE::MeshDeformation::ECotangentWeightMode::Default,
-																		 UE::MeshDeformation::ECotangentAreaMode::VoronoiArea);
+			CotangentTask = UE::Tasks::Launch(UE_SOURCE_LOCATION, [&InOutTargetMesh, &VtxLinearization, &CotangentAssembler]()
+			{
+				UE::MeshDeformation::ConstructFullCotangentLaplacian<double>(InOutTargetMesh, VtxLinearization, CotangentAssembler,
+					UE::MeshDeformation::ECotangentWeightMode::Default,
+					UE::MeshDeformation::ECotangentAreaMode::NoArea);
+			});
+			
+			LaplacianTask = UE::Tasks::Launch(UE_SOURCE_LOCATION, [&InOutTargetMesh, &VtxLinearization, &LaplacianAssembler]()
+			{
+				UE::MeshDeformation::ConstructFullCotangentLaplacian<double>(InOutTargetMesh, VtxLinearization, LaplacianAssembler,
+					UE::MeshDeformation::ECotangentWeightMode::Default,
+					UE::MeshDeformation::ECotangentAreaMode::VoronoiArea);
+			});
 		}
 
 		FSparseMatrixD CotangentMatrix, MassCotangentMatrix;
-		CotangentAssembler.ExtractResult(CotangentMatrix);
-		LaplacianAssembler.ExtractResult(MassCotangentMatrix);
+		UE::Tasks::FTask CotangentExtractTask = UE::Tasks::Launch(UE_SOURCE_LOCATION, [&CotangentMatrix, &CotangentAssembler]()
+		{
+			CotangentAssembler.ExtractResult(CotangentMatrix);
+		}, CotangentTask);
+
+		UE::Tasks::FTask LaplacianExtractTask = UE::Tasks::Launch(UE_SOURCE_LOCATION, [&MassCotangentMatrix, &LaplacianAssembler]()
+		{
+			LaplacianAssembler.ExtractResult(MassCotangentMatrix);
+		}, LaplacianTask);
+
+		TArray<UE::Tasks::FTask> PendingTasks = { CotangentExtractTask, LaplacianExtractTask };
+		UE::Tasks::Wait(PendingTasks);
 
 		// -L * L* M^(-1)*L energy
 		FSparseMatrixD Energy = -1*CotangentMatrix + CotangentMatrix*MassCotangentMatrix;
diff --git a/Engine/Plugins/Runtime/GeometryProcessing/Source/DynamicMesh/Private/Solvers/Internal/SparseMatrix.h b/Engine/Plugins/Runtime/GeometryProcessing/Source/DynamicMesh/Private/Solvers/Internal/SparseMatrix.h
--- a/Engine/Plugins/Runtime/GeometryProcessing/Source/DynamicMesh/Private/Solvers/Internal/SparseMatrix.h
+++ b/Engine/Plugins/Runtime/GeometryProcessing/Source/DynamicMesh/Private/Solvers/Internal/SparseMatrix.h
@@ -82,6 +82,14 @@ public:
 		{
 			EntryTriplets.push_back(MatrixTripletT(i, j, Value));
 		};
+
+		AddEntriesFunc = [this](const FTupleData* Data, int32 NumData)
+		{
+			const MatrixTripletT* MatrixTriplet = reinterpret_cast<const MatrixTripletT*>(Data);
+			static_assert(sizeof(MatrixTripletT) == sizeof(FTupleData));
+
+			EntryTriplets.insert(EntryTriplets.end(), MatrixTriplet, MatrixTriplet + NumData);
+		};
 	}
 
 	void ExtractResult(FSparseMatrixD& Result)
diff --git a/Engine/Plugins/Runtime/GeometryProcessing/Source/DynamicMesh/Public/Solvers/LaplacianMatrixAssembly.h b/Engine/Plugins/Runtime/GeometryProcessing/Source/DynamicMesh/Public/Solvers/LaplacianMatrixAssembly.h
--- a/Engine/Plugins/Runtime/GeometryProcessing/Source/DynamicMesh/Public/Solvers/LaplacianMatrixAssembly.h
+++ b/Engine/Plugins/Runtime/GeometryProcessing/Source/DynamicMesh/Public/Solvers/LaplacianMatrixAssembly.h
@@ -851,8 +851,6 @@ void UE::MeshDeformation::ConstructIDTCotangentLaplacian(const FDynamicMesh3& Dy
 
 
 
-
-
 template<typename RealType>
 void UE::MeshDeformation::ConstructFullCotangentLaplacian(const FDynamicMesh3& Mesh, const FVertexLinearization& VertexMap,
 	UE::Solvers::TSparseMatrixAssembler<RealType>& LaplacianMatrix, 
@@ -863,81 +861,120 @@ void UE::MeshDeformation::ConstructFullCotangentLaplacian(const FDynamicMesh3& M
 	const TArray<int32>& ToIndex = VertexMap.ToIndex();
 	const int32 NumVerts = VertexMap.NumVerts();
 
-	// pre-allocate space when possible
-	int32 NumMatrixEntries = ComputeNumMatrixElements(Mesh, ToMeshV);
-	LaplacianMatrix.ReserveEntriesFunc(NumMatrixEntries);
 
 	// Map the triangles.
 	FTriangleLinearization TriangleMap(Mesh);
-	const TArray<int32>& ToMeshTri = TriangleMap.ToId();
 	const TArray<int32>& ToTriIdx = TriangleMap.ToIndex();
 	const int32 NumTris = TriangleMap.NumTris();
 
+	const int32 NumTasks = FMath::Clamp(FTaskGraphInterface::Get().GetNumWorkerThreads()/2, 1, NumVerts);
+	constexpr int32 MinVerticesByTask = 60;
+	const int32 VerticesByTask = FMath::Max(FMath::DivideAndRoundUp(NumVerts, NumTasks), MinVerticesByTask);
+	const int32 NumBatches = FMath::DivideAndRoundUp(NumVerts, VerticesByTask);
+	TArray<UE::Tasks::FTask> PendingTasks;
+	PendingTasks.Reserve(NumBatches+1);
+
+	UE::Tasks::FTask FirstPendingTask = UE::Tasks::Launch(UE_SOURCE_LOCATION, [&Mesh, &ToMeshV, &LaplacianMatrix]()
+	{
+		// pre-allocate space when possible
+		const int32 NumMatrixEntries = ComputeNumMatrixElements(Mesh, ToMeshV);
+		LaplacianMatrix.ReserveEntriesFunc(NumMatrixEntries);
+	});
+	PendingTasks.Add(FirstPendingTask);
+
 	// Create an array that holds all the geometric information we need for each triangle.
 	TArray<CotanTriangleData> CotangentTriangleDataArray;
 	ConstructTriangleDataArray<CotanTriangleData>(Mesh, TriangleMap, CotangentTriangleDataArray);
 
-	// Construct Laplacian Matrix: loop over verts constructing the corresponding matrix row.
-	for (int32 i = 0; i < NumVerts; ++i)
+
+	using MatrixEntry = UE::Solvers::TSparseMatrixAssembler<RealType>::FTupleData;
+	TArray<TArray<MatrixEntry>> TriplesBatches;
+	TriplesBatches.SetNum(NumBatches);
+
+	for (int32 BatchIndex = 0; BatchIndex < NumBatches; BatchIndex++)
 	{
-		const int32 IVertId = ToMeshV[i]; // I - the row
+		const int32 StartIndex = BatchIndex * VerticesByTask;
+		int32 EndIndex = (BatchIndex + 1) * VerticesByTask;
+		EndIndex = BatchIndex == NumBatches - 1 ? FMath::Min(NumVerts, EndIndex) : EndIndex;
 
-		double WeightArea = 1.0;
-		if (AreaMode == ECotangentAreaMode::VoronoiArea)
+		TArray<MatrixEntry>& TriplesBatch = TriplesBatches[BatchIndex];
+
+		UE::Tasks::FTask PendingTask = UE::Tasks::Launch(UE_SOURCE_LOCATION, [StartIndex, EndIndex, &ToMeshV, &ToIndex, &ToTriIdx,
+			&CotangentTriangleDataArray, &Mesh, AreaMode, WeightMode, &TriplesBatch, &LaplacianMatrix]()
 		{
-			WeightArea = 0.0;
-			Mesh.EnumerateVertexTriangles(IVertId, [&](int32 TriId) {
-				const int32 TriIdx = ToTriIdx[TriId];
-				const CotanTriangleData& TriData = CotangentTriangleDataArray[TriIdx];
-				const FIndex3i TriVertIds = Mesh.GetTriangle(TriId);
-				int32 TriVertIndex = IndexUtil::FindTriIndex(IVertId, TriVertIds);
-				WeightArea += TriData.VoronoiArea[TriVertIndex];
-			});
-		}
+			// Construct Laplacian Matrix: loop over verts constructing the corresponding matrix row.
+			for (int32 i = StartIndex; i < EndIndex; ++i)
+			{
+				const int32 IVertId = ToMeshV[i]; // I - the row
 
-		double WeightII = 0.; // accumulate to equal and opposite the sum of the neighbor weights
+				double WeightArea = 1.0;
+				if (AreaMode == ECotangentAreaMode::VoronoiArea)
+				{
+					WeightArea = 0.0;
+					Mesh.EnumerateVertexTriangles(IVertId, [&](int32 TriId) {
+						const int32 TriIdx = ToTriIdx[TriId];
+						const CotanTriangleData& TriData = CotangentTriangleDataArray[TriIdx];
+						const FIndex3i TriVertIds = Mesh.GetTriangle(TriId);
+						int32 TriVertIndex = IndexUtil::FindTriIndex(IVertId, TriVertIds);
+						WeightArea += TriData.VoronoiArea[TriVertIndex];
+						});
+				}
 
-		for (int32 EdgeId : Mesh.VtxEdgesItr(IVertId))
-		{
-			const FDynamicMesh3::FEdge Edge = Mesh.GetEdge(EdgeId);
+				double WeightII = 0.; // accumulate to equal and opposite the sum of the neighbor weights
 
-			// the other vert in the edge - identifies the matrix column
-			const int32 JVertId = (Edge.Vert[0] == IVertId) ? Edge.Vert[1] : Edge.Vert[0];  // J - the column
+				for (int32 EdgeId : Mesh.VtxEdgesItr(IVertId))
+				{
+					const FDynamicMesh3::FEdge Edge = Mesh.GetEdge(EdgeId);
 
-			// Get the cotangents for this edge.
-			const int32 Tri0Idx = ToTriIdx[Edge.Tri[0]];
-			const CotanTriangleData& Tri0Data = CotangentTriangleDataArray[Tri0Idx];
-			double CotanAlpha = Tri0Data.GetOpposingCotangent(EdgeId);
+					// the other vert in the edge - identifies the matrix column
+					const int32 JVertId = (Edge.Vert[0] == IVertId) ? Edge.Vert[1] : Edge.Vert[0];  // J - the column
 
-			// The second triangle will be invalid if this is an edge!
-			double CotanBeta = (Edge.Tri[1] != FDynamicMesh3::InvalidID) ? CotangentTriangleDataArray[ToTriIdx[Edge.Tri[1]]].GetOpposingCotangent(EdgeId) : 0.0;
+					// Get the cotangents for this edge.
+					const int32 Tri0Idx = ToTriIdx[Edge.Tri[0]];
+					const CotanTriangleData& Tri0Data = CotangentTriangleDataArray[Tri0Idx];
+					double CotanAlpha = Tri0Data.GetOpposingCotangent(EdgeId);
 
-			if (WeightMode == ECotangentWeightMode::TriangleArea) 
-			{
-				CotanAlpha /= Tri0Data.Area;
+					// The second triangle will be invalid if this is an edge!
+					double CotanBeta = (Edge.Tri[1] != FDynamicMesh3::InvalidID) ? CotangentTriangleDataArray[ToTriIdx[Edge.Tri[1]]].GetOpposingCotangent(EdgeId) : 0.0;
 
-				if (Edge.Tri[1] != FDynamicMesh3::InvalidID)
-				{
-					const CotanTriangleData& Tri1Data = CotangentTriangleDataArray[ToTriIdx[Edge.Tri[1]]];
-					CotanBeta /= Tri1Data.Area; 
-				}
-			}
+					if (WeightMode == ECotangentWeightMode::TriangleArea)
+					{
+						CotanAlpha /= Tri0Data.Area;
 
-			// do not need to multiply by 0.5 here...
-			//double WeightIJ = 0.5 * (CotanAlpha + CotanBeta);
-			double WeightIJ = (CotanAlpha + CotanBeta);
+						if (Edge.Tri[1] != FDynamicMesh3::InvalidID)
+						{
+							const CotanTriangleData& Tri1Data = CotangentTriangleDataArray[ToTriIdx[Edge.Tri[1]]];
+							CotanBeta /= Tri1Data.Area;
+						}
+					}
 
-			if (WeightMode == ECotangentWeightMode::ClampedMagnitude)
-			{
-				WeightIJ = FMathd::Clamp(WeightIJ, -1.e5 * WeightArea, 1.e5 * WeightArea);
+					// do not need to multiply by 0.5 here...
+					//double WeightIJ = 0.5 * (CotanAlpha + CotanBeta);
+					double WeightIJ = (CotanAlpha + CotanBeta);
+
+					if (WeightMode == ECotangentWeightMode::ClampedMagnitude)
+					{
+						WeightIJ = FMathd::Clamp(WeightIJ, -1.e5 * WeightArea, 1.e5 * WeightArea);
+					}
+
+					WeightII += WeightIJ;
+
+					const int32 j = ToIndex[JVertId];
+
+					TriplesBatch.Emplace( i, j, (RealType)(WeightIJ / WeightArea));
+				}
+				TriplesBatch.Emplace(i, i, (RealType)(-WeightII / WeightArea));
 			}
+		});
+		PendingTasks.Add(PendingTask);
+	}
 
-			WeightII += WeightIJ;
+	UE::Tasks::Wait(PendingTasks);
 
-			const int32 j = ToIndex[JVertId];
-			LaplacianMatrix.AddEntryFunc(i, j, (RealType)(WeightIJ / WeightArea));
-		}
-		LaplacianMatrix.AddEntryFunc(i, i, (RealType)(-WeightII / WeightArea));
+	for (int32 BatchIndex = 0; BatchIndex < NumBatches; BatchIndex++)
+	{
+		TArray <MatrixEntry>& Triplets = TriplesBatches[BatchIndex];
+		LaplacianMatrix.AddEntriesFunc(Triplets.GetData(), Triplets.Num());
 	}
 }
 
diff --git a/Engine/Plugins/Runtime/GeometryProcessing/Source/DynamicMesh/Public/Solvers/PrecomputedMeshWeightData.h b/Engine/Plugins/Runtime/GeometryProcessing/Source/DynamicMesh/Public/Solvers/PrecomputedMeshWeightData.h
--- a/Engine/Plugins/Runtime/GeometryProcessing/Source/DynamicMesh/Public/Solvers/PrecomputedMeshWeightData.h
+++ b/Engine/Plugins/Runtime/GeometryProcessing/Source/DynamicMesh/Public/Solvers/PrecomputedMeshWeightData.h
@@ -4,6 +4,8 @@
 
 #include "DynamicMesh/DynamicMesh3.h"
 #include "Solvers/MeshLinearization.h"
+#include "Tasks/Task.h"
+#include "Async/TaskGraphInterfaces.h"
 
 #define UE_API DYNAMICMESH_API
 
@@ -210,16 +212,32 @@ namespace UE
 
 			const auto& ToTriId = TriangleLinearization.ToId();
 
-			for (int32 i = 0; i < NumTris; ++i)
+			const int32 NumTasks = FMath::Max(FMath::Min(FTaskGraphInterface::Get().GetNumWorkerThreads()/2, NumTris), 1);
+			constexpr int32 MinTrianglesByTask = 60;
+			const int32 TrianglesByTask = FMath::Max(FMath::Max(FMath::DivideAndRoundUp(NumTris, NumTasks), MinTrianglesByTask), 1);
+			const int32 NumBatches = FMath::DivideAndRoundUp(NumTris, TrianglesByTask);
+			TArray<UE::Tasks::FTask> PendingTasks;
+			PendingTasks.Reserve(NumBatches);
+			for (int32 BatchIndex = 0; BatchIndex < NumBatches; BatchIndex++)
 			{
-				// Current triangle
+				const int32 StartIndex = BatchIndex * TrianglesByTask;
+				int32 EndIndex = (BatchIndex + 1) * TrianglesByTask;
+				EndIndex = BatchIndex == NumBatches - 1 ? FMath::Min(NumTris, EndIndex) : EndIndex;
 
-				const int32 TriId = ToTriId[i];
-
-				// Compute all the geometric data needed for this triangle.
-
-				TriangleDataArray[i].Initialize(DynamicMesh, TriId);
+				UE::Tasks::FTask PendingTask = UE::Tasks::Launch(UE_SOURCE_LOCATION, [StartIndex, EndIndex, &ToTriId, &TriangleDataArray, &DynamicMesh]()
+				{
+					for (int32 i = StartIndex; i < EndIndex; ++i)
+					{
+						// Current triangle
+						const int32 TriId = ToTriId[i];
+
+						// Compute all the geometric data needed for this triangle.
+						TriangleDataArray[i].Initialize(DynamicMesh, TriId);
+					}
+				});
+				PendingTasks.Add(PendingTask);
 			}
+			UE::Tasks::Wait(PendingTasks);
 		}
 
 
diff --git a/Engine/Source/Runtime/GeometryCore/Public/Solvers/MatrixInterfaces.h b/Engine/Source/Runtime/GeometryCore/Public/Solvers/MatrixInterfaces.h
--- a/Engine/Source/Runtime/GeometryCore/Public/Solvers/MatrixInterfaces.h
+++ b/Engine/Source/Runtime/GeometryCore/Public/Solvers/MatrixInterfaces.h
@@ -25,6 +25,16 @@ namespace UE
 
 			/** Add matrix entry tuple (i,j,Value) */
 			TUniqueFunction<void(int32, int32, RealType)> AddEntryFunc;
+
+			struct FTupleData
+			{
+				int32 I;
+				int32 J;
+				RealType Value;
+			};
+
+			/** Add matrix entries (i,j,Value) and size of tuples */
+			TUniqueFunction<void(const FTupleData*, int32)> AddEntriesFunc;
 		};
 
 		typedef TSparseMatrixAssembler<double> FSparseMatrixAssemblerd;
