Fix build : fall out from CL 42399822

diff --git a/Engine/Plugins/Runtime/GeometryProcessing/Source/DynamicMesh/Public/Solvers/LaplacianMatrixAssembly.h b/Engine/Plugins/Runtime/GeometryProcessing/Source/DynamicMesh/Public/Solvers/LaplacianMatrixAssembly.h
--- a/Engine/Plugins/Runtime/GeometryProcessing/Source/DynamicMesh/Public/Solvers/LaplacianMatrixAssembly.h
+++ b/Engine/Plugins/Runtime/GeometryProcessing/Source/DynamicMesh/Public/Solvers/LaplacianMatrixAssembly.h
@@ -887,7 +887,7 @@ void UE::MeshDeformation::ConstructFullCotangentLaplacian(const FDynamicMesh3& M
 	ConstructTriangleDataArray<CotanTriangleData>(Mesh, TriangleMap, CotangentTriangleDataArray);
 
 
-	using MatrixEntry = UE::Solvers::TSparseMatrixAssembler<RealType>::FTupleData;
+	using MatrixEntry = typename UE::Solvers::TSparseMatrixAssembler<RealType>::FTupleData;
 	TArray<TArray<MatrixEntry>> TriplesBatches;
 	TriplesBatches.SetNum(NumBatches);
 
