Fix Android build fallout from CL  42399822

diff --git a/Engine/Plugins/Runtime/GeometryProcessing/Source/DynamicMesh/Public/Solvers/LaplacianMatrixAssembly.h b/Engine/Plugins/Runtime/GeometryProcessing/Source/DynamicMesh/Public/Solvers/LaplacianMatrixAssembly.h
--- a/Engine/Plugins/Runtime/GeometryProcessing/Source/DynamicMesh/Public/Solvers/LaplacianMatrixAssembly.h
+++ b/Engine/Plugins/Runtime/GeometryProcessing/Source/DynamicMesh/Public/Solvers/LaplacianMatrixAssembly.h
@@ -961,9 +961,9 @@ void UE::MeshDeformation::ConstructFullCotangentLaplacian(const FDynamicMesh3& M
 
 					const int32 j = ToIndex[JVertId];
 
-					TriplesBatch.Emplace( i, j, (RealType)(WeightIJ / WeightArea));
+					TriplesBatch.Emplace(MatrixEntry(i, j, (RealType)(WeightIJ / WeightArea)));
 				}
-				TriplesBatch.Emplace(i, i, (RealType)(-WeightII / WeightArea));
+				TriplesBatch.Emplace(MatrixEntry(i, i, (RealType)(-WeightII / WeightArea)));
 			}
 		});
 		PendingTasks.Add(PendingTask);
