[Backout] - CL42399822
#rnx
Original CL Desc
-----------------------------------------------------------------
Optimize Laplacian Matrix construction
Optimize from 16ms to 3.8ms

diff --git a/Engine/Plugins/Runtime/GeometryProcessing/Source/DynamicMesh/Private/Operations/TransferBoneWeights.cpp b/Engine/Plugins/Runtime/GeometryProcessing/Source/DynamicMesh/Private/Operations/TransferBoneWeights.cpp
--- a/Engine/Plugins/Runtime/GeometryProcessing/Source/DynamicMesh/Private/Operations/TransferBoneWeights.cpp
+++ b/Engine/Plugins/Runtime/GeometryProcessing/Source/DynamicMesh/Private/Operations/TransferBoneWeights.cpp
@@ -495,56 +495,32 @@ bool FTransferBoneWeights::TransferWeightsToMesh(FDynamicMesh3& InOutTargetMesh,
 		FEigenSparseMatrixAssembler CotangentAssembler(NumVerts, NumVerts);
 		FEigenSparseMatrixAssembler LaplacianAssembler(NumVerts, NumVerts);
 
-		UE::Tasks::FTask CotangentTask, LaplacianTask;
 		if (bUseIntrinsicLaplacian)
 		{
-			CotangentTask = UE::Tasks::Launch(UE_SOURCE_LOCATION, [&InOutTargetMesh, &VtxLinearization, &CotangentAssembler]()
-			{
-				// Construct the Cotangent weights matrix
-				UE::MeshDeformation::ConstructFullIDTCotangentLaplacian<double>(InOutTargetMesh, VtxLinearization, CotangentAssembler,
-					UE::MeshDeformation::ECotangentWeightMode::Default,
-					UE::MeshDeformation::ECotangentAreaMode::NoArea);
-			});
-
-			
-			LaplacianTask = UE::Tasks::Launch(UE_SOURCE_LOCATION, [&InOutTargetMesh, &VtxLinearization, &LaplacianAssembler]()
-			{
-				// Construct the Laplacian with cotangent weights scaled by the voronoi area (i.e. M^(-1)*L matrix where M is the mass/stiffness matrix)
-				UE::MeshDeformation::ConstructFullIDTCotangentLaplacian<double>(InOutTargetMesh, VtxLinearization, LaplacianAssembler,
-																				UE::MeshDeformation::ECotangentWeightMode::Default,
-																				UE::MeshDeformation::ECotangentAreaMode::VoronoiArea);
-			});
+			// Construct the Cotangent weights matrix
+			UE::MeshDeformation::ConstructFullIDTCotangentLaplacian<double>(InOutTargetMesh, VtxLinearization, CotangentAssembler,
+																			UE::MeshDeformation::ECotangentWeightMode::Default, 
+																			UE::MeshDeformation::ECotangentAreaMode::NoArea);
+
+			// Construct the Laplacian with cotangent weights scaled by the voronoi area (i.e. M^(-1)*L matrix where M is the mass/stiffness matrix)
+			UE::MeshDeformation::ConstructFullIDTCotangentLaplacian<double>(InOutTargetMesh, VtxLinearization, LaplacianAssembler,
+																			UE::MeshDeformation::ECotangentWeightMode::Default,
+																			UE::MeshDeformation::ECotangentAreaMode::VoronoiArea);
 		}
 		else 
 		{
-			CotangentTask = UE::Tasks::Launch(UE_SOURCE_LOCATION, [&InOutTargetMesh, &VtxLinearization, &CotangentAssembler]()
-			{
-				UE::MeshDeformation::ConstructFullCotangentLaplacian<double>(InOutTargetMesh, VtxLinearization, CotangentAssembler,
-					UE::MeshDeformation::ECotangentWeightMode::Default,
-					UE::MeshDeformation::ECotangentAreaMode::NoArea);
-			});
-			
-			LaplacianTask = UE::Tasks::Launch(UE_SOURCE_LOCATION, [&InOutTargetMesh, &VtxLinearization, &LaplacianAssembler]()
-			{
-				UE::MeshDeformation::ConstructFullCotangentLaplacian<double>(InOutTargetMesh, VtxLinearization, LaplacianAssembler,
-					UE::MeshDeformation::ECotangentWeightMode::Default,
-					UE::MeshDeformation::ECotangentAreaMode::VoronoiArea);
-			});
+			UE::MeshDeformation::ConstructFullCotangentLaplacian<double>(InOutTargetMesh, VtxLinearization, CotangentAssembler,
+																		 UE::MeshDeformation::ECotangentWeightMode::Default,
+																		 UE::MeshDeformation::ECotangentAreaMode::NoArea);
+
+			UE::MeshDeformation::ConstructFullCotangentLaplacian<double>(InOutTargetMesh, VtxLinearization, LaplacianAssembler,
+																		 UE::MeshDeformation::ECotangentWeightMode::Default,
+																		 UE::MeshDeformation::ECotangentAreaMode::VoronoiArea);
 		}
 
 		FSparseMatrixD CotangentMatrix, MassCotangentMatrix;
-		UE::Tasks::FTask CotangentExtractTask = UE::Tasks::Launch(UE_SOURCE_LOCATION, [&CotangentMatrix, &CotangentAssembler]()
-		{
-			CotangentAssembler.ExtractResult(CotangentMatrix);
-		}, CotangentTask);
-
-		UE::Tasks::FTask LaplacianExtractTask = UE::Tasks::Launch(UE_SOURCE_LOCATION, [&MassCotangentMatrix, &LaplacianAssembler]()
-		{
-			LaplacianAssembler.ExtractResult(MassCotangentMatrix);
-		}, LaplacianTask);
-
-		TArray<UE::Tasks::FTask> PendingTasks = { CotangentExtractTask, LaplacianExtractTask };
-		UE::Tasks::Wait(PendingTasks);
+		CotangentAssembler.ExtractResult(CotangentMatrix);
+		LaplacianAssembler.ExtractResult(MassCotangentMatrix);
 
 		// -L * L* M^(-1)*L energy
 		FSparseMatrixD Energy = -1*CotangentMatrix + CotangentMatrix*MassCotangentMatrix;
diff --git a/Engine/Plugins/Runtime/GeometryProcessing/Source/DynamicMesh/Private/Solvers/Internal/SparseMatrix.h b/Engine/Plugins/Runtime/GeometryProcessing/Source/DynamicMesh/Private/Solvers/Internal/SparseMatrix.h
--- a/Engine/Plugins/Runtime/GeometryProcessing/Source/DynamicMesh/Private/Solvers/Internal/SparseMatrix.h
+++ b/Engine/Plugins/Runtime/GeometryProcessing/Source/DynamicMesh/Private/Solvers/Internal/SparseMatrix.h
@@ -82,14 +82,6 @@ public:
 		{
 			EntryTriplets.push_back(MatrixTripletT(i, j, Value));
 		};
-
-		AddEntriesFunc = [this](const FTupleData* Data, int32 NumData)
-		{
-			const MatrixTripletT* MatrixTriplet = reinterpret_cast<const MatrixTripletT*>(Data);
-			static_assert(sizeof(MatrixTripletT) == sizeof(FTupleData));
-
-			EntryTriplets.insert(EntryTriplets.end(), MatrixTriplet, MatrixTriplet + NumData);
-		};
 	}
 
 	void ExtractResult(FSparseMatrixD& Result)
diff --git a/Engine/Plugins/Runtime/GeometryProcessing/Source/DynamicMesh/Public/Solvers/LaplacianMatrixAssembly.h b/Engine/Plugins/Runtime/GeometryProcessing/Source/DynamicMesh/Public/Solvers/LaplacianMatrixAssembly.h
--- a/Engine/Plugins/Runtime/GeometryProcessing/Source/DynamicMesh/Public/Solvers/LaplacianMatrixAssembly.h
+++ b/Engine/Plugins/Runtime/GeometryProcessing/Source/DynamicMesh/Public/Solvers/LaplacianMatrixAssembly.h
@@ -851,6 +851,8 @@ void UE::MeshDeformation::ConstructIDTCotangentLaplacian(const FDynamicMesh3& Dy
 
 
 
+
+
 template<typename RealType>
 void UE::MeshDeformation::ConstructFullCotangentLaplacian(const FDynamicMesh3& Mesh, const FVertexLinearization& VertexMap,
 	UE::Solvers::TSparseMatrixAssembler<RealType>& LaplacianMatrix, 
@@ -861,120 +863,81 @@ void UE::MeshDeformation::ConstructFullCotangentLaplacian(const FDynamicMesh3& M
 	const TArray<int32>& ToIndex = VertexMap.ToIndex();
 	const int32 NumVerts = VertexMap.NumVerts();
 
+	// pre-allocate space when possible
+	int32 NumMatrixEntries = ComputeNumMatrixElements(Mesh, ToMeshV);
+	LaplacianMatrix.ReserveEntriesFunc(NumMatrixEntries);
 
 	// Map the triangles.
 	FTriangleLinearization TriangleMap(Mesh);
+	const TArray<int32>& ToMeshTri = TriangleMap.ToId();
 	const TArray<int32>& ToTriIdx = TriangleMap.ToIndex();
 	const int32 NumTris = TriangleMap.NumTris();
 
-	const int32 NumTasks = FMath::Clamp(FTaskGraphInterface::Get().GetNumWorkerThreads()/2, 1, NumVerts);
-	constexpr int32 MinVerticesByTask = 60;
-	const int32 VerticesByTask = FMath::Max(FMath::DivideAndRoundUp(NumVerts, NumTasks), MinVerticesByTask);
-	const int32 NumBatches = FMath::DivideAndRoundUp(NumVerts, VerticesByTask);
-	TArray<UE::Tasks::FTask> PendingTasks;
-	PendingTasks.Reserve(NumBatches+1);
-
-	UE::Tasks::FTask FirstPendingTask = UE::Tasks::Launch(UE_SOURCE_LOCATION, [&Mesh, &ToMeshV, &LaplacianMatrix]()
-	{
-		// pre-allocate space when possible
-		const int32 NumMatrixEntries = ComputeNumMatrixElements(Mesh, ToMeshV);
-		LaplacianMatrix.ReserveEntriesFunc(NumMatrixEntries);
-	});
-	PendingTasks.Add(FirstPendingTask);
-
 	// Create an array that holds all the geometric information we need for each triangle.
 	TArray<CotanTriangleData> CotangentTriangleDataArray;
 	ConstructTriangleDataArray<CotanTriangleData>(Mesh, TriangleMap, CotangentTriangleDataArray);
 
-
-	using MatrixEntry = typename UE::Solvers::TSparseMatrixAssembler<RealType>::FTupleData;
-	TArray<TArray<MatrixEntry>> TriplesBatches;
-	TriplesBatches.SetNum(NumBatches);
-
-	for (int32 BatchIndex = 0; BatchIndex < NumBatches; BatchIndex++)
+	// Construct Laplacian Matrix: loop over verts constructing the corresponding matrix row.
+	for (int32 i = 0; i < NumVerts; ++i)
 	{
-		const int32 StartIndex = BatchIndex * VerticesByTask;
-		int32 EndIndex = (BatchIndex + 1) * VerticesByTask;
-		EndIndex = BatchIndex == NumBatches - 1 ? FMath::Min(NumVerts, EndIndex) : EndIndex;
-
-		TArray<MatrixEntry>& TriplesBatch = TriplesBatches[BatchIndex];
+		const int32 IVertId = ToMeshV[i]; // I - the row
 
-		UE::Tasks::FTask PendingTask = UE::Tasks::Launch(UE_SOURCE_LOCATION, [StartIndex, EndIndex, &ToMeshV, &ToIndex, &ToTriIdx,
-			&CotangentTriangleDataArray, &Mesh, AreaMode, WeightMode, &TriplesBatch, &LaplacianMatrix]()
+		double WeightArea = 1.0;
+		if (AreaMode == ECotangentAreaMode::VoronoiArea)
 		{
-			// Construct Laplacian Matrix: loop over verts constructing the corresponding matrix row.
-			for (int32 i = StartIndex; i < EndIndex; ++i)
-			{
-				const int32 IVertId = ToMeshV[i]; // I - the row
-
-				double WeightArea = 1.0;
-				if (AreaMode == ECotangentAreaMode::VoronoiArea)
-				{
-					WeightArea = 0.0;
-					Mesh.EnumerateVertexTriangles(IVertId, [&](int32 TriId) {
-						const int32 TriIdx = ToTriIdx[TriId];
-						const CotanTriangleData& TriData = CotangentTriangleDataArray[TriIdx];
-						const FIndex3i TriVertIds = Mesh.GetTriangle(TriId);
-						int32 TriVertIndex = IndexUtil::FindTriIndex(IVertId, TriVertIds);
-						WeightArea += TriData.VoronoiArea[TriVertIndex];
-						});
-				}
-
-				double WeightII = 0.; // accumulate to equal and opposite the sum of the neighbor weights
-
-				for (int32 EdgeId : Mesh.VtxEdgesItr(IVertId))
-				{
-					const FDynamicMesh3::FEdge Edge = Mesh.GetEdge(EdgeId);
-
-					// the other vert in the edge - identifies the matrix column
-					const int32 JVertId = (Edge.Vert[0] == IVertId) ? Edge.Vert[1] : Edge.Vert[0];  // J - the column
+			WeightArea = 0.0;
+			Mesh.EnumerateVertexTriangles(IVertId, [&](int32 TriId) {
+				const int32 TriIdx = ToTriIdx[TriId];
+				const CotanTriangleData& TriData = CotangentTriangleDataArray[TriIdx];
+				const FIndex3i TriVertIds = Mesh.GetTriangle(TriId);
+				int32 TriVertIndex = IndexUtil::FindTriIndex(IVertId, TriVertIds);
+				WeightArea += TriData.VoronoiArea[TriVertIndex];
+			});
+		}
 
-					// Get the cotangents for this edge.
-					const int32 Tri0Idx = ToTriIdx[Edge.Tri[0]];
-					const CotanTriangleData& Tri0Data = CotangentTriangleDataArray[Tri0Idx];
-					double CotanAlpha = Tri0Data.GetOpposingCotangent(EdgeId);
+		double WeightII = 0.; // accumulate to equal and opposite the sum of the neighbor weights
 
-					// The second triangle will be invalid if this is an edge!
-					double CotanBeta = (Edge.Tri[1] != FDynamicMesh3::InvalidID) ? CotangentTriangleDataArray[ToTriIdx[Edge.Tri[1]]].GetOpposingCotangent(EdgeId) : 0.0;
+		for (int32 EdgeId : Mesh.VtxEdgesItr(IVertId))
+		{
+			const FDynamicMesh3::FEdge Edge = Mesh.GetEdge(EdgeId);
 
-					if (WeightMode == ECotangentWeightMode::TriangleArea)
-					{
-						CotanAlpha /= Tri0Data.Area;
+			// the other vert in the edge - identifies the matrix column
+			const int32 JVertId = (Edge.Vert[0] == IVertId) ? Edge.Vert[1] : Edge.Vert[0];  // J - the column
 
-						if (Edge.Tri[1] != FDynamicMesh3::InvalidID)
-						{
-							const CotanTriangleData& Tri1Data = CotangentTriangleDataArray[ToTriIdx[Edge.Tri[1]]];
-							CotanBeta /= Tri1Data.Area;
-						}
-					}
+			// Get the cotangents for this edge.
+			const int32 Tri0Idx = ToTriIdx[Edge.Tri[0]];
+			const CotanTriangleData& Tri0Data = CotangentTriangleDataArray[Tri0Idx];
+			double CotanAlpha = Tri0Data.GetOpposingCotangent(EdgeId);
 
-					// do not need to multiply by 0.5 here...
-					//double WeightIJ = 0.5 * (CotanAlpha + CotanBeta);
-					double WeightIJ = (CotanAlpha + CotanBeta);
+			// The second triangle will be invalid if this is an edge!
+			double CotanBeta = (Edge.Tri[1] != FDynamicMesh3::InvalidID) ? CotangentTriangleDataArray[ToTriIdx[Edge.Tri[1]]].GetOpposingCotangent(EdgeId) : 0.0;
 
-					if (WeightMode == ECotangentWeightMode::ClampedMagnitude)
-					{
-						WeightIJ = FMathd::Clamp(WeightIJ, -1.e5 * WeightArea, 1.e5 * WeightArea);
-					}
+			if (WeightMode == ECotangentWeightMode::TriangleArea) 
+			{
+				CotanAlpha /= Tri0Data.Area;
 
-					WeightII += WeightIJ;
+				if (Edge.Tri[1] != FDynamicMesh3::InvalidID)
+				{
+					const CotanTriangleData& Tri1Data = CotangentTriangleDataArray[ToTriIdx[Edge.Tri[1]]];
+					CotanBeta /= Tri1Data.Area; 
+				}
+			}
 
-					const int32 j = ToIndex[JVertId];
+			// do not need to multiply by 0.5 here...
+			//double WeightIJ = 0.5 * (CotanAlpha + CotanBeta);
+			double WeightIJ = (CotanAlpha + CotanBeta);
 
-					TriplesBatch.Emplace(MatrixEntry(i, j, (RealType)(WeightIJ / WeightArea)));
-				}
-				TriplesBatch.Emplace(MatrixEntry(i, i, (RealType)(-WeightII / WeightArea)));
+			if (WeightMode == ECotangentWeightMode::ClampedMagnitude)
+			{
+				WeightIJ = FMathd::Clamp(WeightIJ, -1.e5 * WeightArea, 1.e5 * WeightArea);
 			}
-		});
-		PendingTasks.Add(PendingTask);
-	}
 
-	UE::Tasks::Wait(PendingTasks);
+			WeightII += WeightIJ;
 
-	for (int32 BatchIndex = 0; BatchIndex < NumBatches; BatchIndex++)
-	{
-		TArray <MatrixEntry>& Triplets = TriplesBatches[BatchIndex];
-		LaplacianMatrix.AddEntriesFunc(Triplets.GetData(), Triplets.Num());
+			const int32 j = ToIndex[JVertId];
+			LaplacianMatrix.AddEntryFunc(i, j, (RealType)(WeightIJ / WeightArea));
+		}
+		LaplacianMatrix.AddEntryFunc(i, i, (RealType)(-WeightII / WeightArea));
 	}
 }
 
diff --git a/Engine/Plugins/Runtime/GeometryProcessing/Source/DynamicMesh/Public/Solvers/PrecomputedMeshWeightData.h b/Engine/Plugins/Runtime/GeometryProcessing/Source/DynamicMesh/Public/Solvers/PrecomputedMeshWeightData.h
--- a/Engine/Plugins/Runtime/GeometryProcessing/Source/DynamicMesh/Public/Solvers/PrecomputedMeshWeightData.h
+++ b/Engine/Plugins/Runtime/GeometryProcessing/Source/DynamicMesh/Public/Solvers/PrecomputedMeshWeightData.h
@@ -4,8 +4,6 @@
 
 #include "DynamicMesh/DynamicMesh3.h"
 #include "Solvers/MeshLinearization.h"
-#include "Tasks/Task.h"
-#include "Async/TaskGraphInterfaces.h"
 
 #define UE_API DYNAMICMESH_API
 
@@ -212,32 +210,16 @@ namespace UE
 
 			const auto& ToTriId = TriangleLinearization.ToId();
 
-			const int32 NumTasks = FMath::Max(FMath::Min(FTaskGraphInterface::Get().GetNumWorkerThreads()/2, NumTris), 1);
-			constexpr int32 MinTrianglesByTask = 60;
-			const int32 TrianglesByTask = FMath::Max(FMath::Max(FMath::DivideAndRoundUp(NumTris, NumTasks), MinTrianglesByTask), 1);
-			const int32 NumBatches = FMath::DivideAndRoundUp(NumTris, TrianglesByTask);
-			TArray<UE::Tasks::FTask> PendingTasks;
-			PendingTasks.Reserve(NumBatches);
-			for (int32 BatchIndex = 0; BatchIndex < NumBatches; BatchIndex++)
+			for (int32 i = 0; i < NumTris; ++i)
 			{
-				const int32 StartIndex = BatchIndex * TrianglesByTask;
-				int32 EndIndex = (BatchIndex + 1) * TrianglesByTask;
-				EndIndex = BatchIndex == NumBatches - 1 ? FMath::Min(NumTris, EndIndex) : EndIndex;
+				// Current triangle
 
-				UE::Tasks::FTask PendingTask = UE::Tasks::Launch(UE_SOURCE_LOCATION, [StartIndex, EndIndex, &ToTriId, &TriangleDataArray, &DynamicMesh]()
-				{
-					for (int32 i = StartIndex; i < EndIndex; ++i)
-					{
-						// Current triangle
-						const int32 TriId = ToTriId[i];
-
-						// Compute all the geometric data needed for this triangle.
-						TriangleDataArray[i].Initialize(DynamicMesh, TriId);
-					}
-				});
-				PendingTasks.Add(PendingTask);
+				const int32 TriId = ToTriId[i];
+
+				// Compute all the geometric data needed for this triangle.
+
+				TriangleDataArray[i].Initialize(DynamicMesh, TriId);
 			}
-			UE::Tasks::Wait(PendingTasks);
 		}
 
 
diff --git a/Engine/Source/Runtime/GeometryCore/Public/Solvers/MatrixInterfaces.h b/Engine/Source/Runtime/GeometryCore/Public/Solvers/MatrixInterfaces.h
--- a/Engine/Source/Runtime/GeometryCore/Public/Solvers/MatrixInterfaces.h
+++ b/Engine/Source/Runtime/GeometryCore/Public/Solvers/MatrixInterfaces.h
@@ -25,21 +25,6 @@ namespace UE
 
 			/** Add matrix entry tuple (i,j,Value) */
 			TUniqueFunction<void(int32, int32, RealType)> AddEntryFunc;
-
-			struct FTupleData
-			{
-				FTupleData(int32 IIn, int32 JIn, RealType ValueIn)
-					: I(IIn)
-					, J(JIn)
-					, Value(ValueIn)
-				{}
-				int32 I;
-				int32 J;
-				RealType Value;
-			};
-
-			/** Add matrix entries (i,j,Value) and size of tuples */
-			TUniqueFunction<void(const FTupleData*, int32)> AddEntriesFunc;
 		};
 
 		typedef TSparseMatrixAssembler<double> FSparseMatrixAssemblerd;
