Add default functor to add several entires in DNCP Sparse Matrix

diff --git a/Engine/Plugins/Runtime/GeometryProcessing/Source/DynamicMesh/Private/Solvers/Internal/MeshUVSolver.cpp b/Engine/Plugins/Runtime/GeometryProcessing/Source/DynamicMesh/Private/Solvers/Internal/MeshUVSolver.cpp
--- a/Engine/Plugins/Runtime/GeometryProcessing/Source/DynamicMesh/Private/Solvers/Internal/MeshUVSolver.cpp
+++ b/Engine/Plugins/Runtime/GeometryProcessing/Source/DynamicMesh/Private/Solvers/Internal/MeshUVSolver.cpp
@@ -43,6 +43,19 @@ public:
 			// set lower-right block value
 			EntryTriplets.push_back(MatrixTripletT(NumVert + RowIdx, NumVert + ColIdx, Value));
 		};
+
+		AddEntriesFunc = [this](const FTupleData* Data, int32 NumData)
+		{
+			const MatrixTripletT* MatrixTriplet = reinterpret_cast<const MatrixTripletT*>(Data);
+			static_assert(sizeof(MatrixTripletT) == sizeof(FTupleData));
+			for (int32 Index = 0; Index < NumData; Index++)
+			{
+				// set upper-left block value
+				EntryTriplets.push_back(MatrixTriplet[Index]);
+				// set lower-right block value
+				EntryTriplets.push_back(MatrixTripletT(NumVert + MatrixTriplet[Index].row(), NumVert + MatrixTriplet[Index].col(), MatrixTriplet[Index].value()));
+			}
+		};
 	}
 
 	void ExtractResult(FSparseMatrixD& Result)
