Clean up DNCP Sparse Matrix
Remove unnecessary reinterpret cast

diff --git a/Engine/Plugins/Runtime/GeometryProcessing/Source/DynamicMesh/Private/Solvers/Internal/MeshUVSolver.cpp b/Engine/Plugins/Runtime/GeometryProcessing/Source/DynamicMesh/Private/Solvers/Internal/MeshUVSolver.cpp
--- a/Engine/Plugins/Runtime/GeometryProcessing/Source/DynamicMesh/Private/Solvers/Internal/MeshUVSolver.cpp
+++ b/Engine/Plugins/Runtime/GeometryProcessing/Source/DynamicMesh/Private/Solvers/Internal/MeshUVSolver.cpp
@@ -46,14 +46,12 @@ public:
 
 		AddEntriesFunc = [this](const FTupleData* Data, int32 NumData)
 		{
-			const MatrixTripletT* MatrixTriplet = reinterpret_cast<const MatrixTripletT*>(Data);
-			static_assert(sizeof(MatrixTripletT) == sizeof(FTupleData));
 			for (int32 Index = 0; Index < NumData; Index++)
 			{
 				// set upper-left block value
-				EntryTriplets.push_back(MatrixTriplet[Index]);
+				EntryTriplets.push_back(MatrixTripletT(Data[Index].I, Data[Index].J, Data[Index].Value));
 				// set lower-right block value
-				EntryTriplets.push_back(MatrixTripletT(NumVert + MatrixTriplet[Index].row(), NumVert + MatrixTriplet[Index].col(), MatrixTriplet[Index].value()));
+				EntryTriplets.push_back(MatrixTripletT(NumVert + Data[Index].I, NumVert + Data[Index].J, Data[Index].Value));
 			}
 		};
 	}
