Lattice Deformer Tool:
- Modify the FFDLattice to optionally accept an oriented bounding box rather than computing one from the input mesh. This allows the lattice to operate on just a portion of the mesh. Lattice computations are now all done in local space.
- Add an ILatticeStateStorage interface, which will enable the interactive LatticeDeformerTool to store its state in when the tool ends

diff --git a/Engine/Plugins/Runtime/GeometryProcessing/Source/DynamicMesh/Private/Operations/FFDLattice.cpp b/Engine/Plugins/Runtime/GeometryProcessing/Source/DynamicMesh/Private/Operations/FFDLattice.cpp
--- a/Engine/Plugins/Runtime/GeometryProcessing/Source/DynamicMesh/Private/Operations/FFDLattice.cpp
+++ b/Engine/Plugins/Runtime/GeometryProcessing/Source/DynamicMesh/Private/Operations/FFDLattice.cpp
@@ -59,6 +59,7 @@ namespace FFDLatticeHelpers
 
 
 FFFDLattice::FFFDLattice(const FVector3i& InDims, const FDynamicMesh3& Mesh, float Padding) :
+	Transform(FTransformSRT3d::Identity()),
 	Dimensions(InDims)
 {
 	check(InDims.X > 1 && InDims.Y > 1 && InDims.Z > 1);
@@ -69,26 +70,49 @@ FFFDLattice::FFFDLattice(const FVector3i& InDims, const FDynamicMesh3& Mesh, flo
 		InitialBounds = FAxisAlignedBox3d(FVector3d::Zero(), FVector3d::Zero());
 	}
 
-	// Expand the initial bounding box to make the computation of which grid cell a mesh vertex is inside of a little 
-	// less susceptible to numerical error issues.
-
-	const float ClampedPadding = FMath::Clamp(Padding, 0.01f, 5.0f);
 	const FVector3d Center = InitialBounds.Center();
-	FVector3d Extents = InitialBounds.Extents();
-
-	// Compute padding based on maximum component of the diagonal. Avoids problems when one or more component is zero.
-	double MaxDiagonal = MaxElement(InitialBounds.Diagonal());
-	Extents = Extents + (0.5 * ClampedPadding * MaxDiagonal);
+	const FVector3d Extents = InitialBounds.Extents();
+	Transform.SetTranslation(Center - Extents);
 
-	InitialBounds.Min = Center - Extents;
-	InitialBounds.Max = Center + Extents;
-	
-	FVector3d Diag = InitialBounds.Diagonal();
-	CellSize = Diag / FVector3d(Dimensions - 1);
+	InitLattice(Padding);
 
 	ComputeInitialEmbedding(Mesh);
 }
 
+FFFDLattice::FFFDLattice(const FVector3i& InDims, const FDynamicMesh3* Mesh, float Padding, const FAxisAlignedBox3d& InInitialBounds, const FTransformSRT3d& LatticeTransform) :
+	Transform(LatticeTransform),
+	Dimensions(InDims),
+	InitialBounds(InInitialBounds)
+{
+	check(InDims.X > 1 && InDims.Y > 1 && InDims.Z > 1);
+
+	if (InitialBounds.IsEmpty())
+	{
+		InitialBounds = FAxisAlignedBox3d(FVector3d::Zero(), FVector3d::Zero());
+	}
+
+	InitLattice(Padding);
+
+	if (Mesh)
+	{
+		ComputeInitialEmbedding(*Mesh);
+	}
+}
+
+void FFFDLattice::InitLattice(float Padding)
+{
+	// Create initial bounding box with user-specified padding
+	const FVector3d Extents = InitialBounds.Extents();
+	const float ClampedPadding = FMath::Clamp(Padding, 0.0f, 5.0f);
+	const double MaxDiagonal = MaxElement(InitialBounds.Diagonal());
+
+	InitialBounds.Min = FVector3d(-0.5 * ClampedPadding * MaxDiagonal);
+	InitialBounds.Max = 2.0 * Extents + 0.5 * ClampedPadding * MaxDiagonal;
+
+	const FVector3d Diag = InitialBounds.Diagonal();
+	CellSize = Diag / FVector3d(Dimensions - 1);
+}
+
 void FFFDLattice::GenerateInitialLatticePositions(TArray<FVector3d>& OutLatticePositions) const
 {
 	int TotalNumLatticePoints = Dimensions.X * Dimensions.Y * Dimensions.Z;
@@ -105,7 +129,7 @@ void FFFDLattice::GenerateInitialLatticePositions(TArray<FVector3d>& OutLatticeP
 				int PointID = ControlPointIndexFromCoordinates(i, j, k);
 				double Z = CellSize.Z * k;
 
-				OutLatticePositions[PointID] = InitialBounds.Min + FVector3d{ X,Y,Z };
+				OutLatticePositions[PointID] = Transform.TransformPosition(FVector3d{ X,Y,Z });
 			}
 		}
 	}
@@ -147,35 +171,72 @@ void FFFDLattice::GenerateLatticeEdges(TArray<FVector2i>& OutLatticeEdges) const
 
 FVector3d FFFDLattice::ComputeTrilinearWeights(const FVector3d& QueryPoint, FVector3i& GridCoordinates) const
 {
-	FVector3d GridPoint = (QueryPoint - InitialBounds.Min) / CellSize;
+	const FVector3d GridPoint = QueryPoint / CellSize;
 	GridCoordinates = FVector3i(GridPoint);
 	FVector3d Weights = GridPoint - FVector3d(GridCoordinates);
+
+	for (int32 Coord = 0; Coord < 3; ++Coord)
+	{
+		double OrigWeight = Weights[Coord];
+
+		if (GridCoordinates[Coord] < 0)
+		{
+			GridCoordinates[Coord] = 0;
+			Weights[Coord] = 0.0;
+		}
+		if (GridCoordinates[Coord] > Dimensions[Coord] - 2)
+		{
+			GridCoordinates[Coord] = Dimensions[Coord] - 2;
+			Weights[Coord] = 1.0;
+		}
+	}
+
 	return Weights;
 }
 
-
 void FFFDLattice::ComputeInitialEmbedding(const FDynamicMesh3& Mesh, FLatticeExecutionInfo ExecutionInfo)
 {
-	int MaxVertexID = Mesh.MaxVertexID();
+	const int MaxVertexID = Mesh.MaxVertexID();
 	VertexEmbeddings.SetNum(MaxVertexID);
 
-	EParallelForFlags ParallelForFlags = ExecutionInfo.bParallel ? EParallelForFlags::None : EParallelForFlags::ForceSingleThread;
+	const EParallelForFlags ParallelForFlags = ExecutionInfo.bParallel ? EParallelForFlags::None : EParallelForFlags::ForceSingleThread;
+
+	// Expand the box a little to catch vertices right on the lattice boundary. Points which are a bit outside the box will have their embedding cells/weights clamped 
+	// in ComputeTrilinearWeights. Note this is different than the user-supplied padding which changes the dimensions of the bounding box.
+	FAxisAlignedBox3d ExpandedBox = InitialBounds;
+	ExpandedBox.Expand(UE_KINDA_SMALL_NUMBER);
 
 	ParallelFor(MaxVertexID,
-				[this, &Mesh, MaxVertexID](int VertexID)
+				[this, &Mesh, &ExpandedBox, MaxVertexID](int VertexID)
 	{
+		FEmbedding& Embedding = VertexEmbeddings[VertexID];
+		Embedding.LatticeCell = { -1, -1, -1 };
+
 		if (!Mesh.IsVertex(VertexID))
 		{
 			return;
 		}
 
-		FVector3d VertexPosition = Mesh.GetVertex(VertexID);
+		const FVector3d MeshVertexPosition = Mesh.GetVertex(VertexID);
+		const FVector3d LatticeSpaceVertexPosition = Transform.InverseTransformPosition(MeshVertexPosition);
 
-		FEmbedding& Embedding = VertexEmbeddings[VertexID];
-		Embedding.CellWeighting = ComputeTrilinearWeights(VertexPosition, Embedding.LatticeCell);
+		if (!ExpandedBox.Contains(LatticeSpaceVertexPosition))
+		{
+			return;
+		}
+
+		Embedding.CellWeighting = ComputeTrilinearWeights(LatticeSpaceVertexPosition, Embedding.LatticeCell);
 	},
-				ParallelForFlags);
+	ParallelForFlags);
+}
 
+bool FFFDLattice::VertexHasValidEmbedding(int32 VertexIndex) const
+{
+	return (VertexIndex >= 0) 
+		&& (VertexIndex < VertexEmbeddings.Num()) 
+		&& (VertexEmbeddings[VertexIndex].LatticeCell[0] >= 0) 
+		&& (VertexEmbeddings[VertexIndex].LatticeCell[1] >= 0) 
+		&& (VertexEmbeddings[VertexIndex].LatticeCell[2] >= 0);
 }
 
 
diff --git a/Engine/Plugins/Runtime/GeometryProcessing/Source/DynamicMesh/Public/Operations/FFDLattice.h b/Engine/Plugins/Runtime/GeometryProcessing/Source/DynamicMesh/Public/Operations/FFDLattice.h
--- a/Engine/Plugins/Runtime/GeometryProcessing/Source/DynamicMesh/Public/Operations/FFDLattice.h
+++ b/Engine/Plugins/Runtime/GeometryProcessing/Source/DynamicMesh/Public/Operations/FFDLattice.h
@@ -50,6 +50,13 @@ public:
 	/// the weight of each mesh vertex inside the appropriate lattice cell.
 	UE_API FFFDLattice(const FVector3i& InDims, const FDynamicMesh3& InMesh, float Padding);
 
+	/// Create a lattice that fits the oriented box defined by the AABB and Transform, and has the given resolution along each dimension. Also precompute
+	/// the weight of each mesh vertex inside the appropriate lattice cell. Mesh vertices outside of the oriented bounding box should have a LatticeCell value of {-1, -1, -1}.
+	UE_API FFFDLattice(const FVector3i& InDims, const FDynamicMesh3* Mesh, float Padding, const FAxisAlignedBox3d& InInitialBounds, const FTransformSRT3d& LatticeTransform = FTransformSRT3d::Identity());
+
+	/// Whether the given vertex has a valid embedding, i.e. whether the original vertex location was inside a lattice cell and it would be affected by lattice deformation
+	UE_API bool VertexHasValidEmbedding(int32 VertexIndex) const;
+
 	/// Using the Lattice's current control point positions and the original embedding information, compute the
 	/// deformed mesh vertex positions.
 	UE_API void GetDeformedMeshVertexPositions(const TArray<FVector3d>& LatticeControlPoints,
@@ -115,6 +122,9 @@ public:
 
 protected:
 
+	// Lattice transform
+	FTransformSRT3d Transform;
+
 	// Number of lattice points in each dimension
 	FVector3i Dimensions;
 
@@ -127,6 +137,9 @@ protected:
 	// Interpolation weights and cell indices per vertex of the input FDynamicMesh3
 	TArray<FEmbedding> VertexEmbeddings;
 
+	// Ccnstructor helper function for setting up lattice data
+	void InitLattice(float Padding);
+
 	/// For each vertex in Mesh, compute the lattice cell it resides in and its weighting
 	UE_API void ComputeInitialEmbedding(const FDynamicMesh3& Mesh, FLatticeExecutionInfo ExecutionInfo = FLatticeExecutionInfo());
 
diff --git a/Engine/Plugins/Runtime/MeshModelingToolset/Source/MeshModelingTools/Private/LatticeDeformerTool.cpp b/Engine/Plugins/Runtime/MeshModelingToolset/Source/MeshModelingTools/Private/LatticeDeformerTool.cpp
--- a/Engine/Plugins/Runtime/MeshModelingToolset/Source/MeshModelingTools/Private/LatticeDeformerTool.cpp
+++ b/Engine/Plugins/Runtime/MeshModelingToolset/Source/MeshModelingTools/Private/LatticeDeformerTool.cpp
@@ -114,7 +114,22 @@ void ULatticeDeformerToolProperties::PostAction(ELatticeDeformerToolAction Actio
 
 UMultiTargetWithSelectionTool* ULatticeDeformerToolBuilder::CreateNewTool(const FToolBuilderState& SceneState) const
 {
-	return NewObject<ULatticeDeformerTool>(SceneState.ToolManager);
+	const TObjectPtr<ULatticeDeformerTool> Tool = NewObject<ULatticeDeformerTool>(SceneState.ToolManager);
+
+	// Check if the selected component implements ILatticeStateStorage 
+	ensure(SceneState.TargetManager->CountSelectedAndTargetable(SceneState, GetTargetRequirements()) == 1);		// should be enforced by CanBuildTool
+	SceneState.TargetManager->EnumerateSelectedAndTargetableComponents(SceneState, GetTargetRequirements(), [&Tool](UActorComponent* Component)
+	{
+		if (Component->GetClass()->ImplementsInterface(ULatticeStateStorage::StaticClass()))
+		{
+			TScriptInterface<ILatticeStateStorage> LatticeInterface;
+			LatticeInterface.SetObject(Component);
+			LatticeInterface.SetInterface(Cast<ILatticeStateStorage>(Component));
+			Tool->SetLatticeStorage(LatticeInterface);
+		}
+	});
+
+	return Tool;
 }
 
 bool ULatticeDeformerToolBuilder::CanBuildTool(const FToolBuilderState& SceneState) const
@@ -205,16 +220,27 @@ void ULatticeDeformerTool::InitializeLattice(TArray<FVector3d>& OutLatticePoints
 	{
 		MeshToDeform = OriginalMesh.Get();
 	}
-	Lattice = MakeShared<FFFDLattice, ESPMode::ThreadSafe>(GetLatticeResolution(), *MeshToDeform, Settings->Padding);
 
+	if (!LatticeStorage)
+	{
+		Lattice = MakeShared<FFFDLattice, ESPMode::ThreadSafe>(GetLatticeResolution(), *MeshToDeform, Settings->Padding);
+	}
+	else
+	{
+		Lattice = MakeShared<FFFDLattice, ESPMode::ThreadSafe>(GetLatticeResolution(), MeshToDeform, Settings->Padding, LatticeStorage->GetInitialBounds(), LatticeStorage->GetTransform());
+	}
 
 	Lattice->GenerateInitialLatticePositions(OutLatticePoints);
 
-	// Put the lattice in world space
-	FTransform3d LocalToWorld(Cast<ISceneComponentBackedTarget>(Targets[0])->GetWorldTransform());
-	Algo::ForEach(OutLatticePoints, [&LocalToWorld](FVector3d& Point) {
-		Point = LocalToWorld.TransformPosition(Point);
-	});
+	// Put the lattice in world space. (If we have LatticeStorage they will be copied over in world space later.)
+	if (!LatticeStorage)
+	{
+		FTransform3d LocalToWorld(Cast<ISceneComponentBackedTarget>(Targets[0])->GetWorldTransform());
+		Algo::ForEach(OutLatticePoints, [&LocalToWorld](FVector3d& Point)
+		{
+			Point = LocalToWorld.TransformPosition(Point);
+		});
+	}
 
 	Lattice->GenerateLatticeEdges(OutLatticeEdges);
 }
@@ -280,6 +306,13 @@ void ULatticeDeformerTool::Setup()
 		}
 	});
 
+	if (LatticeStorage)
+	{
+		const FVector3i StoredResolution = LatticeStorage->GetResolution();
+		Settings->XAxisResolution = StoredResolution[0];
+		Settings->YAxisResolution = StoredResolution[1];
+		Settings->ZAxisResolution = StoredResolution[2];
+	}
 
 	TArray<FVector3d> LatticePoints;
 	TArray<FVector2i> LatticeEdges;
@@ -350,6 +383,26 @@ void ULatticeDeformerTool::Setup()
 		}
 	}
 
+	if (LatticeStorage)
+	{
+		TArray<FVector3d> SavedPoints;
+		LatticeStorage->ReadLatticePoints(SavedPoints);
+		if (SavedPoints.Num() == LatticePoints.Num())
+		{
+			ControlPointsMechanic->UpdatePointLocations(SavedPoints);
+			ControlPointsMechanic->bHasChanged = true;
+			OnPointsChangedLambda();
+		}
+
+		Settings->InterpolationType = LatticeStorage->ReadInterpolationType();
+	}
+
+	if (Settings->bSoftDeformation)
+	{
+		RebuildDeformer();
+	}
+
+	Settings->SilentUpdateWatched();		// On next tick, don't react to any settings changes made in this function
 }
 
 
@@ -445,6 +498,12 @@ void ULatticeDeformerTool::SoftDeformLattice()
 
 void ULatticeDeformerTool::OnShutdown(EToolShutdownType ShutdownType)
 {
+	if (ShutdownType == EToolShutdownType::Accept && LatticeStorage)
+	{
+		LatticeStorage->StoreLatticePoints(ControlPointsMechanic->GetControlPoints());
+		LatticeStorage->StoreInterpolationType(Settings->InterpolationType);
+	}
+
 	Settings->SaveProperties(this);
 	ControlPointsMechanic->Shutdown();
 
@@ -561,6 +620,11 @@ void ULatticeDeformerTool::Render(IToolsContextRenderAPI* RenderAPI)
 	}
 }
 
+void ULatticeDeformerTool::SetLatticeStorage(const TScriptInterface<ILatticeStateStorage>& InLatticeStorage)
+{
+	LatticeStorage = InLatticeStorage;
+}
+
 void ULatticeDeformerTool::RequestAction(ELatticeDeformerToolAction Action)
 {
 	if (PendingAction == ELatticeDeformerToolAction::NoAction)
diff --git a/Engine/Plugins/Runtime/MeshModelingToolset/Source/MeshModelingTools/Public/LatticeDeformerTool.h b/Engine/Plugins/Runtime/MeshModelingToolset/Source/MeshModelingTools/Public/LatticeDeformerTool.h
--- a/Engine/Plugins/Runtime/MeshModelingToolset/Source/MeshModelingTools/Public/LatticeDeformerTool.h
+++ b/Engine/Plugins/Runtime/MeshModelingToolset/Source/MeshModelingTools/Public/LatticeDeformerTool.h
@@ -11,6 +11,7 @@
 #include "DynamicMesh/MeshSharingUtil.h"
 #include "Operations/FFDLattice.h"
 #include "Solvers/ConstrainedMeshDeformer.h"
+#include "LatticeManager.h"
 
 #include "LatticeDeformerTool.generated.h"
 
@@ -24,6 +25,7 @@ namespace UE::Geometry
 class ULatticeControlPointsMechanic;
 class UMeshOpPreviewWithBackgroundCompute;
 class UMeshSculptLayerProperties;
+class ILatticeStateStorage;
 
 UCLASS(MinimalAPI)
 class ULatticeDeformerToolBuilder : public UMultiTargetWithSelectionToolBuilder
@@ -41,16 +43,6 @@ private:
 };
 
 
-UENUM()
-enum class ELatticeInterpolationType : uint8
-{
-	/** Use trilinear interpolation to get new mesh vertex positions from the lattice */
-	Linear UMETA(DisplayName = "Linear"),
-
-	/** Use tricubic interpolation to get new mesh vertex positions from the lattice */
-	Cubic UMETA(DisplayName = "Cubic")
-};
-
 UENUM()
 enum class ELatticeDeformerToolAction  : uint8
 {
@@ -163,6 +155,7 @@ public:
 	UE_API virtual void Render(IToolsContextRenderAPI* RenderAPI) override;
 
 	UE_API UE::Geometry::FVector3i GetLatticeResolution() const;
+	UE_API void SetLatticeStorage(const TScriptInterface<ILatticeStateStorage>& InLatticeStorage);
 
 	// IInteractiveToolManageGeometrySelectionAPI -- this tool won't update external geometry selection or change selection-relevant mesh IDs
 	virtual bool IsInputSelectionValidOnOutput() override
@@ -193,6 +186,9 @@ protected:
 	UPROPERTY()
 	TObjectPtr<UMeshSculptLayerProperties> SculptLayerProperties;
 
+	UPROPERTY()
+	TScriptInterface<ILatticeStateStorage> LatticeStorage;
+
 	UPROPERTY()
 	bool bLatticeDeformed = false;
 
diff --git a/Engine/Plugins/Runtime/MeshModelingToolset/Source/ModelingComponents/Private/Mechanics/LatticeControlPointsMechanic.cpp b/Engine/Plugins/Runtime/MeshModelingToolset/Source/ModelingComponents/Private/Mechanics/LatticeControlPointsMechanic.cpp
--- a/Engine/Plugins/Runtime/MeshModelingToolset/Source/ModelingComponents/Private/Mechanics/LatticeControlPointsMechanic.cpp
+++ b/Engine/Plugins/Runtime/MeshModelingToolset/Source/ModelingComponents/Private/Mechanics/LatticeControlPointsMechanic.cpp
@@ -385,6 +385,15 @@ void ULatticeControlPointsMechanic::UpdatePointLocations(const TMap<int32, FVect
 	UpdateDrawables();
 }
 
+void ULatticeControlPointsMechanic::UpdatePointLocations(const TArray<FVector3d>& NewLocations)
+{
+	if (ControlPoints.Num() == NewLocations.Num())
+	{
+		ControlPoints = NewLocations;
+	}
+	UpdateDrawables();
+}
+
 bool ULatticeControlPointsMechanic::HitTest(const FInputDeviceRay& ClickPos, FInputRayHit& ResultOut)
 {
 	FGeometrySet3::FNearest Nearest;
diff --git a/Engine/Plugins/Runtime/MeshModelingToolset/Source/ModelingComponents/Public/LatticeManager.h b/Engine/Plugins/Runtime/MeshModelingToolset/Source/ModelingComponents/Public/LatticeManager.h
new file mode 100644
--- /dev/null
+++ b/Engine/Plugins/Runtime/MeshModelingToolset/Source/ModelingComponents/Public/LatticeManager.h
@@ -0,0 +1,42 @@
+// Copyright Epic Games, Inc. All Rights Reserved.
+
+#pragma once
+
+#include "CoreMinimal.h"
+#include "UObject/Interface.h"
+#include "GeometryBase.h"
+#include "BoxTypes.h"
+#include "IntVectorTypes.h"
+#include "LatticeManager.generated.h"
+
+UENUM()
+enum class ELatticeInterpolationType : uint8
+{
+	/** Use trilinear interpolation to get new mesh vertex positions from the lattice */
+	Linear UMETA(DisplayName = "Linear"),
+
+	/** Use tricubic interpolation to get new mesh vertex positions from the lattice */
+	Cubic UMETA(DisplayName = "Cubic")
+};
+
+UINTERFACE()
+class MODELINGCOMPONENTS_API ULatticeStateStorage : public UInterface
+{
+	GENERATED_BODY()
+};
+
+class MODELINGCOMPONENTS_API ILatticeStateStorage
+{
+	GENERATED_BODY()
+public:
+
+	virtual void StoreLatticePoints(const TArray<FVector3d>& LatticePoints) = 0;
+	virtual void ReadLatticePoints(TArray<FVector3d>& LatticePoints) const = 0;
+	
+	virtual void StoreInterpolationType(ELatticeInterpolationType Type) = 0;
+	virtual ELatticeInterpolationType ReadInterpolationType() const = 0;
+
+	virtual UE::Geometry::FAxisAlignedBox3d GetInitialBounds() const = 0;
+	virtual UE::Geometry::FTransformSRT3d GetTransform() const = 0;
+	virtual UE::Geometry::FVector3i GetResolution() const = 0;
+};
diff --git a/Engine/Plugins/Runtime/MeshModelingToolset/Source/ModelingComponents/Public/Mechanics/LatticeControlPointsMechanic.h b/Engine/Plugins/Runtime/MeshModelingToolset/Source/ModelingComponents/Public/Mechanics/LatticeControlPointsMechanic.h
--- a/Engine/Plugins/Runtime/MeshModelingToolset/Source/ModelingComponents/Public/Mechanics/LatticeControlPointsMechanic.h
+++ b/Engine/Plugins/Runtime/MeshModelingToolset/Source/ModelingComponents/Public/Mechanics/LatticeControlPointsMechanic.h
@@ -105,6 +105,9 @@ public:
 
 	UE_API void UpdatePointLocations(const TMap<int32, FVector3d>& NewLocations);
 
+	/// Set all lattice control points
+	UE_API void UpdatePointLocations(const TArray<FVector3d>& NewLocations);
+
 	bool IsGizmoBeingDragged() const
 	{
 		return bGizmoBeingDragged;
diff --git a/Engine/Plugins/Runtime/MeshModelingToolset/Source/ModelingOperators/Private/DeformationOps/LatticeDeformerOp.cpp b/Engine/Plugins/Runtime/MeshModelingToolset/Source/ModelingOperators/Private/DeformationOps/LatticeDeformerOp.cpp
--- a/Engine/Plugins/Runtime/MeshModelingToolset/Source/ModelingOperators/Private/DeformationOps/LatticeDeformerOp.cpp
+++ b/Engine/Plugins/Runtime/MeshModelingToolset/Source/ModelingOperators/Private/DeformationOps/LatticeDeformerOp.cpp
@@ -66,7 +66,10 @@ void FLatticeDeformerOp::CalculateResult(FProgressCancel* Progress)
 		// if deforming entire mesh, all the mesh's new positions will be in DeformedPositions
 		for (const int VID : ResultMesh->VertexIndicesItr())
 		{
-			ResultMesh->SetVertex(VID, DeformedPositions[VID]);
+			if (Lattice->VertexHasValidEmbedding(VID))
+			{
+				ResultMesh->SetVertex(VID, DeformedPositions[VID]);
+			}
 		}
 	}
 
