GeometryProcessing: add Red-Green tessellation as a SelectiveTessellation pattern. Also add a Validate function to all patterns.

diff --git a/Engine/Plugins/Runtime/GeometryProcessing/Source/DynamicMesh/Private/Operations/RedGreenTessellation.cpp b/Engine/Plugins/Runtime/GeometryProcessing/Source/DynamicMesh/Private/Operations/RedGreenTessellation.cpp
new file mode 100644
--- /dev/null
+++ b/Engine/Plugins/Runtime/GeometryProcessing/Source/DynamicMesh/Private/Operations/RedGreenTessellation.cpp
@@ -0,0 +1,481 @@
+// Copyright Epic Games, Inc. All Rights Reserved.
+
+#include "Operations/RedGreenTessellation.h"
+#include "DynamicMesh/DynamicMesh3.h"
+
+namespace UE
+{
+namespace Geometry
+{
+namespace RedGreenTessellationPatternHelpers
+{
+	/*
+	* Given a triangle and optional midpoint vertices, produce a valid triangulation
+	*/
+	void TriangulatePoly(const TStaticArray<int32, 3>& Triangle, const TStaticArray<int32, 3>& MidPoints, TArray<FIndex3i, TInlineAllocator<4>>& Tris)
+	{
+		int32 NumSplits = 0;
+		for (int32 MidpointIndex = 0; MidpointIndex < 3; ++MidpointIndex)
+		{
+			if (MidPoints[MidpointIndex] >= 0)
+			{
+				++NumSplits;
+			}
+		}
+
+		if (NumSplits == 0)
+		{
+			Tris.Add({ Triangle[0], Triangle[1], Triangle[2] });
+		}
+		else if (NumSplits == 1)
+		{
+			// Green triangle
+			for (int32 TriVertIndex = 0; TriVertIndex < 3; ++TriVertIndex)
+			{
+				if (MidPoints[TriVertIndex] >= 0)
+				{
+					const int32 A = Triangle[TriVertIndex];
+					const int32 B = Triangle[(TriVertIndex + 1) % 3];
+					const int32 C = Triangle[(TriVertIndex + 2) % 3];
+					const int32 M = MidPoints[TriVertIndex];
+					Tris.Add({ A, M, C });
+					Tris.Add({ M, B, C });
+					break;
+				}
+			}
+		}
+		else if (NumSplits == 2)
+		{
+			// "Double green" (?) triangle
+			for (int32 TriVertIndex = 0; TriVertIndex < 3; ++TriVertIndex)
+			{
+				const int32 Next = (TriVertIndex + 1) % 3;
+				if (MidPoints[TriVertIndex] >= 0 && MidPoints[Next] >= 0)
+				{
+					const int32 A = Triangle[TriVertIndex];
+					const int32 B = Triangle[Next];
+					const int32 C = Triangle[(TriVertIndex + 2) % 3];
+					const int32 AB = MidPoints[TriVertIndex];
+					const int32 BC = MidPoints[Next];
+					Tris.Add({ A, AB, C });
+					Tris.Add({ AB, BC, C });
+					Tris.Add({ AB, B, BC });
+					break;
+				}
+			}
+		}
+		else
+		{
+			// Red triangle
+			ensure(NumSplits == 3);	
+			const int32 A = Triangle[0];
+			const int32 B = Triangle[1];
+			const int32 C = Triangle[2];
+			const int32 AB = MidPoints[0];
+			const int32 BC = MidPoints[1];
+			const int32 CA = MidPoints[2];
+			Tris.Add({ A, AB, CA });
+			Tris.Add({ AB, B, BC });
+			Tris.Add({ CA, BC, C });
+			Tris.Add({ AB, BC, CA });
+		}
+	}
+}	// namespace RedGreenTessellationPatternHelpers
+
+
+//
+// FRedGreenTessellationPattern
+//
+
+/**
+A single iteration of Red-Green subdivision is shown here, where the center triangle is "red" subdivided (1 to 4), and the neighboring triangles are "green" subdivided (1 to 2) to
+maintain a conforming mesh (i.e. avoid T-junctions or cracks in the mesh.)
+
+
+                       O                                                O
+                      / \                                              /|\
+                     /   \                                            / | \
+                    /     \                                          /  |  \
+                   /       \                                        /   |   \
+                  /         \                                      /    |    \
+                 /           \                                    /     |     \
+                /             \                                  /      |      \
+               /               \                                /       |       \
+              /                 \                              /        |        \
+             O\-----------------/O                            O\--------O--------/O
+            /  \               /  \                          /  \      / \      /  \
+           /    \             /    \                        /    \    /   \    /    \
+          /      \           /      \                      /      \  /     \  /      \
+         /        \         /        \                    /        O/_______O         \
+        /          \       /          \                  /        / \       /\         \
+       /            \     /            \                /      /     \     /    \       \
+      /              \   /              \              /    /         \   /        \     \
+     /                \ /                \            /  /             \ /            \   \
+    /                  \                  \          //                 \                \ \
+   O-------------------O-------------------O    =>  O-------------------O--------------------O
+
+To achieve deeper levels of tessellation, we could repeatedly apply red-green over the entire mesh, however it can be more efficient to tessellate
+each triangle to its specified level in one shot.
+
+Here is an example of our triangle patch tessellation. The triangle has a specified tessellation level 2, but the neighbouring triangle below it has level 3,
+which necessitates additional triangles along the bottom edge.
+
+    O                                                       O
+    |\                                                      | \
+    │  \                                                    │   \
+    │     \                                                 │     \
+    │       \                                               │       \
+    O         O                                             O---------O
+    │           \                                           │\        │ \
+    │             \                                         │   \     │   \
+    │               \                                       │      \  │     \
+    │                 \                                     │        \│       \
+    O                   O                                   O---------O---------O
+    │                     \                                 │\        │\        │ \
+    │                       \                               │   \     │   \     │   \
+    │                         \                             │      \  │      \  │     \
+    │                           \                           │        \│        \│       \
+    O                             O                         O---------O---------O---------O
+    │                               \                       │\\       │\\       │\\       │\\
+    │                                 \                     │ \ \     │ \ \     │ \ \     │ \ \
+    │                                   \                   │  \   \  │  \  \   │  \  \   │  \  \
+    │                                     \                 │   \   \ │   \   \ │   \   \ │   \   \
+    O----*----O----*----O----*----O----*----O        =>     O----*----O----*----O----*----O----*----O
+*/
+
+
+
+FRedGreenTessellationPattern::FRedGreenTessellationPattern(const FDynamicMesh3* InMesh, const TArray<int>& InTriangleTessLevels) :
+	FTessellationPattern(InMesh),
+	TriangleLevels(InTriangleTessLevels)
+{
+}
+
+EOperationValidationResult FRedGreenTessellationPattern::Validate() const
+{
+	if (!Mesh)
+	{
+		return EOperationValidationResult::Failed_UnknownReason;
+	}
+
+	if (TriangleLevels.Num() < Mesh->MaxTriangleID())
+	{
+		// Too few triangle levels specified
+		return EOperationValidationResult::Failed_UnknownReason;
+	}
+
+	// Verify the triangle tessellation levels don't vary by more than 1 between neighboring triangles
+	for (int32 TriangleIndex = 0; TriangleIndex < Mesh->MaxTriangleID(); ++TriangleIndex)
+	{
+		const FIndex3i Neighbors = Mesh->GetTriNeighbourTris(TriangleIndex);
+
+		for (const int NeighborIndex : {Neighbors[0], Neighbors[1], Neighbors[2]})
+		{
+			if (Mesh->IsTriangle(NeighborIndex))
+			{
+				if (FMath::Abs(TriangleLevels[TriangleIndex] - TriangleLevels[NeighborIndex]) > 1)
+				{
+					return EOperationValidationResult::Failed_UnknownReason;
+				}
+			}
+		}
+	}
+
+	return EOperationValidationResult::Ok;
+}
+
+int FRedGreenTessellationPattern::GetNumberOfNewVerticesForEdgePatch(const int InEdgeID) const
+{
+	if (!Mesh->IsEdge(InEdgeID))
+	{
+		return FTessellationPattern::InvalidIndex;
+	}
+	const FIndex2i EdgeTris = Mesh->GetEdgeT(InEdgeID);
+
+	int32 MaxLevel = 0;
+	if (EdgeTris[0] != FDynamicMesh3::InvalidID)
+	{
+		MaxLevel = FMath::Max(MaxLevel, TriangleLevels[EdgeTris[0]]);
+	}
+	if (EdgeTris[1] != FDynamicMesh3::InvalidID)
+	{
+		MaxLevel = FMath::Max(MaxLevel, TriangleLevels[EdgeTris[1]]);
+	}
+
+	return (1 << MaxLevel) - 1;
+}
+
+int FRedGreenTessellationPattern::GetNumberOfNewVerticesForTrianglePatch(const int InTriangleID) const
+{
+	if (!Mesh->IsTriangle(InTriangleID))
+	{
+		return FTessellationPattern::InvalidIndex;
+	}
+	const int32 Level = TriangleLevels[InTriangleID];
+	return ((1 << Level) - 1) * ((1 << Level) - 2) / 2;
+}
+
+
+int FRedGreenTessellationPattern::GetNumberOfPatchTriangles(const int InTriangleID) const
+{
+	if (!Mesh->IsTriangle(InTriangleID))
+	{
+		return FTessellationPattern::InvalidIndex;
+	}
+
+	const int32 Level = TriangleLevels[InTriangleID];
+
+	int32 NewTriangles = 1 << (2 * Level);
+
+	const FIndex3i NeighborTris = Mesh->GetTriNeighbourTris(InTriangleID);
+	for (int32 NIndex = 0; NIndex < 3; ++NIndex)
+	{
+		const int NeighborTri = NeighborTris[NIndex];
+		if (Mesh->IsTriangle(NeighborTri) && TriangleLevels[NeighborTri] > Level)
+		{
+			ensure(TriangleLevels[NeighborTri] == Level + 1);
+			// We need to subdivide new segments along this edge with green triangles
+			NewTriangles += (1 << TriangleLevels[NeighborTri]) - (1 << Level);
+		}
+	}
+
+	return NewTriangles;
+}
+
+void FRedGreenTessellationPattern::TessellateEdgePatch(EdgePatch& EdgePatch) const
+{
+	const int32 NumNewVertices = GetNumberOfNewVerticesForEdgePatch(EdgePatch.EdgeID);
+	const float Param = 1.0f / (static_cast<float>(NumNewVertices) + 1.0f);
+	for (int NewVertex = 0; NewVertex < NumNewVertices; ++NewVertex)
+	{
+		EdgePatch.LinearCoord[NewVertex] = static_cast<float>(NewVertex + 1) * Param;
+	}
+}
+
+void FRedGreenTessellationPattern::TessellateTriPatch(TrianglePatch& TriPatch) const
+{
+	if (TriPatch.Triangles.Num() == 1)
+	{
+		TriPatch.Triangles[0] = TriPatch.UVWCorners;
+		return;
+	}
+
+	int32 NumTessellatedEdges = 0;
+	NumTessellatedEdges += TriPatch.UVEdge.VIDs.Num() > 0 ? 1 : 0;
+	NumTessellatedEdges += TriPatch.VWEdge.VIDs.Num() > 0 ? 1 : 0;
+	NumTessellatedEdges += TriPatch.UWEdge.VIDs.Num() > 0 ? 1 : 0;
+
+	if (NumTessellatedEdges == 1)
+	{
+		// Green triangle
+
+		if (TriPatch.UVEdge.VIDs.Num() > 0)
+		{
+			const int32 V0 = TriPatch.UVWCorners[0];
+			const int32 V1 = TriPatch.UVWCorners[1];
+			const int32 V2 = TriPatch.UVWCorners[2];
+			const int32 Mid = TriPatch.UVEdge.VIDs[0];
+			TriPatch.Triangles[0] = { V2, V0, Mid };
+			TriPatch.Triangles[1] = { V2, Mid, V1 };
+		}
+		else if (TriPatch.VWEdge.VIDs.Num() > 0)
+		{
+			const int32 V0 = TriPatch.UVWCorners[1];
+			const int32 V1 = TriPatch.UVWCorners[2];
+			const int32 V2 = TriPatch.UVWCorners[0];
+			const int32 Mid = TriPatch.VWEdge.VIDs[0];
+			TriPatch.Triangles[0] = { V2, V0, Mid };
+			TriPatch.Triangles[1] = { V2, Mid, V1 };
+		}
+		else
+		{
+			ensure(TriPatch.UWEdge.VIDs.Num() > 0);
+			const int32 V0 = TriPatch.UVWCorners[2];
+			const int32 V1 = TriPatch.UVWCorners[0];
+			const int32 V2 = TriPatch.UVWCorners[1];
+			const int32 Mid = TriPatch.UWEdge.VIDs[0];
+			TriPatch.Triangles[0] = { V2, V0, Mid };
+			TriPatch.Triangles[1] = { V2, Mid, V1 };
+		}
+
+		return;
+	}
+
+	const int32 Level = TriangleLevels[TriPatch.TriangleID];
+
+	// Create a grid of vertices which will then be triangulated. The grid structure stores vertex indices
+	// (Note the grid itself is triangular, only storing values at nodes where i + j < N)
+
+	const int32 GridN = (1 << Level) + 1;
+
+	// TODO: If using a TMap is slow, investigate creating a triangular grid structure. Maybe similar to TDenseGrid2 but only storing the lower half?
+	TMap<TPair<int32, int32>, int32> Grid;
+	Grid.Reserve(GridN * GridN / 2);
+
+	for (int32 I = 0; I < GridN; ++I)
+	{
+		for (int32 J = 0; J < GridN - I; ++J)
+		{
+			Grid.Add({I, J}) = -1;
+		}
+	}
+
+	// generate interior (non-edge) vertices
+
+	const float GridStep = 1.0f / static_cast<float>(GridN - 1);
+	int32 InteriorVertexIndex = 0;
+	for (int32 I = 1; I < GridN - 1; ++I)
+	{
+		for (int32 J = 1; J < GridN - I - 1; ++J)
+		{
+			const float Bary1 = static_cast<float>(I) * GridStep;
+			const float Bary2 = static_cast<float>(J) * GridStep;
+			const float Bary0 = 1.0f - Bary1 - Bary2;
+
+			TriPatch.BaryCoord[InteriorVertexIndex] = { Bary0, Bary1, Bary2 };
+			Grid[{I, J}] = TriPatch.VIDs[InteriorVertexIndex];
+
+			++InteriorVertexIndex;
+		}
+	}
+
+	// corners
+
+	const int32 U = TriPatch.UVWCorners[0];
+	const int32 V = TriPatch.UVWCorners[1];
+	const int32 W = TriPatch.UVWCorners[2];
+
+	Grid[{0, 0}] = U;
+	Grid[{GridN - 1, 0}] = V;
+	Grid[{0, GridN - 1}] = W;
+
+
+	// edges
+	
+	// Edge UV
+	const int32 NumUVVertices = TriPatch.UVEdge.VIDs.Num();
+	const bool bUVEdgeHasAdditionalVertices = (NumUVVertices > (1 << Level) - 1);	// Whether the edge has extra vertices because the neighboring patch is higher level
+	for (int32 I = 0; I < GridN - 2; ++I)
+	{
+		if (bUVEdgeHasAdditionalVertices)
+		{
+			Grid[{I + 1, 0}] = TriPatch.UVEdge.bIsReversed ? TriPatch.UVEdge.VIDs[NumUVVertices - 2 - (2 * I)] : TriPatch.UVEdge.VIDs[2 * I + 1];
+		}
+		else
+		{
+			Grid[{I + 1, 0}] = TriPatch.UVEdge.bIsReversed ? TriPatch.UVEdge.VIDs[GridN - 3 - I] : TriPatch.UVEdge.VIDs[I];
+		}
+	}
+
+	// Edge UW
+	const int32 NumUWVertices = TriPatch.UWEdge.VIDs.Num();
+	const bool bUWEdgeHasAdditionalVertices = (NumUWVertices > (1 << Level) - 1);
+	for (int32 J = 0; J < GridN - 2; ++J)
+	{
+		if (bUWEdgeHasAdditionalVertices)
+		{
+			Grid[{0, J + 1}] = !TriPatch.UWEdge.bIsReversed ? TriPatch.UWEdge.VIDs[NumUWVertices - 2 - (2 * J)] : TriPatch.UWEdge.VIDs[2 * J + 1];
+		}
+		else
+		{
+			Grid[{0, J + 1}] = !TriPatch.UWEdge.bIsReversed ? TriPatch.UWEdge.VIDs[GridN - 3 - J] : TriPatch.UWEdge.VIDs[J];
+		}
+	}
+
+	// Edge VW
+	const int32 NumVWVertices = TriPatch.VWEdge.VIDs.Num();
+	const bool bVWEdgeHasAdditionalVertices = (NumVWVertices > (1 << Level) - 1);
+	for (int32 J = 0; J < GridN - 2; ++J)
+	{
+		const int32 I = GridN - 1 - J;
+
+		const TPair<int32, int32> MapKey{ I - 1, J + 1 };
+		ensure(Grid[MapKey] == -1);
+
+		if (bVWEdgeHasAdditionalVertices)
+		{
+			Grid[MapKey] = TriPatch.VWEdge.bIsReversed ? TriPatch.VWEdge.VIDs[NumVWVertices - 2 - (2 * J)] : TriPatch.VWEdge.VIDs[2 * J + 1];
+		}
+		else
+		{
+			Grid[MapKey] = TriPatch.VWEdge.bIsReversed ? TriPatch.VWEdge.VIDs[GridN - 3 - J] : TriPatch.VWEdge.VIDs[J];
+		}
+	}
+
+	// triangulate each grid cell
+
+	ensure(TriPatch.Triangles.Num() > 1);
+
+	int32 TriangleIndex = 0;
+	for (int I = 0; I < GridN - 1; ++I)
+	{
+		for (int J = 0; J < GridN - 1 - I; ++J)
+		{
+			TStaticArray<int32, 3> Triangle = { Grid[{I, J}], Grid[{I + 1, J}], Grid[{I, J + 1}] };
+			ensure(Triangle[0] >= 0);
+			ensure(Triangle[1] >= 0);
+			ensure(Triangle[2] >= 0);
+
+			// If we are on an edge where the neighboring triangle has a higher tessellation level, we'll have more edge vertices we need to think about
+			TStaticArray<int32, 3> Midpoints = {-1, -1, -1};
+			if (J == 0 && bUVEdgeHasAdditionalVertices)
+			{
+				const int32 NumVs = TriPatch.UVEdge.VIDs.Num();
+				Midpoints[0] = TriPatch.UVEdge.bIsReversed ? TriPatch.UVEdge.VIDs[NumVs - 1 - (2 * I)] : TriPatch.UVEdge.VIDs[2 * I];
+			}
+			if (I + J == GridN - 2 && bVWEdgeHasAdditionalVertices)
+			{
+				const int32 NumVs = TriPatch.VWEdge.VIDs.Num();
+				Midpoints[1] = !TriPatch.VWEdge.bIsReversed ? TriPatch.VWEdge.VIDs[NumVs - 1 - (2 * I)] : TriPatch.VWEdge.VIDs[2 * I];
+			}
+			if (I == 0 && bUWEdgeHasAdditionalVertices)
+			{
+				const int32 NumVs = TriPatch.UWEdge.VIDs.Num();
+				Midpoints[2] = !TriPatch.UWEdge.bIsReversed ? TriPatch.UWEdge.VIDs[NumVs - 1 - (2 * J)] : TriPatch.UWEdge.VIDs[2 * J];
+			}
+
+			// lower triangle(s) for the grid cell
+
+			if ((Midpoints[0] == -1) && (Midpoints[1] == -1) && (Midpoints[2] == -1))
+			{
+				// simple case, no extra triangles
+				TriPatch.Triangles[TriangleIndex++] = {Triangle[0], Triangle[1], Triangle[2]};
+			}
+			else
+			{
+				// adjacent to a patch with higher tessellation level, need extra triangles
+
+				TArray<FIndex3i, TInlineAllocator<4>> Tris;
+				RedGreenTessellationPatternHelpers::TriangulatePoly(Triangle, Midpoints, Tris);
+
+				for (const FIndex3i& Tri : Tris)
+				{
+					ensure(Tri[0] >= 0);
+					ensure(Tri[1] >= 0);
+					ensure(Tri[2] >= 0);
+
+					TriPatch.Triangles[TriangleIndex++] = Tri;
+				}
+			}
+
+			// upper triangle for the grid cell
+			if (I + J < GridN - 2)
+			{
+				const FIndex3i UpperTri = { Grid[{I + 1, J}], Grid[{I + 1, J + 1}], Grid[{I, J + 1}] };
+				ensure(UpperTri[0] >= 0);
+				ensure(UpperTri[1] >= 0);
+				ensure(UpperTri[2] >= 0);
+
+				TriPatch.Triangles[TriangleIndex++] = UpperTri;
+			}
+		}
+	}
+
+	ensure(TriangleIndex == TriPatch.Triangles.Num());
+}
+
+
+} // end namespace UE::Geometry
+} // end namespace UE
+
+#undef UE_API
diff --git a/Engine/Plugins/Runtime/GeometryProcessing/Source/DynamicMesh/Private/Operations/RedGreenTessellation.h b/Engine/Plugins/Runtime/GeometryProcessing/Source/DynamicMesh/Private/Operations/RedGreenTessellation.h
new file mode 100644
--- /dev/null
+++ b/Engine/Plugins/Runtime/GeometryProcessing/Source/DynamicMesh/Private/Operations/RedGreenTessellation.h
@@ -0,0 +1,48 @@
+// Copyright Epic Games, Inc. All Rights Reserved.
+
+#pragma once
+
+#include "Operations/SelectiveTessellate.h"
+
+#define UE_API DYNAMICMESH_API
+
+namespace UE
+{
+namespace Geometry
+{
+
+/**
+* Adaptive red-green subdivision. Triangles are tessellated to their specified level while also ensuring crack- and T-junction-free triangulation across 
+* triangles with different tessellation levels.
+* 
+* Each "triangle level" is the number of times to repeatedly apply red subdivision (1 to 4 triangle splits). We must simultaneously apply green subdivision on
+* neighboring triangles with lower levels to resolve T-junctions.
+* 
+* NOTE: Triangle tessellation levels cannot vary by more than one between neighboring triangles. E.g. you cannot have a triangle with level 4 adjacent to a triangle with level 2.
+*       Support for this could be added in a future version.
+*/
+
+class FRedGreenTessellationPattern : public FTessellationPattern
+{
+public:
+
+	FRedGreenTessellationPattern(const FDynamicMesh3* InMesh, const TArray<int>& InTriangleTessLevels);
+	virtual ~FRedGreenTessellationPattern() = default;
+
+	TArray<int32> TriangleLevels;
+
+private:
+
+	// FTessellationPattern  overrides
+	virtual EOperationValidationResult Validate() const override;
+	virtual int GetNumberOfNewVerticesForEdgePatch(const int InEdgeID) const override;
+	virtual int GetNumberOfNewVerticesForTrianglePatch(const int InTriangleID) const override;
+	virtual int GetNumberOfPatchTriangles(const int InTriangleID) const override;
+	virtual void TessellateEdgePatch(EdgePatch& EdgePatch) const override;
+	virtual void TessellateTriPatch(TrianglePatch& TriPatch) const override;
+};
+
+} // end namespace UE::Geometry
+} // end namespace UE
+
+#undef UE_API
diff --git a/Engine/Plugins/Runtime/GeometryProcessing/Source/DynamicMesh/Private/Operations/SelectiveTessellate.cpp b/Engine/Plugins/Runtime/GeometryProcessing/Source/DynamicMesh/Private/Operations/SelectiveTessellate.cpp
--- a/Engine/Plugins/Runtime/GeometryProcessing/Source/DynamicMesh/Private/Operations/SelectiveTessellate.cpp
+++ b/Engine/Plugins/Runtime/GeometryProcessing/Source/DynamicMesh/Private/Operations/SelectiveTessellate.cpp
@@ -10,6 +10,7 @@
 #include "Util/CompactMaps.h"
 #include "DynamicMesh/DynamicMeshAttributeSet.h"
 #include "DynamicMesh/DynamicVertexSkinWeightsAttribute.h"
+#include "Operations/RedGreenTessellation.h"
 
 using namespace UE::Geometry;
 
@@ -173,35 +174,50 @@ namespace SelectiveTessellateLocals
 		/** @return Array view of the data block containing the barycentric coordinates for new vertices added inside the triangle. */
 		TArrayView<FVector3d> MapInnerCoordBufferBlock(const int TriangleID)
 		{
-			TArrayView<FVector3d> ArrayView(InnerCoord);
-			const FIndex2i OffsetLength = InnerOffsets[TriangleID]; 
-			return ArrayView.Slice(OffsetLength[0], OffsetLength[1]);
+			if (InnerOffsets.Contains(TriangleID))
+			{
+				TArrayView<FVector3d> ArrayView(InnerCoord);
+				const FIndex2i OffsetLength = InnerOffsets[TriangleID];
+				return ArrayView.Slice(OffsetLength[0], OffsetLength[1]);
+			}
+			return TArrayView<FVector3d>();
 		}
 
 		/** @return Array view of the data block containing the ids for new vertices added inside the triangle. */
 		TArrayView<int> MapInnerIDBufferBlock(const int TriangleID)
 		{
-			TArrayView<int> ArrayView(InnerIDs);
-			const FIndex2i OffsetLength = InnerOffsets[TriangleID]; 
-			return ArrayView.Slice(OffsetLength[0], OffsetLength[1]);
+			if (InnerOffsets.Contains(TriangleID))
+			{
+				TArrayView<int> ArrayView(InnerIDs);
+				const FIndex2i OffsetLength = InnerOffsets[TriangleID];
+				return ArrayView.Slice(OffsetLength[0], OffsetLength[1]);
+			}
+			return TArrayView<int>();
 		}
 
 
 		/** @return Array view of a data block containing the triangle ids for new triangles. */
 		TArrayView<int> MapTriangleIDBufferBlock(const int TriangleID)
 		{
-			TArrayView<int> ArrayView(TriangleIDs);
-			checkSlow(TrianglesOffsets.Contains(TriangleID));
-			FIndex2i OffsetLengthPair = TrianglesOffsets[TriangleID]; 
-			return ArrayView.Slice(OffsetLengthPair[0], OffsetLengthPair[1]);
+			if (TrianglesOffsets.Contains(TriangleID))
+			{
+				TArrayView<int> ArrayView(TriangleIDs);
+				FIndex2i OffsetLengthPair = TrianglesOffsets[TriangleID];
+				return ArrayView.Slice(OffsetLengthPair[0], OffsetLengthPair[1]);
+			}
+			return TArrayView<int>();
 		}
 
 		/** @return Array view of a data block containing the triangle indices for new triangles. */
 		TArrayView<FIndex3i> MapTrianglesBufferBlock(const int TriangleID)
 		{
-			TArrayView<FIndex3i> ArrayView(Triangles);
-			FIndex2i OffsetLengthPair = TrianglesOffsets[TriangleID]; 
-			return ArrayView.Slice(OffsetLengthPair[0], OffsetLengthPair[1]);
+			if (TrianglesOffsets.Contains(TriangleID))
+			{
+				TArrayView<FIndex3i> ArrayView(Triangles);
+				FIndex2i OffsetLengthPair = TrianglesOffsets[TriangleID];
+				return ArrayView.Slice(OffsetLengthPair[0], OffsetLengthPair[1]);
+			}
+			return TArrayView<FIndex3i>();
 		}
 
 		//TODO: Add const versions of the Map* methods
@@ -520,7 +536,11 @@ namespace SelectiveTessellateLocals
 		{
 		}
 
-		
+		virtual EOperationValidationResult Validate() const override
+		{
+			return EOperationValidationResult::Ok;
+		}
+
 		void Init(const TArray<int>& InEdgeTessLevels, const TArray<int>& InInnerTessLevels) 
 		{
 			EdgeTessLevels = InEdgeTessLevels;
@@ -1257,7 +1277,13 @@ namespace SelectiveTessellateLocals
 				const FIndex3i Tri = Mesh->GetTriangle(TriangleID);
 				const int TriGrp = Mesh->GetTriangleGroup(TriangleID);
 				const FIndex3i MappedTri = CompactInfo.GetVertexMapping(Tri);
-				CompactInfo.SetTriangleMapping(TriangleID, ResultMesh->AppendTriangle(MappedTri, TriGrp));
+				const int NewTriID = ResultMesh->AppendTriangle(MappedTri, TriGrp);
+				CompactInfo.SetTriangleMapping(TriangleID, NewTriID);
+
+				if (TessInfo.TriangleIDToNewTriangleIDs)
+				{
+					TessInfo.TriangleIDToNewTriangleIDs->FindOrAdd(TriangleID).Add(NewTriID);
+				}
 			}
 		}
 		
@@ -1274,7 +1300,13 @@ namespace SelectiveTessellateLocals
 				const int NewTriangleID = TriangleIDBlock[Index];
 				const FIndex3i Tri = TrianglesBlock[Index];
 				const FIndex3i MappedTri = CompactInfo.GetVertexMapping(Tri);
-				CompactInfo.SetTriangleMapping(NewTriangleID, ResultMesh->AppendTriangle(MappedTri, TriGrp));
+				const int MappedTriId = ResultMesh->AppendTriangle(MappedTri, TriGrp);
+				CompactInfo.SetTriangleMapping(NewTriangleID, MappedTriId);
+
+				if (TessInfo.TriangleIDToNewTriangleIDs)
+				{
+					TessInfo.TriangleIDToNewTriangleIDs->FindOrAdd(TriangleID).Add(MappedTriId);
+				}
 			}	 
 		}
 		
@@ -2108,7 +2140,11 @@ TUniquePtr<FTessellationPattern> CreateUniformTessellationPattern(const FDynamic
 	return nullptr;
 }
 
-
+TUniquePtr<FTessellationPattern> FSelectiveTessellate::CreateRedGreenTessellationPattern(const FDynamicMesh3* InMesh, const TArray<int>& InTriangleTessLevels)
+{
+	TUniquePtr<FRedGreenTessellationPattern> Pattern = MakeUnique<FRedGreenTessellationPattern>(InMesh, InTriangleTessLevels);
+	return Pattern;
+}
 
 
 void FTessellationPattern::ComputeInnerConcentricTriangle(const FVector3d& V1,
diff --git a/Engine/Plugins/Runtime/GeometryProcessing/Source/DynamicMesh/Public/Operations/SelectiveTessellate.h b/Engine/Plugins/Runtime/GeometryProcessing/Source/DynamicMesh/Public/Operations/SelectiveTessellate.h
--- a/Engine/Plugins/Runtime/GeometryProcessing/Source/DynamicMesh/Public/Operations/SelectiveTessellate.h
+++ b/Engine/Plugins/Runtime/GeometryProcessing/Source/DynamicMesh/Public/Operations/SelectiveTessellate.h
@@ -116,7 +116,12 @@ public:
     virtual ~FTessellationPattern() 
     {
     }
-    
+
+	/**
+	* @return EOperationValidationResult::Ok if we can perform tessellation on the given mesh with this pattern
+	*/
+	virtual EOperationValidationResult Validate() const = 0;
+
     /** 
      * @return the number of new vertices added along the edge patch. If InEdgeID is not valid then return 
      * FTessellationPattern::InvalidIndex. 
@@ -249,6 +254,8 @@ public:
          * triangles affected by the tessellation.
          */
         TArray<int32>* SelectedVertices = nullptr;
+
+		TMap<int32, TArray<int32>>* TriangleIDToNewTriangleIDs = nullptr;
     };
 
     FTessellationInformation TessInfo;  
@@ -307,7 +314,14 @@ public:
 			return EOperationValidationResult::Failed_UnknownReason;
 		}
 
-		return EOperationValidationResult::Ok;
+		if (Pattern)
+		{
+			return Pattern->Validate();
+		}
+		else
+		{
+			return EOperationValidationResult::Failed_UnknownReason;
+		}
 	}
 
     inline void SetPattern(FTessellationPattern* InPattern)
@@ -496,6 +510,8 @@ public:
                                                                              const int InTessellationLevel,
                                                                              const TArray<int>& InTriangleList);
 
+	static UE_API TUniquePtr<FTessellationPattern> CreateRedGreenTessellationPattern(const FDynamicMesh3* InMesh, const TArray<int>& InTriangleTessLevels);
+
     //TODO: add more variation of different ways to provide tessellation levels similar to the ConcentricRings pattern.
 
 protected:
