Lattice Deformer Tool:
- Fix for not incorporating padding in Transform computation when no transform is provided externally
- Don't use padding when a transform is provided (i.e. when the lattice is initialized from an external source)
- Fix unit tests to including padding when checking lattice point locations

diff --git a/Engine/Plugins/Runtime/GeometryProcessing/Source/DynamicMesh/Private/Operations/FFDLattice.cpp b/Engine/Plugins/Runtime/GeometryProcessing/Source/DynamicMesh/Private/Operations/FFDLattice.cpp
--- a/Engine/Plugins/Runtime/GeometryProcessing/Source/DynamicMesh/Private/Operations/FFDLattice.cpp
+++ b/Engine/Plugins/Runtime/GeometryProcessing/Source/DynamicMesh/Private/Operations/FFDLattice.cpp
@@ -70,12 +70,14 @@ FFFDLattice::FFFDLattice(const FVector3i& InDims, const FDynamicMesh3& Mesh, flo
 		InitialBounds = FAxisAlignedBox3d(FVector3d::Zero(), FVector3d::Zero());
 	}
 
-	const FVector3d Center = InitialBounds.Center();
-	const FVector3d Extents = InitialBounds.Extents();
-	Transform.SetTranslation(Center - Extents);
+	// InitLattice() moves InitialBounds so that the min is at the origin, then applies the Padding parameter, so save the initial box center here
+	const FVector3d InputCenter = InitialBounds.Center();
 
 	InitLattice(Padding);
 
+	const FVector3d NewExtents = InitialBounds.Extents();
+	Transform.SetTranslation(InputCenter - NewExtents);
+
 	ComputeInitialEmbedding(Mesh);
 }
 
diff --git a/Engine/Plugins/Runtime/MeshModelingToolset/Source/MeshModelingTools/Private/LatticeDeformerTool.cpp b/Engine/Plugins/Runtime/MeshModelingToolset/Source/MeshModelingTools/Private/LatticeDeformerTool.cpp
--- a/Engine/Plugins/Runtime/MeshModelingToolset/Source/MeshModelingTools/Private/LatticeDeformerTool.cpp
+++ b/Engine/Plugins/Runtime/MeshModelingToolset/Source/MeshModelingTools/Private/LatticeDeformerTool.cpp
@@ -227,7 +227,7 @@ void ULatticeDeformerTool::InitializeLattice(TArray<FVector3d>& OutLatticePoints
 	}
 	else
 	{
-		Lattice = MakeShared<FFFDLattice, ESPMode::ThreadSafe>(GetLatticeResolution(), MeshToDeform, Settings->Padding, LatticeStorage->GetInitialBounds(), LatticeStorage->GetTransform());
+		Lattice = MakeShared<FFFDLattice, ESPMode::ThreadSafe>(GetLatticeResolution(), MeshToDeform, 0.0f, LatticeStorage->GetInitialBounds(), LatticeStorage->GetTransform());
 	}
 
 	Lattice->GenerateInitialLatticePositions(OutLatticePoints);
