ModelingTools: Fixed issue with the Mirror tool not averaging normals along the weld.

diff --git a/Engine/Plugins/Experimental/MeshModelingToolsetExp/Source/MeshModelingToolsExp/Private/MirrorTool.cpp b/Engine/Plugins/Experimental/MeshModelingToolsetExp/Source/MeshModelingToolsExp/Private/MirrorTool.cpp
--- a/Engine/Plugins/Experimental/MeshModelingToolsetExp/Source/MeshModelingToolsExp/Private/MirrorTool.cpp
+++ b/Engine/Plugins/Experimental/MeshModelingToolsetExp/Source/MeshModelingToolsExp/Private/MirrorTool.cpp
@@ -64,6 +64,15 @@ TUniquePtr<FDynamicMeshOperator> UMirrorOperatorFactory::MakeNewOperator()
 	MirrorOp->bCropFirst = MirrorTool->Settings->bCropAlongMirrorPlaneFirst;
 	MirrorOp->bSimplifyAlongNewEdges = MirrorTool->Settings->bSimplifyAlongCrop;
 	MirrorOp->bWeldAlongPlane = MirrorTool->Settings->bWeldVerticesOnMirrorPlane;
+	switch (MirrorTool->Settings->WeldVerticesNormalMode)
+	{
+	case EMeshMirrorWeldNormalMode::MirrorNormals:
+		MirrorOp->WeldNormalMode = UE::Geometry::EMeshMirrorNormalMode::MirrorNormals;
+		break;
+	case EMeshMirrorWeldNormalMode::AverageMirrorNormals:
+		MirrorOp->WeldNormalMode = UE::Geometry::EMeshMirrorNormalMode::AverageMirrorNormals;
+		break;
+	}
 	MirrorOp->PlaneTolerance = MirrorTool->Settings->PlaneTolerance;
 	MirrorOp->bAllowBowtieVertexCreation = MirrorTool->Settings->bAllowBowtieVertexCreation;
 
diff --git a/Engine/Plugins/Experimental/MeshModelingToolsetExp/Source/MeshModelingToolsExp/Public/MirrorTool.h b/Engine/Plugins/Experimental/MeshModelingToolsetExp/Source/MeshModelingToolsExp/Public/MirrorTool.h
--- a/Engine/Plugins/Experimental/MeshModelingToolsetExp/Source/MeshModelingToolsExp/Public/MirrorTool.h
+++ b/Engine/Plugins/Experimental/MeshModelingToolsetExp/Source/MeshModelingToolsExp/Public/MirrorTool.h
@@ -48,6 +48,16 @@ enum class EMirrorOperationMode : uint8
 	MirrorExisting = 1,
 };
 
+UENUM()
+enum class EMeshMirrorWeldNormalMode : uint8
+{
+	/** Normals are split and mirrored across the plane. */
+	MirrorNormals = 0,
+
+	/** Normals are averaged with their mirrored normal across the plane. */
+	AverageMirrorNormals = 1
+};
+
 UCLASS(MinimalAPI)
 class UMirrorToolProperties : public UInteractiveToolPropertySet
 {
@@ -72,6 +82,11 @@ public:
 		EditCondition = "OperationMode == EMirrorOperationMode::MirrorAndAppend", EditConditionHides))
 	bool bWeldVerticesOnMirrorPlane = true;
 
+	/** The normal compute method for vertices welded along the mirror plane. */
+	UPROPERTY(EditAnywhere, Category = Options, meta = (
+		EditCondition = "OperationMode == EMirrorOperationMode::MirrorAndAppend && bWeldVerticesOnMirrorPlane", EditConditionHides))
+	EMeshMirrorWeldNormalMode WeldVerticesNormalMode = EMeshMirrorWeldNormalMode::MirrorNormals;
+
 	/** Distance (in unscaled mesh space) to allow a point to be from the plane and still consider it "on the mirror plane". */
 	UPROPERTY(EditAnywhere, Category = Options, meta = (
 		EditCondition = "OperationMode == EMirrorOperationMode::MirrorAndAppend && bWeldVerticesOnMirrorPlane", EditConditionHides,
diff --git a/Engine/Plugins/Runtime/GeometryProcessing/Source/DynamicMesh/Private/Operations/MeshMirror.cpp b/Engine/Plugins/Runtime/GeometryProcessing/Source/DynamicMesh/Private/Operations/MeshMirror.cpp
--- a/Engine/Plugins/Runtime/GeometryProcessing/Source/DynamicMesh/Private/Operations/MeshMirror.cpp
+++ b/Engine/Plugins/Runtime/GeometryProcessing/Source/DynamicMesh/Private/Operations/MeshMirror.cpp
@@ -11,7 +11,7 @@ using namespace UE::Geometry;
 void MirrorAndAppendVertices(FDynamicMesh3& Mesh, const FVector3d& PlaneNormal, const FVector3d& PlaneOrigin,
 	TArray<int32>& OutVidToMirrorMap);
 void MirrorAndAppendVerticesWithWelding(FDynamicMesh3& Mesh, const FVector3d& PlaneNormal, const FVector3d& PlaneOrigin,
-	TArray<int32>& OutVidToMirrorMap, bool bAllowBowtieVertexCreation, double PlaneTolerance);
+	TArray<int32>& OutVidToMirrorMap, bool bAllowBowtieVertexCreation, double PlaneTolerance, EMeshMirrorNormalMode WeldNormalMode);
 FVertexInfo CreateMirrorVertex(const FVertexInfo& OriginalVertexInfo, const FVector3d& PlaneNormal, double SignedDistance);
 FVertexInfo CreateMirrorVertex(const FVertexInfo& OriginalVertexInfo, const FVector3d& PlaneNormal, const FVector3d& PlaneOrigin);
 
@@ -93,7 +93,7 @@ void FMeshMirror::MirrorAndAppend(FProgressCancel* Progress)
 	if (bWeldAlongPlane)
 	{
 		MirrorAndAppendVerticesWithWelding(*Mesh, PlaneNormal, PlaneOrigin, VidToMirrorMap,
-			bAllowBowtieVertexCreation, PlaneTolerance);
+			bAllowBowtieVertexCreation, PlaneTolerance, WeldNormalMode);
 	}
 	else
 	{
@@ -233,7 +233,7 @@ void MirrorAndAppendVertices(FDynamicMesh3& Mesh, const FVector3d& PlaneNormal,
  * @param PlaneTolerance Tolerance to use when determining whether a vertex is in the mirror plane.
  */
 void MirrorAndAppendVerticesWithWelding(FDynamicMesh3& Mesh, const FVector3d& PlaneNormal, const FVector3d& PlaneOrigin,
-	TArray<int32>& OutVidToMirrorMap, bool bAllowBowtieVertexCreation, double PlaneTolerance)
+	TArray<int32>& OutVidToMirrorMap, bool bAllowBowtieVertexCreation, double PlaneTolerance, EMeshMirrorNormalMode WeldNormalMode)
 {
 	// This will map an original vertex id to the id of its mirror. For vertices on the mirror plane,
 	// this will map the vertex id to the same vertex id, (unless we duplicate the vertex), though we
@@ -272,6 +272,8 @@ void MirrorAndAppendVerticesWithWelding(FDynamicMesh3& Mesh, const FVector3d& Pl
 	// Sort the weldable vertices into ones that will and will not actually be duplicated (vertices may
 	// need duplication if they will create non-manifold geometry after mirroring. Also, depending on settings,
 	// we may allow or disallow bowtie vertex creation.
+	const bool bHasVertexNormals = Mesh.HasVertexNormals();
+	const bool bHasAttributeNormals = Mesh.Attributes() && Mesh.Attributes()->PrimaryNormals();
 	for (int32 Vid : VerticesOnPlane)
 	{
 		// The rules we use are as follows:
@@ -346,17 +348,39 @@ void MirrorAndAppendVerticesWithWelding(FDynamicMesh3& Mesh, const FVector3d& Pl
 			FVector3d NewPosition = OldPosition - PlaneNormal.Dot(OldPosition - PlaneOrigin) * PlaneNormal;
 			Mesh.SetVertex(Vid, NewPosition);
 
-			if (Mesh.HasVertexNormals())
+			if (WeldNormalMode == EMeshMirrorNormalMode::AverageMirrorNormals)
 			{
-				FVector3f OldNormal = Mesh.GetVertexNormal(Vid);
-				FVector3f NewNormal = (OldNormal - OldNormal.Dot((FVector3f)PlaneNormal) * (FVector3f)PlaneNormal); // Not yet normalized
+				auto ComputeVertexNormal = [&PlaneNormal](const FVector3f& OldNormal, FVector3f& NewNormal)
+				{
+					NewNormal = (OldNormal - OldNormal.Dot((FVector3f)PlaneNormal) * (FVector3f)PlaneNormal); // Not yet normalized
+				
+					// For bowtie vertices, the projection of the normal onto the plane might be 0. The normal here is nonsensical anyway,
+					// so it doesn't matter what we do. Let's just leave it as is if we didn't get a projection.
+					return Normalize(NewNormal) != 0;
+				};
+
+				if (bHasVertexNormals)
+				{
+					FVector3f OldNormal = Mesh.GetVertexNormal(Vid);
+					FVector3f NewNormal;
+					if (ComputeVertexNormal(OldNormal, NewNormal))
+					{
+						Mesh.SetVertexNormal(Vid, NewNormal);
+					}
+				}
 
-				// For bowtie vertices, the projection of the normal onto the plane might be 0. The normal here is nonsensical anyway,
-				// so it doesn't matter what we do. Let's just leave it as is if we didn't get a projection.
-				if (Normalize(NewNormal) != 0)
+				if (bHasAttributeNormals)
 				{
-					// Projection worked
-					Mesh.SetVertexNormal(Vid, NewNormal);
+					FDynamicMeshNormalOverlay* NormalOverlay = Mesh.Attributes()->PrimaryNormals();
+					NormalOverlay->EnumerateVertexElements(Vid,[NormalOverlay, &ComputeVertexNormal](int, int ElementId, const FVector3f& OldNormal)
+					{
+						FVector3f NewNormal;
+						if (ComputeVertexNormal(OldNormal, NewNormal))
+						{
+							NormalOverlay->SetElement(ElementId, NewNormal);
+						}
+						return true;
+					});
 				}
 			}
 		}//end if not duplicating vertex
diff --git a/Engine/Plugins/Runtime/GeometryProcessing/Source/DynamicMesh/Public/Operations/MeshMirror.h b/Engine/Plugins/Runtime/GeometryProcessing/Source/DynamicMesh/Public/Operations/MeshMirror.h
--- a/Engine/Plugins/Runtime/GeometryProcessing/Source/DynamicMesh/Public/Operations/MeshMirror.h
+++ b/Engine/Plugins/Runtime/GeometryProcessing/Source/DynamicMesh/Public/Operations/MeshMirror.h
@@ -11,6 +11,14 @@ namespace UE
 namespace Geometry
 {
 
+enum class EMeshMirrorNormalMode
+{
+	/** Normals are split and mirrored across the mirror plane. */
+	MirrorNormals = 0,
+	/** Normals are averaged with their mirrored normal across the mirror plane. */
+	AverageMirrorNormals = 1
+};
+
 class FMeshMirror
 {
 
@@ -33,6 +41,9 @@ public:
 	/** Whether, when using MirrorAndAppend, vertices on the mirror plane should be welded. */
 	bool bWeldAlongPlane = true;
 
+	/** The normal compute method for welded vertices along the mirror plane. */
+	EMeshMirrorNormalMode WeldNormalMode = EMeshMirrorNormalMode::MirrorNormals;
+
 	/**
 	 * Whether, when welding, the creation of new bowtie vertices should be allowed (if a point lies 
 	 * in the mirror plane without an edge in the plane).
diff --git a/Engine/Plugins/Runtime/MeshModelingToolset/Source/ModelingOperators/Private/CompositionOps/MirrorOp.cpp b/Engine/Plugins/Runtime/MeshModelingToolset/Source/ModelingOperators/Private/CompositionOps/MirrorOp.cpp
--- a/Engine/Plugins/Runtime/MeshModelingToolset/Source/ModelingOperators/Private/CompositionOps/MirrorOp.cpp
+++ b/Engine/Plugins/Runtime/MeshModelingToolset/Source/ModelingOperators/Private/CompositionOps/MirrorOp.cpp
@@ -3,7 +3,6 @@
 #include "CompositionOps/MirrorOp.h"
 
 #include "Operations/MeshPlaneCut.h"
-#include "Operations/MeshMirror.h"
 
 using namespace UE::Geometry;
 
@@ -49,6 +48,7 @@ void FMirrorOp::CalculateResult(FProgressCancel* Progress)
 	// Set up the mirror operation
 	FMeshMirror Mirrorer(ResultMesh.Get(), LocalPlaneOrigin, LocalPlaneNormal);
 	Mirrorer.bWeldAlongPlane = bWeldAlongPlane;
+	Mirrorer.WeldNormalMode = WeldNormalMode;
 	Mirrorer.bAllowBowtieVertexCreation = bAllowBowtieVertexCreation;
 	Mirrorer.PlaneTolerance = PlaneTolerance;
 
diff --git a/Engine/Plugins/Runtime/MeshModelingToolset/Source/ModelingOperators/Public/CompositionOps/MirrorOp.h b/Engine/Plugins/Runtime/MeshModelingToolset/Source/ModelingOperators/Public/CompositionOps/MirrorOp.h
--- a/Engine/Plugins/Runtime/MeshModelingToolset/Source/ModelingOperators/Public/CompositionOps/MirrorOp.h
+++ b/Engine/Plugins/Runtime/MeshModelingToolset/Source/ModelingOperators/Public/CompositionOps/MirrorOp.h
@@ -5,6 +5,7 @@
 #include "CoreMinimal.h"
 #include "ModelingOperators.h"
 #include "Util/ProgressCancel.h"
+#include "Operations/MeshMirror.h"
 
 #define UE_API MODELINGOPERATORS_API
 
@@ -37,6 +38,9 @@ public:
 
 	/** Whether vertices on the mirror plane should be welded. Only relevant if bAppendToOriginal is true. */
 	bool bWeldAlongPlane = true;
+	
+	/** The normal compute method for welded vertices along the mirror plane. */
+	EMeshMirrorNormalMode WeldNormalMode = EMeshMirrorNormalMode::MirrorNormals;
 
 	/** Whether, when welding, new bowtie vertex creation should be allowed. */
 	bool bAllowBowtieVertexCreation = false;
