Fixed the dynamic mesh color transfer method handling of non-compact meshes when biasing element positions, and simplified the lerp-toward-centroid logic used

diff --git a/Engine/Plugins/Runtime/GeometryProcessing/Source/DynamicMesh/Private/Operations/TransferDynamicMeshAttributes.cpp b/Engine/Plugins/Runtime/GeometryProcessing/Source/DynamicMesh/Private/Operations/TransferDynamicMeshAttributes.cpp
--- a/Engine/Plugins/Runtime/GeometryProcessing/Source/DynamicMesh/Private/Operations/TransferDynamicMeshAttributes.cpp
+++ b/Engine/Plugins/Runtime/GeometryProcessing/Source/DynamicMesh/Private/Operations/TransferDynamicMeshAttributes.cpp
@@ -80,12 +80,6 @@ namespace TransferValuesLocals
 		return FVector3f((float)Normal.X, (float)Normal.Y, (float)Normal.Z);
 	}
 
-	struct FTriangleData
-	{
-		FVector3d Normal;
-		FVector3d Centroid;
-	};
-	
 	void GetBiasedElementPositions(const FDynamicMesh3& InOutTargetMesh, const float InRatio, TArray<FVector>& OutElementPositions)
 	{
 		const FDynamicMeshColorOverlay* ColorOverlay = InOutTargetMesh.Attributes()->PrimaryColors();
@@ -94,15 +88,12 @@ namespace TransferValuesLocals
 			return;
 		}
 		
-		// get normal + centroid of each face
-		TArray<FTriangleData> TrianglesData;
-		TrianglesData.Reserve(InOutTargetMesh.TriangleCount());
-		double DummyTriArea;
+		// get centroid of each face
+		TArray<FVector> TrianglesCentroids;
+		TrianglesCentroids.SetNumUninitialized(InOutTargetMesh.MaxTriangleID());
 		for (int TriIdx : InOutTargetMesh.TriangleIndicesItr())
 		{
-			FTriangleData Data;
-			InOutTargetMesh.GetTriInfo(TriIdx, Data.Normal, DummyTriArea, Data.Centroid);
-			TrianglesData.Emplace(MoveTemp(Data));
+			TrianglesCentroids[TriIdx] = InOutTargetMesh.GetTriCentroid(TriIdx);
 		}
 
 		// store the triangle for each element  
@@ -135,19 +126,9 @@ namespace TransferValuesLocals
 			const int32 TriangleIndex = ElementToTriangle[ElementID];
 			if (TriangleIndex != INDEX_NONE)
 			{
-				const FTriangleData& Triangle = TrianglesData[TriangleIndex];
-
-				// build triangle base with VertexPos - TriCentroid as the X axis and the triangle Normal as Z
-				const FVector& TriCentroid = Triangle.Centroid;
-				const FVector X = (VertexPos - TriCentroid).GetSafeNormal();
-				const FVector Y = Triangle.Normal.Cross(X);
-				const FQuat Rotation = FRotationMatrix::MakeFromXY(X, Y).ToQuat();
-				const FTransform TriTransform(Rotation, Triangle.Centroid);
-
-				// slightly move the vertex along the X axis to "shrink the triangle"
-				VertexPos = TriTransform.InverseTransformPosition(VertexPos);
-				VertexPos.X = VertexPos.X - (VertexPos.X * Ratio);
-				VertexPos = TriTransform.TransformPosition(VertexPos);
+				const FVector& Centroid = TrianglesCentroids[TriangleIndex];
+
+				VertexPos = FMath::Lerp(VertexPos, Centroid, Ratio);
 			}
 		}
 	}
