Make MeshConvexHull's grid sampler exit when mesh bounds are empty, to avoid ensure on grid w/ zero cell size

diff --git a/Engine/Plugins/Runtime/GeometryProcessing/Source/DynamicMesh/Private/Operations/MeshConvexHull.cpp b/Engine/Plugins/Runtime/GeometryProcessing/Source/DynamicMesh/Private/Operations/MeshConvexHull.cpp
--- a/Engine/Plugins/Runtime/GeometryProcessing/Source/DynamicMesh/Private/Operations/MeshConvexHull.cpp
+++ b/Engine/Plugins/Runtime/GeometryProcessing/Source/DynamicMesh/Private/Operations/MeshConvexHull.cpp
@@ -214,6 +214,12 @@ void FMeshConvexHull::GridSample(const FDynamicMesh3& Mesh,
 	// Simple spatial hash to find a representative vertex for each occupied grid cell
 
 	FAxisAlignedBox3d Bounds = Mesh.GetBounds();
+	
+	// Empty bounds, nothing to sample
+	if (Bounds.IsEmpty())
+	{
+		return;
+	}
 	Bounds.Min = Bounds.Min - 1e-4;			// Pad to avoid problems with vertices lying exactly on bounding box
 	Bounds.Max = Bounds.Max + 1e-4;
 	// Clamp grid resolution to prevent overflowing int
