Do not leave unused vertices in the cluster simplifier mesh

diff --git a/Engine/Plugins/Runtime/GeometryProcessing/Source/DynamicMesh/Private/Operations/MeshClusterSimplifier.cpp b/Engine/Plugins/Runtime/GeometryProcessing/Source/DynamicMesh/Private/Operations/MeshClusterSimplifier.cpp
--- a/Engine/Plugins/Runtime/GeometryProcessing/Source/DynamicMesh/Private/Operations/MeshClusterSimplifier.cpp
+++ b/Engine/Plugins/Runtime/GeometryProcessing/Source/DynamicMesh/Private/Operations/MeshClusterSimplifier.cpp
@@ -508,6 +508,39 @@ bool Simplify(const FDynamicMesh3& InMesh, FDynamicMesh3& ResultMesh, const FSim
 		break;
 	}
 
+	// We don't need the ToResVID mapping anymore.
+	// Empty it to make it clear it is not valid (since we don't remap it if we compact, below)
+	ToResVID.Empty();
+
+	// Remove any unused vertices and compact + fix maps accordingly
+	// Note we do this before transferring attributes, when it is cheaper,
+	// though this does require updating mappings to the source mesh accordingly
+	if (ResultMesh.RemoveUnusedVertices() > 0)
+	{
+		// Triangles are already compact by construction, so don't need to remap
+		check(ResultMesh.IsCompactT());
+
+		// Compact away unused vertices
+		FCompactMaps CompactMaps;
+		ResultMesh.CompactInPlace(&CompactMaps);
+
+		if (SimplifyOptions.bTransferAttributes)
+		{
+			// remap FromResVID, which maps Result VID -> Source VID
+			TArray<int32> NewFromResVID;
+			NewFromResVID.SetNumUninitialized(ResultMesh.MaxVertexID());
+			for (int32 PrevVID = 0, PrevNum = CompactMaps.NumVertexMappings(); PrevVID < PrevNum; ++PrevVID)
+			{
+				int32 NewVID = CompactMaps.GetVertexMapping(PrevVID);
+				if (NewVID != IndexConstants::InvalidID)
+				{
+					NewFromResVID[NewVID] = FromResVID[PrevVID];
+				}
+			}
+			FromResVID = MoveTemp(NewFromResVID);
+		}
+	}
+
 	///
 	/// Step 4: After accepting the final ResultMesh triangulation, copy the input mesh's attributes (UVs, materials, etc) over as well
 	/// 
diff --git a/Engine/Source/Runtime/GeometryCore/Private/DynamicMesh/DynamicMesh3_Edits.cpp b/Engine/Source/Runtime/GeometryCore/Private/DynamicMesh/DynamicMesh3_Edits.cpp
--- a/Engine/Source/Runtime/GeometryCore/Private/DynamicMesh/DynamicMesh3_Edits.cpp
+++ b/Engine/Source/Runtime/GeometryCore/Private/DynamicMesh/DynamicMesh3_Edits.cpp
@@ -286,13 +286,15 @@ EMeshResult FDynamicMesh3::InsertTriangle(int tid, const FIndex3i& tv, int gid,
 
 
 
-void FDynamicMesh3::RemoveUnusedVertices()
+int32 FDynamicMesh3::RemoveUnusedVertices()
 {
+	int32 NumRemoved = 0;
 	for (int32 VID = 0; VID < MaxVertexID(); ++VID)
 	{
 		// If vertex exists but is not referenced by any triangles
 		if (VertexRefCounts.GetRefCount(VID) == 1)
 		{
+			NumRemoved++;
 			VertexRefCounts.Decrement(VID);
 
 			if (HasAttributes())
@@ -304,7 +306,11 @@ void FDynamicMesh3::RemoveUnusedVertices()
 		}
 	}
 
-	UpdateChangeStamps(true, true);
+	if (NumRemoved > 0)
+	{
+		UpdateChangeStamps(true, true);
+	}
+	return NumRemoved;
 }
 
 bool FDynamicMesh3::HasUnusedVertices() const
diff --git a/Engine/Source/Runtime/GeometryCore/Public/DynamicMesh/DynamicMesh3.h b/Engine/Source/Runtime/GeometryCore/Public/DynamicMesh/DynamicMesh3.h
--- a/Engine/Source/Runtime/GeometryCore/Public/DynamicMesh/DynamicMesh3.h
+++ b/Engine/Source/Runtime/GeometryCore/Public/DynamicMesh/DynamicMesh3.h
@@ -1308,8 +1308,9 @@ public:
 
 	/**
 	 * Remove unused vertices. Note: Does not compact the remaining vertices.
+	 * @return number of removed vertices
 	 */
-	GEOMETRYCORE_API void RemoveUnusedVertices();
+	GEOMETRYCORE_API int32 RemoveUnusedVertices();
 
 	/**
 	 * @return true if any vertices are unused (not in any triangles)
