Red-Green Tessellation: remove restrictions on relative tessellation depth across neighboring triangles. We should now be able to subdivide any triangle to any depth regardless of the depth of its neighbors and still get a conforming mesh out

diff --git a/Engine/Plugins/Runtime/GeometryProcessing/Source/DynamicMesh/Private/Operations/RedGreenTessellation.cpp b/Engine/Plugins/Runtime/GeometryProcessing/Source/DynamicMesh/Private/Operations/RedGreenTessellation.cpp
--- a/Engine/Plugins/Runtime/GeometryProcessing/Source/DynamicMesh/Private/Operations/RedGreenTessellation.cpp
+++ b/Engine/Plugins/Runtime/GeometryProcessing/Source/DynamicMesh/Private/Operations/RedGreenTessellation.cpp
@@ -2,6 +2,7 @@
 
 #include "Operations/RedGreenTessellation.h"
 #include "DynamicMesh/DynamicMesh3.h"
+#include "CompGeom/PolygonTriangulation.h"
 
 namespace UE
 {
@@ -10,53 +11,56 @@ namespace Geometry
 namespace RedGreenTessellationPatternHelpers
 {
 	/*
-	* Given a triangle and optional midpoint vertices, produce a valid triangulation
-	*/
-	void TriangulatePoly(const TStaticArray<int32, 3>& Triangle, const TStaticArray<int32, 3>& MidPoints, TArray<FIndex3i, TInlineAllocator<4>>& Tris)
+	 * Given a triangle and additional vertices for each triangle edge, produce a valid triangulation
+	 */
+	static void TriangulatePoly(const FIndex3i& Triangle, const TStaticArray<TArray<int32>, 3>& EdgePoints, TArray<FIndex3i>& Tris)
 	{
-		int32 NumSplits = 0;
-		for (int32 MidpointIndex = 0; MidpointIndex < 3; ++MidpointIndex)
+		int32 NumSplitEdges = 0;
+		int32 MaxEdgeTessellationLevel = 0;
+		for (int32 EdgeIndex = 0; EdgeIndex < 3; ++EdgeIndex)
 		{
-			if (MidPoints[MidpointIndex] >= 0)
+			if (EdgePoints[EdgeIndex].Num() > 0)
 			{
-				++NumSplits;
+				++NumSplitEdges;
+				MaxEdgeTessellationLevel = FMath::Max(MaxEdgeTessellationLevel, EdgePoints[EdgeIndex].Num());
 			}
 		}
 
-		if (NumSplits == 0)
+		if (NumSplitEdges == 0)
 		{
 			Tris.Add({ Triangle[0], Triangle[1], Triangle[2] });
 		}
-		else if (NumSplits == 1)
+		else if (NumSplitEdges == 1 && MaxEdgeTessellationLevel == 1)
 		{
 			// Green triangle
 			for (int32 TriVertIndex = 0; TriVertIndex < 3; ++TriVertIndex)
 			{
-				if (MidPoints[TriVertIndex] >= 0)
+				if (EdgePoints[TriVertIndex].Num() > 0)
 				{
+					ensure(EdgePoints[TriVertIndex].Num() == 1);
 					const int32 A = Triangle[TriVertIndex];
 					const int32 B = Triangle[(TriVertIndex + 1) % 3];
 					const int32 C = Triangle[(TriVertIndex + 2) % 3];
-					const int32 M = MidPoints[TriVertIndex];
+					const int32 M = EdgePoints[TriVertIndex][0];
 					Tris.Add({ A, M, C });
 					Tris.Add({ M, B, C });
 					break;
 				}
 			}
 		}
-		else if (NumSplits == 2)
+		else if (NumSplitEdges == 2 && MaxEdgeTessellationLevel == 1)
 		{
 			// "Double green" (?) triangle
 			for (int32 TriVertIndex = 0; TriVertIndex < 3; ++TriVertIndex)
 			{
 				const int32 Next = (TriVertIndex + 1) % 3;
-				if (MidPoints[TriVertIndex] >= 0 && MidPoints[Next] >= 0)
+				if (EdgePoints[TriVertIndex].Num() == 1 && EdgePoints[Next].Num() == 1)
 				{
 					const int32 A = Triangle[TriVertIndex];
 					const int32 B = Triangle[Next];
 					const int32 C = Triangle[(TriVertIndex + 2) % 3];
-					const int32 AB = MidPoints[TriVertIndex];
-					const int32 BC = MidPoints[Next];
+					const int32 AB = EdgePoints[TriVertIndex][0];
+					const int32 BC = EdgePoints[Next][0];
 					Tris.Add({ A, AB, C });
 					Tris.Add({ AB, BC, C });
 					Tris.Add({ AB, B, BC });
@@ -64,22 +68,65 @@ namespace RedGreenTessellationPatternHelpers
 				}
 			}
 		}
-		else
+		else if (NumSplitEdges == 3 && MaxEdgeTessellationLevel == 1)
 		{
 			// Red triangle
-			ensure(NumSplits == 3);	
+			ensure(EdgePoints[0].Num() == 1 && EdgePoints[1].Num() == 1 && EdgePoints[2].Num() == 1);
 			const int32 A = Triangle[0];
 			const int32 B = Triangle[1];
 			const int32 C = Triangle[2];
-			const int32 AB = MidPoints[0];
-			const int32 BC = MidPoints[1];
-			const int32 CA = MidPoints[2];
+			const int32 AB = EdgePoints[0][0];
+			const int32 BC = EdgePoints[1][0];
+			const int32 CA = EdgePoints[2][0];
 			Tris.Add({ A, AB, CA });
 			Tris.Add({ AB, B, BC });
 			Tris.Add({ CA, BC, C });
 			Tris.Add({ AB, BC, CA });
 		}
+		else
+		{
+			// Too complicated -- go full poly triangulation mode
+			TArray<int32> Polygon;
+			TArray<FVector2d> PolygonVertices;
+
+			const TStaticArray<FVector2d, 3> Corners = { FVector2d{0,0}, FVector2d{1,0}, FVector2d{0,1} };
+
+			for (int32 TriEdgeIndex = 0; TriEdgeIndex < 3; ++TriEdgeIndex)
+			{
+				PolygonVertices.Add(Corners[TriEdgeIndex]);
+				Polygon.Add(Triangle[TriEdgeIndex]);
+
+				const TArray<int32>& Edge = EdgePoints[TriEdgeIndex];
+				if (Edge.Num() > 0)
+				{
+					const double Dx = 1.0 / (Edge.Num() + 1);
+					const FVector2d& ThisCorner = Corners[TriEdgeIndex];
+					const FVector2d& NextCorner = Corners[(TriEdgeIndex + 1) % 3];
+
+					for (int32 EdgeVertexIndex = 0; EdgeVertexIndex < Edge.Num(); ++EdgeVertexIndex)
+					{
+						Polygon.Add(Edge[EdgeVertexIndex]);
+						const double Alpha = static_cast<double>(EdgeVertexIndex + 1) * Dx;
+						PolygonVertices.Add(ThisCorner + Alpha * (NextCorner - ThisCorner));
+					}
+				}
+			}
+
+			const int32 ExpectedNumberOfTriangles = PolygonVertices.Num() - 2;
+
+			TArray<FIndex3i> LocalTris;
+			constexpr bool bOrientAsHoleFill = false;
+			PolygonTriangulation::TriangulateSimplePolygon(PolygonVertices, LocalTris, bOrientAsHoleFill);
+
+			ensure(LocalTris.Num() == ExpectedNumberOfTriangles);
+
+			for (const FIndex3i& LocalTri : LocalTris)
+			{
+				Tris.Add(FIndex3i{ Polygon[LocalTri[0]], Polygon[LocalTri[1]], Polygon[LocalTri[2]] });
+			}
+		}
 	}
+
 }	// namespace RedGreenTessellationPatternHelpers
 
 
@@ -164,23 +211,6 @@ EOperationValidationResult FRedGreenTessellationPattern::Validate() const
 		return EOperationValidationResult::Failed_UnknownReason;
 	}
 
-	// Verify the triangle tessellation levels don't vary by more than 1 between neighboring triangles
-	for (int32 TriangleIndex = 0; TriangleIndex < Mesh->MaxTriangleID(); ++TriangleIndex)
-	{
-		const FIndex3i Neighbors = Mesh->GetTriNeighbourTris(TriangleIndex);
-
-		for (const int NeighborIndex : {Neighbors[0], Neighbors[1], Neighbors[2]})
-		{
-			if (Mesh->IsTriangle(NeighborIndex))
-			{
-				if (FMath::Abs(TriangleLevels[TriangleIndex] - TriangleLevels[NeighborIndex]) > 1)
-				{
-					return EOperationValidationResult::Failed_UnknownReason;
-				}
-			}
-		}
-	}
-
 	return EOperationValidationResult::Ok;
 }
 
@@ -233,7 +263,6 @@ int FRedGreenTessellationPattern::GetNumberOfPatchTriangles(const int InTriangle
 		const int NeighborTri = NeighborTris[NIndex];
 		if (Mesh->IsTriangle(NeighborTri) && TriangleLevels[NeighborTri] > Level)
 		{
-			ensure(TriangleLevels[NeighborTri] == Level + 1);
 			// We need to subdivide new segments along this edge with green triangles
 			NewTriangles += (1 << TriangleLevels[NeighborTri]) - (1 << Level);
 		}
@@ -259,48 +288,7 @@ void FRedGreenTessellationPattern::TessellateTriPatch(TrianglePatch& TriPatch) c
 		TriPatch.Triangles[0] = TriPatch.UVWCorners;
 		return;
 	}
-
-	int32 NumTessellatedEdges = 0;
-	NumTessellatedEdges += TriPatch.UVEdge.VIDs.Num() > 0 ? 1 : 0;
-	NumTessellatedEdges += TriPatch.VWEdge.VIDs.Num() > 0 ? 1 : 0;
-	NumTessellatedEdges += TriPatch.UWEdge.VIDs.Num() > 0 ? 1 : 0;
-
-	if (NumTessellatedEdges == 1)
-	{
-		// Green triangle
-
-		if (TriPatch.UVEdge.VIDs.Num() > 0)
-		{
-			const int32 V0 = TriPatch.UVWCorners[0];
-			const int32 V1 = TriPatch.UVWCorners[1];
-			const int32 V2 = TriPatch.UVWCorners[2];
-			const int32 Mid = TriPatch.UVEdge.VIDs[0];
-			TriPatch.Triangles[0] = { V2, V0, Mid };
-			TriPatch.Triangles[1] = { V2, Mid, V1 };
-		}
-		else if (TriPatch.VWEdge.VIDs.Num() > 0)
-		{
-			const int32 V0 = TriPatch.UVWCorners[1];
-			const int32 V1 = TriPatch.UVWCorners[2];
-			const int32 V2 = TriPatch.UVWCorners[0];
-			const int32 Mid = TriPatch.VWEdge.VIDs[0];
-			TriPatch.Triangles[0] = { V2, V0, Mid };
-			TriPatch.Triangles[1] = { V2, Mid, V1 };
-		}
-		else
-		{
-			ensure(TriPatch.UWEdge.VIDs.Num() > 0);
-			const int32 V0 = TriPatch.UVWCorners[2];
-			const int32 V1 = TriPatch.UVWCorners[0];
-			const int32 V2 = TriPatch.UVWCorners[1];
-			const int32 Mid = TriPatch.UWEdge.VIDs[0];
-			TriPatch.Triangles[0] = { V2, V0, Mid };
-			TriPatch.Triangles[1] = { V2, Mid, V1 };
-		}
-
-		return;
-	}
-
+	
 	const int32 Level = TriangleLevels[TriPatch.TriangleID];
 
 	// Create a grid of vertices which will then be triangulated. The grid structure stores vertex indices
@@ -349,42 +337,68 @@ void FRedGreenTessellationPattern::TessellateTriPatch(TrianglePatch& TriPatch) c
 	Grid[{GridN - 1, 0}] = V;
 	Grid[{0, GridN - 1}] = W;
 
-
 	// edges
-	
+
+	auto Reverse = [](const TArrayView<int>& Edge) -> TArray<int>
+	{
+		TArray<int> RevEdge;
+		for (int EV : Edge)
+		{
+			RevEdge.Add(EV);
+		}
+		Algo::Reverse(RevEdge);
+		return RevEdge;
+	};
+
+	const TArray<int> EdgeUVSortedVIDs = TriPatch.UVEdge.bIsReversed ? Reverse(TriPatch.UVEdge.VIDs) : TArray<int>(TriPatch.UVEdge.VIDs);
+	const TArray<int> EdgeVWSortedVIDs = TriPatch.VWEdge.bIsReversed ? Reverse(TriPatch.VWEdge.VIDs) : TArray<int>(TriPatch.VWEdge.VIDs);
+	const TArray<int> EdgeUWSortedVIDs = TriPatch.UWEdge.bIsReversed ? Reverse(TriPatch.UWEdge.VIDs) : TArray<int>(TriPatch.UWEdge.VIDs);
+
+	const FIndex3i NeighborTris = Mesh->GetTriNeighbourTris(TriPatch.TriangleID);
+
 	// Edge UV
 	const int32 NumUVVertices = TriPatch.UVEdge.VIDs.Num();
-	const bool bUVEdgeHasAdditionalVertices = (NumUVVertices > (1 << Level) - 1);	// Whether the edge has extra vertices because the neighboring patch is higher level
+
+	// Whether the edge has extra vertices because the neighboring patch is higher level. If so, then we need to only grab every Nth vertex from the TriPatch Edge
+	const bool bUVEdgeHasAdditionalVertices = (NumUVVertices > (1 << Level) - 1);
+
+	const int32 NbrUVLevel = Mesh->IsTriangle(NeighborTris[0]) ? TriangleLevels[NeighborTris[0]] : Level;
 	for (int32 I = 0; I < GridN - 2; ++I)
 	{
 		if (bUVEdgeHasAdditionalVertices)
 		{
-			Grid[{I + 1, 0}] = TriPatch.UVEdge.bIsReversed ? TriPatch.UVEdge.VIDs[NumUVVertices - 2 - (2 * I)] : TriPatch.UVEdge.VIDs[2 * I + 1];
+			const int32 StepSize = (1 << (NbrUVLevel - Level));
+			Grid[{I + 1, 0}] = EdgeUVSortedVIDs[StepSize * (I + 1) - 1];
 		}
 		else
 		{
-			Grid[{I + 1, 0}] = TriPatch.UVEdge.bIsReversed ? TriPatch.UVEdge.VIDs[GridN - 3 - I] : TriPatch.UVEdge.VIDs[I];
+			Grid[{I + 1, 0}] = EdgeUVSortedVIDs[I];
 		}
 	}
 
 	// Edge UW
 	const int32 NumUWVertices = TriPatch.UWEdge.VIDs.Num();
 	const bool bUWEdgeHasAdditionalVertices = (NumUWVertices > (1 << Level) - 1);
+	const int32 NbrUWLevel = Mesh->IsTriangle(NeighborTris[2]) ? TriangleLevels[NeighborTris[2]] : Level;
+	// Vertices stored in UWEdge initially run along the edge from vertex W to vertex U. But our grid places W at the (i=0, j=N-1) grid point. 
+	// So to fill in the i = 0 column from j = 0 to N-1 we need to traverse UWEdge in reverse
 	for (int32 J = 0; J < GridN - 2; ++J)
 	{
 		if (bUWEdgeHasAdditionalVertices)
 		{
-			Grid[{0, J + 1}] = !TriPatch.UWEdge.bIsReversed ? TriPatch.UWEdge.VIDs[NumUWVertices - 2 - (2 * J)] : TriPatch.UWEdge.VIDs[2 * J + 1];
+			const int32 StepSize = (1 << (NbrUWLevel - Level));
+			Grid[{0, J + 1}] = EdgeUWSortedVIDs[NumUWVertices - StepSize * (J + 1)];
 		}
 		else
 		{
-			Grid[{0, J + 1}] = !TriPatch.UWEdge.bIsReversed ? TriPatch.UWEdge.VIDs[GridN - 3 - J] : TriPatch.UWEdge.VIDs[J];
+			Grid[{0, J + 1}] = EdgeUWSortedVIDs[GridN - 3 - J];
 		}
 	}
 
 	// Edge VW
 	const int32 NumVWVertices = TriPatch.VWEdge.VIDs.Num();
 	const bool bVWEdgeHasAdditionalVertices = (NumVWVertices > (1 << Level) - 1);
+	const int32 NbrVWLevel = Mesh->IsTriangle(NeighborTris[1]) ? TriangleLevels[NeighborTris[1]] : Level;
 	for (int32 J = 0; J < GridN - 2; ++J)
 	{
 		const int32 I = GridN - 1 - J;
@@ -394,14 +408,16 @@ void FRedGreenTessellationPattern::TessellateTriPatch(TrianglePatch& TriPatch) c
 
 		if (bVWEdgeHasAdditionalVertices)
 		{
-			Grid[MapKey] = TriPatch.VWEdge.bIsReversed ? TriPatch.VWEdge.VIDs[NumVWVertices - 2 - (2 * J)] : TriPatch.VWEdge.VIDs[2 * J + 1];
+			const int32 StepSize = (1 << (NbrVWLevel - Level));
+			Grid[MapKey] = EdgeVWSortedVIDs[StepSize * (J + 1) - 1];
 		}
 		else
 		{
-			Grid[MapKey] = TriPatch.VWEdge.bIsReversed ? TriPatch.VWEdge.VIDs[GridN - 3 - J] : TriPatch.VWEdge.VIDs[J];
+			Grid[MapKey] = EdgeVWSortedVIDs[J];
 		}
 	}
 
+
 	// triangulate each grid cell
 
 	ensure(TriPatch.Triangles.Num() > 1);
@@ -411,42 +427,75 @@ void FRedGreenTessellationPattern::TessellateTriPatch(TrianglePatch& TriPatch) c
 	{
 		for (int J = 0; J < GridN - 1 - I; ++J)
 		{
-			TStaticArray<int32, 3> Triangle = { Grid[{I, J}], Grid[{I + 1, J}], Grid[{I, J + 1}] };
-			ensure(Triangle[0] >= 0);
-			ensure(Triangle[1] >= 0);
-			ensure(Triangle[2] >= 0);
+			FIndex3i GridTriangle = { Grid[{I, J}], Grid[{I + 1, J}], Grid[{I, J + 1}] };
+			ensure(GridTriangle[0] >= 0);
+			ensure(GridTriangle[1] >= 0);
+			ensure(GridTriangle[2] >= 0);
+
+			auto FillEdgePoints = [&](const int32 NbrLevel, const int32 I, const TArray<int32>& InFullEdgePoints, TArray<int32>& EdgePoints, bool bGridCountReversed)
+			{
+				const int32 StepSize = (1 << (NbrLevel - Level));
+
+				int32 Start, End;
+				if (!bGridCountReversed)
+				{
+					Start = I * StepSize;
+					End = (I + 1) * StepSize - 1;
+					for (int32 V = Start; V < End; ++V)
+					{
+						EdgePoints.Add(InFullEdgePoints[V]);
+					}
+				}
+				else
+				{
+					const int J = GridN - I - 1;
+					const int JNext = J - 1;
+					
+					Start = JNext * StepSize;
+					End = J * StepSize - 1;
+
+					for (int32 V = Start; V < End; ++V)
+					{
+						EdgePoints.Add(InFullEdgePoints[V]);
+					}
+				}
+
+				ensure(EdgePoints.Num() == StepSize - 1);
+			};
 
 			// If we are on an edge where the neighboring triangle has a higher tessellation level, we'll have more edge vertices we need to think about
-			TStaticArray<int32, 3> Midpoints = {-1, -1, -1};
+			TStaticArray<TArray<int32>, 3> EdgePoints;		// EdgePoints are new vertices in the UV, WU, and VW edges respectively. 
 			if (J == 0 && bUVEdgeHasAdditionalVertices)
 			{
-				const int32 NumVs = TriPatch.UVEdge.VIDs.Num();
-				Midpoints[0] = TriPatch.UVEdge.bIsReversed ? TriPatch.UVEdge.VIDs[NumVs - 1 - (2 * I)] : TriPatch.UVEdge.VIDs[2 * I];
+				FillEdgePoints(NbrUVLevel, I, EdgeUVSortedVIDs, EdgePoints[0], false);
 			}
 			if (I + J == GridN - 2 && bVWEdgeHasAdditionalVertices)
 			{
-				const int32 NumVs = TriPatch.VWEdge.VIDs.Num();
-				Midpoints[1] = !TriPatch.VWEdge.bIsReversed ? TriPatch.VWEdge.VIDs[NumVs - 1 - (2 * I)] : TriPatch.VWEdge.VIDs[2 * I];
+				FillEdgePoints(NbrVWLevel, J, EdgeVWSortedVIDs, EdgePoints[1], false);
 			}
 			if (I == 0 && bUWEdgeHasAdditionalVertices)
 			{
-				const int32 NumVs = TriPatch.UWEdge.VIDs.Num();
-				Midpoints[2] = !TriPatch.UWEdge.bIsReversed ? TriPatch.UWEdge.VIDs[NumVs - 1 - (2 * J)] : TriPatch.UWEdge.VIDs[2 * J];
+				FillEdgePoints(NbrUWLevel, J, EdgeUWSortedVIDs, EdgePoints[2], true);
+			}
+
+			if (EdgePoints[0].Num() == 1 && EdgePoints[1].Num() == 1 && EdgePoints[2].Num() == 1)
+			{
+				ensure(false);
 			}
 
 			// lower triangle(s) for the grid cell
 
-			if ((Midpoints[0] == -1) && (Midpoints[1] == -1) && (Midpoints[2] == -1))
+			if ((EdgePoints[0].Num() == 0) && (EdgePoints[1].Num() == 0) && (EdgePoints[2].Num() == 0))
 			{
 				// simple case, no extra triangles
-				TriPatch.Triangles[TriangleIndex++] = {Triangle[0], Triangle[1], Triangle[2]};
+				TriPatch.Triangles[TriangleIndex++] = GridTriangle;
 			}
 			else
 			{
 				// adjacent to a patch with higher tessellation level, need extra triangles
 
-				TArray<FIndex3i, TInlineAllocator<4>> Tris;
-				RedGreenTessellationPatternHelpers::TriangulatePoly(Triangle, Midpoints, Tris);
+				TArray<FIndex3i> Tris;
+				RedGreenTessellationPatternHelpers::TriangulatePoly(GridTriangle, EdgePoints, Tris);
 
 				for (const FIndex3i& Tri : Tris)
 				{
diff --git a/Engine/Plugins/Runtime/GeometryProcessing/Source/DynamicMesh/Private/Operations/RedGreenTessellation.h b/Engine/Plugins/Runtime/GeometryProcessing/Source/DynamicMesh/Private/Operations/RedGreenTessellation.h
--- a/Engine/Plugins/Runtime/GeometryProcessing/Source/DynamicMesh/Private/Operations/RedGreenTessellation.h
+++ b/Engine/Plugins/Runtime/GeometryProcessing/Source/DynamicMesh/Private/Operations/RedGreenTessellation.h
@@ -17,9 +17,6 @@ namespace Geometry
 * 
 * Each "triangle level" is the number of times to repeatedly apply red subdivision (1 to 4 triangle splits). We must simultaneously apply green subdivision on
 * neighboring triangles with lower levels to resolve T-junctions.
-* 
-* NOTE: Triangle tessellation levels cannot vary by more than one between neighboring triangles. E.g. you cannot have a triangle with level 4 adjacent to a triangle with level 2.
-*       Support for this could be added in a future version.
 */
 
 class FRedGreenTessellationPattern : public FTessellationPattern
