Red-Green Tessellation: remove an ensure that fires when a triangle of level zero is surrounded by triangles of level 1. Also try to silence some linter warnings.

diff --git a/Engine/Plugins/Runtime/GeometryProcessing/Source/DynamicMesh/Private/Operations/RedGreenTessellation.cpp b/Engine/Plugins/Runtime/GeometryProcessing/Source/DynamicMesh/Private/Operations/RedGreenTessellation.cpp
--- a/Engine/Plugins/Runtime/GeometryProcessing/Source/DynamicMesh/Private/Operations/RedGreenTessellation.cpp
+++ b/Engine/Plugins/Runtime/GeometryProcessing/Source/DynamicMesh/Private/Operations/RedGreenTessellation.cpp
@@ -114,6 +114,7 @@ namespace RedGreenTessellationPatternHelpers
 
 			const int32 ExpectedNumberOfTriangles = PolygonVertices.Num() - 2;
 
+			// The input "polygon" is geometrically triangular by construction, so the resulting triangulation will have no foldovers or topological issues in 3D
 			TArray<FIndex3i> LocalTris;
 			constexpr bool bOrientAsHoleFill = false;
 			PolygonTriangulation::TriangulateSimplePolygon(PolygonVertices, LocalTris, bOrientAsHoleFill);
@@ -296,7 +297,7 @@ void FRedGreenTessellationPattern::TessellateTriPatch(TrianglePatch& TriPatch) c
 
 	const int32 GridN = (1 << Level) + 1;
 
-	// TODO: If using a TMap is slow, investigate creating a triangular grid structure. Maybe similar to TDenseGrid2 but only storing the lower half?
+	// TODO: If using a TMap is slow, investigate creating a triangular grid structure which would allow more direct access. Maybe similar to TDenseGrid2 but only storing the lower half?
 	TMap<TPair<int32, int32>, int32> Grid;
 	Grid.Reserve(GridN * GridN / 2);
 
@@ -367,6 +368,7 @@ void FRedGreenTessellationPattern::TessellateTriPatch(TrianglePatch& TriPatch) c
 	{
 		if (bUVEdgeHasAdditionalVertices)
 		{
+			checkf(NbrUVLevel >= Level, TEXT("bUVEdgeHasAdditionalVertices is true but Level > NbrUVLevel"));
 			const int32 StepSize = (1 << (NbrUVLevel - Level));
 			Grid[{I + 1, 0}] = EdgeUVSortedVIDs[StepSize * (I + 1) - 1];
 		}
@@ -386,6 +388,7 @@ void FRedGreenTessellationPattern::TessellateTriPatch(TrianglePatch& TriPatch) c
 	{
 		if (bUWEdgeHasAdditionalVertices)
 		{
+			checkf(NbrUWLevel >= Level, TEXT("bUWEdgeHasAdditionalVertices is true but Level > NbrUWLevel"));
 			const int32 StepSize = (1 << (NbrUWLevel - Level));
 			Grid[{0, J + 1}] = EdgeUWSortedVIDs[NumUWVertices - StepSize * (J + 1)];
 		}
@@ -408,6 +411,7 @@ void FRedGreenTessellationPattern::TessellateTriPatch(TrianglePatch& TriPatch) c
 
 		if (bVWEdgeHasAdditionalVertices)
 		{
+			checkf(NbrVWLevel >= Level, TEXT("bVWEdgeHasAdditionalVertices is true but Level > NbrVWLevel"));
 			const int32 StepSize = (1 << (NbrVWLevel - Level));
 			Grid[MapKey] = EdgeVWSortedVIDs[StepSize * (J + 1) - 1];
 		}
@@ -434,6 +438,7 @@ void FRedGreenTessellationPattern::TessellateTriPatch(TrianglePatch& TriPatch) c
 
 			auto FillEdgePoints = [&](const int32 NbrLevel, const int32 I, const TArray<int32>& InFullEdgePoints, TArray<int32>& EdgePoints, bool bGridCountReversed)
 			{
+				check(NbrLevel >= Level);	// This lambda should only be called when we are adding additional points due to the neighbor having a higher subdivision level
 				const int32 StepSize = (1 << (NbrLevel - Level));
 
 				int32 Start, End;
@@ -478,11 +483,6 @@ void FRedGreenTessellationPattern::TessellateTriPatch(TrianglePatch& TriPatch) c
 				FillEdgePoints(NbrUWLevel, J, EdgeUWSortedVIDs, EdgePoints[2], true);
 			}
 
-			if (EdgePoints[0].Num() == 1 && EdgePoints[1].Num() == 1 && EdgePoints[2].Num() == 1)
-			{
-				ensure(false);
-			}
-
 			// lower triangle(s) for the grid cell
 
 			if ((EdgePoints[0].Num() == 0) && (EdgePoints[1].Num() == 0) && (EdgePoints[2].Num() == 0))
