Added api exports to the OBJMeshUtil methods so they can be called from other modules

diff --git a/Engine/Plugins/Runtime/GeometryProcessing/Source/MeshFileUtils/Public/OBJMeshUtil.h b/Engine/Plugins/Runtime/GeometryProcessing/Source/MeshFileUtils/Public/OBJMeshUtil.h
--- a/Engine/Plugins/Runtime/GeometryProcessing/Source/MeshFileUtils/Public/OBJMeshUtil.h
+++ b/Engine/Plugins/Runtime/GeometryProcessing/Source/MeshFileUtils/Public/OBJMeshUtil.h
@@ -45,10 +45,10 @@ namespace UE::MeshFileUtils
 	};
 
 	// Attempt to load an OBJ file into a Dynamic Mesh
-	ELoadOBJStatus LoadOBJ(const char* Path, UE::Geometry::FDynamicMesh3& Mesh, const FLoadOBJSettings& Settings = FLoadOBJSettings());
+	MESHFILEUTILS_API ELoadOBJStatus LoadOBJ(const char* Path, UE::Geometry::FDynamicMesh3& Mesh, const FLoadOBJSettings& Settings = FLoadOBJSettings());
 	// Load an OBJ file into a Dynamic Mesh, or fail a check() if it cannot be loaded
-	UE::Geometry::FDynamicMesh3 LoadOBJChecked(const char* Path, const FLoadOBJSettings& Settings = FLoadOBJSettings());
+	MESHFILEUTILS_API UE::Geometry::FDynamicMesh3 LoadOBJChecked(const char* Path, const FLoadOBJSettings& Settings = FLoadOBJSettings());
 	// Write an OBJ file to the target path; returns true on success
-	bool WriteOBJ(const char* Path, const UE::Geometry::FDynamicMesh3& InMesh, const FWriteOBJSettings& Settings = FWriteOBJSettings());
+	MESHFILEUTILS_API bool WriteOBJ(const char* Path, const UE::Geometry::FDynamicMesh3& InMesh, const FWriteOBJSettings& Settings = FWriteOBJSettings());
 }
 
