Added optional scale factor to DynamicMeshUVEditor::ScaleUVAreaTo3DArea utility function

diff --git a/Engine/Plugins/Runtime/GeometryProcessing/Source/DynamicMesh/Private/Parameterization/DynamicMeshUVEditor.cpp b/Engine/Plugins/Runtime/GeometryProcessing/Source/DynamicMesh/Private/Parameterization/DynamicMeshUVEditor.cpp
--- a/Engine/Plugins/Runtime/GeometryProcessing/Source/DynamicMesh/Private/Parameterization/DynamicMeshUVEditor.cpp
+++ b/Engine/Plugins/Runtime/GeometryProcessing/Source/DynamicMesh/Private/Parameterization/DynamicMeshUVEditor.cpp
@@ -1436,7 +1436,7 @@ void FDynamicMeshUVEditor::SetTriangleUVsFromCylinderProjection(
 
 
 
-bool FDynamicMeshUVEditor::ScaleUVAreaTo3DArea(const TArray<int32>& Triangles, bool bRecenterAtOrigin)
+bool FDynamicMeshUVEditor::ScaleUVAreaTo3DArea(const TArray<int32>& Triangles, bool bRecenterAtOrigin, float ScaleFactor)
 {
 	double Area3D = TMeshQueries<FDynamicMesh3>::GetVolumeArea(*Mesh, Triangles).Y;
 	if (FMathd::Abs(Area3D) < FMathf::Epsilon || FMathd::IsFinite(Area3D) == false)
@@ -1462,7 +1462,7 @@ bool FDynamicMeshUVEditor::ScaleUVAreaTo3DArea(const TArray<int32>& Triangles, b
 		return false;
 	}
 
-	double UVScale = FMathd::Sqrt(Area3D) / FMathd::Sqrt(Area2D);
+	double UVScale = ScaleFactor * FMathd::Sqrt(Area3D) / FMathd::Sqrt(Area2D);
 	if (!FMathd::IsFinite(UVScale))
 	{
 		return false;
diff --git a/Engine/Plugins/Runtime/GeometryProcessing/Source/DynamicMesh/Public/Parameterization/DynamicMeshUVEditor.h b/Engine/Plugins/Runtime/GeometryProcessing/Source/DynamicMesh/Public/Parameterization/DynamicMeshUVEditor.h
--- a/Engine/Plugins/Runtime/GeometryProcessing/Source/DynamicMesh/Public/Parameterization/DynamicMeshUVEditor.h
+++ b/Engine/Plugins/Runtime/GeometryProcessing/Source/DynamicMesh/Public/Parameterization/DynamicMeshUVEditor.h
@@ -294,8 +294,9 @@ public:
 	* Compute the UV-space and 3D area of the given Triangles, and then scale the UV area to be equvalent
 	* to the 3D area. Scaling is around the center of the UV bounding box.
 	* @param bRecenterAtOrigin if true, UVs are translated after scaling such that the UV bounding box center is at (0,0)
+	* @param ScaleFactor Additional scale factor to multiply into resulting UVs
 	*/
-	UE_API bool ScaleUVAreaTo3DArea(const TArray<int32>& Triangles, bool bRecenterAtOrigin);
+	UE_API bool ScaleUVAreaTo3DArea(const TArray<int32>& Triangles, bool bRecenterAtOrigin, float ScaleFactor = 1.0f);
 
 
 	/**
