Change local fdynamicmesh array-based queries to use the inline-allocated array in a number of cases (especially where the arrays are created per-query)

diff --git a/Engine/Plugins/Editor/UVEditor/Source/UVEditorTools/Private/Actions/UVSeamSewAction.cpp b/Engine/Plugins/Editor/UVEditor/Source/UVEditorTools/Private/Actions/UVSeamSewAction.cpp
--- a/Engine/Plugins/Editor/UVEditor/Source/UVEditorTools/Private/Actions/UVSeamSewAction.cpp
+++ b/Engine/Plugins/Editor/UVEditor/Source/UVEditorTools/Private/Actions/UVSeamSewAction.cpp
@@ -126,7 +126,7 @@ bool UUVSeamSewAction::ExecuteAction()
 				FIndex2i Vids = MeshToSew.GetEdgeV(EidPair[PairIndex]);
 				for (int VidIndex = 0; VidIndex < 2; ++VidIndex)
 				{
-					TArray<int32> Tids;
+					FDynamicMesh3::FLocalIntArray Tids;
 					MeshToSew.GetVtxTriangles(Vids[VidIndex], Tids);
 					ChangedTidSet.Append(Tids);
 				}
diff --git a/Engine/Plugins/Editor/UVEditor/Source/UVEditorTools/Private/Actions/UVSplitAction.cpp b/Engine/Plugins/Editor/UVEditor/Source/UVEditorTools/Private/Actions/UVSplitAction.cpp
--- a/Engine/Plugins/Editor/UVEditor/Source/UVEditorTools/Private/Actions/UVSplitAction.cpp
+++ b/Engine/Plugins/Editor/UVEditor/Source/UVEditorTools/Private/Actions/UVSplitAction.cpp
@@ -58,7 +58,7 @@ namespace UVSplitActionLocals
 		TSet<int32> TidSet;
 		for (int32 UnwrapVid : UVEditResult.NewUVElements)
 		{
-			TArray<int32> VertTids;
+			FDynamicMesh3::FLocalIntArray VertTids;
 			Target->AppliedCanonical->GetVtxTriangles(Target->UnwrapVidToAppliedVid(UnwrapVid), VertTids);
 			TidSet.Append(VertTids);
 		}
@@ -132,7 +132,7 @@ namespace UVSplitActionLocals
 		TSet<int32> TidSet;
 		for (int32 UnwrapVid : NewUVElements)
 		{
-			TArray<int32> VertTids;
+			FDynamicMesh3::FLocalIntArray VertTids;
 			Target->AppliedCanonical->GetVtxTriangles(Target->UnwrapVidToAppliedVid(UnwrapVid), VertTids);
 			TidSet.Append(VertTids);
 		}
@@ -218,7 +218,7 @@ namespace UVSplitActionLocals
 		TSet<int32> TidSet;
 		for (int32 UnwrapVid : UVEditResult.NewUVElements)
 		{
-			TArray<int32> VertTids;
+			FDynamicMesh3::FLocalIntArray VertTids;
 			Target->AppliedCanonical->GetVtxTriangles(Target->UnwrapVidToAppliedVid(UnwrapVid), VertTids);
 			TidSet.Append(VertTids);
 		}
diff --git a/Engine/Plugins/Editor/UVEditor/Source/UVEditorTools/Private/UVEditorSeamTool.cpp b/Engine/Plugins/Editor/UVEditor/Source/UVEditorTools/Private/UVEditorSeamTool.cpp
--- a/Engine/Plugins/Editor/UVEditor/Source/UVEditorTools/Private/UVEditorSeamTool.cpp
+++ b/Engine/Plugins/Editor/UVEditor/Source/UVEditorTools/Private/UVEditorSeamTool.cpp
@@ -903,7 +903,7 @@ void UUVEditorSeamTool::ApplySeam(const TArray<int32>& AppliedVidsIn)
 		// Don't actually need to update VidSet because no element values changed or were created, only connectivity.
 		for (int32 Vid : AppliedVidsIn)
 		{
-			TArray<int32> VertTids;
+			FDynamicMesh3::FLocalIntArray VertTids;
 			ClickedTarget->AppliedCanonical->GetVtxTriangles(Vid, VertTids);
 			TidSet.Append(VertTids);
 		}
@@ -915,7 +915,7 @@ void UUVEditorSeamTool::ApplySeam(const TArray<int32>& AppliedVidsIn)
 
 		for (int32 UnwrapVid : UVEditResult.NewUVElements)
 		{
-			TArray<int32> VertTids;
+			FDynamicMesh3::FLocalIntArray VertTids;
 			ClickedTarget->AppliedCanonical->GetVtxTriangles(ClickedTarget->UnwrapVidToAppliedVid(UnwrapVid), VertTids);
 			TidSet.Append(VertTids);
 		}
diff --git a/Engine/Plugins/Editor/UVEditor/Source/UVEditorTools/Private/UVSelectTool.cpp b/Engine/Plugins/Editor/UVEditor/Source/UVEditorTools/Private/UVSelectTool.cpp
--- a/Engine/Plugins/Editor/UVEditor/Source/UVEditorTools/Private/UVSelectTool.cpp
+++ b/Engine/Plugins/Editor/UVEditor/Source/UVEditorTools/Private/UVSelectTool.cpp
@@ -455,7 +455,7 @@ void UUVSelectTool::ReinitializeFromSelection(bool bForceRecomputeCachedSelectio
 			VidSet.SetNumUninitialized(Selection.SelectedIDs.Num());
 			
 			int32 VIndex = 0;
-			TArray<int> Tids;
+			FDynamicMesh3::FLocalIntArray Tids;
 			for (const int32 Vid : Selection.SelectedIDs)
 			{
 				VidSet[VIndex++] = Vid;
diff --git a/Engine/Plugins/Runtime/GeometryProcessing/Source/DynamicMesh/Private/NormalFlowRemesher.cpp b/Engine/Plugins/Runtime/GeometryProcessing/Source/DynamicMesh/Private/NormalFlowRemesher.cpp
--- a/Engine/Plugins/Runtime/GeometryProcessing/Source/DynamicMesh/Private/NormalFlowRemesher.cpp
+++ b/Engine/Plugins/Runtime/GeometryProcessing/Source/DynamicMesh/Private/NormalFlowRemesher.cpp
@@ -295,7 +295,7 @@ void FNormalFlowRemesher::TrackedFaceProjectionPass(double& MaxDistanceMoved, bo
 		TempWeightBuffer[VertexIndex] = 0.0;
 		TempPosBuffer[VertexIndex] = { 0.0, 0.0, 0.0 };
 
-		TArray<int> IncidentTriangles;
+		FDynamicMesh3::FLocalIntArray IncidentTriangles;
 		Mesh->GetVtxTriangles(VertexIndex, IncidentTriangles);
 
 		for (int TriID : IncidentTriangles)
diff --git a/Engine/Plugins/Runtime/GeometryProcessing/Source/DynamicMesh/Private/Operations/ExtrudeBoundaryEdges.cpp b/Engine/Plugins/Runtime/GeometryProcessing/Source/DynamicMesh/Private/Operations/ExtrudeBoundaryEdges.cpp
--- a/Engine/Plugins/Runtime/GeometryProcessing/Source/DynamicMesh/Private/Operations/ExtrudeBoundaryEdges.cpp
+++ b/Engine/Plugins/Runtime/GeometryProcessing/Source/DynamicMesh/Private/Operations/ExtrudeBoundaryEdges.cpp
@@ -233,7 +233,7 @@ namespace ExtrudeBoundaryEdgesLocals
 	}
 
 	int32 FindSharpestAngleMatchEdge(const FDynamicMesh3& Mesh, int32 EidToMatch, int32 Vid, 
-		const FVector3d& VertexNormal, const TArray<int32> CandidateEids)
+		const FVector3d& VertexNormal, const FDynamicMesh3::FLocalIntArray& CandidateEids)
 	{
 		// The approach here is pretty much the same as in FMeshBoundaryLoops::FindLeftTurnEdge, except that
 		//  we're actually looking for the right turn edge (better if you're concerned with following a hole/boundary
@@ -579,8 +579,8 @@ bool FExtrudeBoundaryEdges::GetInputEdgePairings(
 	{
 		UnambiguousVidToEdgePairings.Remove(AmbiguousVid);
 
-		TArray<int32> IncomingEids;
-		TArray<int32> OutgoingEids;
+		FDynamicMesh3::FLocalIntArray IncomingEids;
+		FDynamicMesh3::FLocalIntArray OutgoingEids;
 		for (int32 Eid : Mesh.VtxEdgesItr(AmbiguousVid))
 		{
 			if (!Mesh.IsBoundaryEdge(Eid) || !SelectedEids.Contains(Eid))
@@ -634,7 +634,7 @@ bool FExtrudeBoundaryEdges::GetInputEdgePairings(
 			}
 
 			FVector3d VertexNormal = CalculateVertexNormal(VertData.SourceVid);
-			TArray<int32> NeighborhoodEids;
+			FDynamicMesh3::FLocalIntArray NeighborhoodEids;
 			Mesh.GetAllVtxBoundaryEdges(VertData.SourceVid, NeighborhoodEids);
 
 			if (VertData.SourceEidPair.A == IndexConstants::InvalidID 
diff --git a/Engine/Plugins/Runtime/GeometryProcessing/Source/DynamicMesh/Private/Operations/IntrinsicCorrespondenceUtils.cpp b/Engine/Plugins/Runtime/GeometryProcessing/Source/DynamicMesh/Private/Operations/IntrinsicCorrespondenceUtils.cpp
--- a/Engine/Plugins/Runtime/GeometryProcessing/Source/DynamicMesh/Private/Operations/IntrinsicCorrespondenceUtils.cpp
+++ b/Engine/Plugins/Runtime/GeometryProcessing/Source/DynamicMesh/Private/Operations/IntrinsicCorrespondenceUtils.cpp
@@ -487,9 +487,9 @@ IntrinsicCorrespondenceUtils::FSignpost::FSignpost(const FDynamicMesh3& SrcMesh)
 			constexpr int32 PlusTwoModThree[3] = { 2, 0, 1 }; // PlusTwoModThree[i] = (i + 2)%3
 
 			// initialize by computing the angles
-			TArray<int32> Triangles;
-			TArray<int32> ContiguousGroupLengths;
-			TArray<bool> GroupIsLoop;
+			FDynamicMesh3::FLocalIntArray Triangles;
+			FDynamicMesh3::FLocalIntArray ContiguousGroupLengths;
+			FDynamicMesh3::FLocalBoolArray GroupIsLoop;
 
 			{
 
diff --git a/Engine/Plugins/Runtime/GeometryProcessing/Source/DynamicMesh/Private/Parameterization/DynamicMeshUVEditor.cpp b/Engine/Plugins/Runtime/GeometryProcessing/Source/DynamicMesh/Private/Parameterization/DynamicMeshUVEditor.cpp
--- a/Engine/Plugins/Runtime/GeometryProcessing/Source/DynamicMesh/Private/Parameterization/DynamicMeshUVEditor.cpp
+++ b/Engine/Plugins/Runtime/GeometryProcessing/Source/DynamicMesh/Private/Parameterization/DynamicMeshUVEditor.cpp
@@ -1125,7 +1125,7 @@ bool UE::Geometry::FDynamicMeshUVEditor::MakeIsland(const TSet<int32>& TidsToMak
 		//  to make CreateSeamsAtEdges and RemoveSeamsAtEdges output changed tids instead.
 		for (int32 Vid : TouchedVids)
 		{
-			TArray<int32> Tids;
+			FDynamicMesh3::FLocalIntArray Tids;
 			Mesh->GetVtxTriangles(Vid, Tids);
 			ChangedTrianglesOut->Append(Tids);
 		}
diff --git a/Engine/Plugins/Runtime/MeshModelingToolset/Source/MeshModelingTools/Private/DisplaceMeshTool.cpp b/Engine/Plugins/Runtime/MeshModelingToolset/Source/MeshModelingTools/Private/DisplaceMeshTool.cpp
--- a/Engine/Plugins/Runtime/MeshModelingToolset/Source/MeshModelingTools/Private/DisplaceMeshTool.cpp
+++ b/Engine/Plugins/Runtime/MeshModelingToolset/Source/MeshModelingTools/Private/DisplaceMeshTool.cpp
@@ -252,7 +252,7 @@ namespace DisplaceMeshToolLocals{
 				if (UVOverlay->IsSeamVertex(VID)) 
 				{
 					// If the vertex is on a UV seam edge then average out offsets for all elements sharing the vertex
-					TArray<int> TrianglesOut;
+					FDynamicMesh3::FLocalIntArray TrianglesOut;
 					Mesh.GetVtxTriangles(VID, TrianglesOut);
 
 					for (const int TID : TrianglesOut) 
diff --git a/Engine/Source/Runtime/GeometryCore/Private/DynamicMesh/DynamicMesh3.cpp b/Engine/Source/Runtime/GeometryCore/Private/DynamicMesh/DynamicMesh3.cpp
--- a/Engine/Source/Runtime/GeometryCore/Private/DynamicMesh/DynamicMesh3.cpp
+++ b/Engine/Source/Runtime/GeometryCore/Private/DynamicMesh/DynamicMesh3.cpp
@@ -1376,7 +1376,7 @@ bool FDynamicMesh3::CheckValidity(FValidityOptions ValidityOptions, EValidityChe
 			CheckOrFailF(IsEdge(edge));
 		}
 
-		TArray<int> vTris;
+		FDynamicMesh3::FLocalIntArray vTris;
 		GetVtxTriangles(vID, vTris);
 		//System.Console.WriteLine(string.Format("{0} {1} {2}", vID, vTris.Count, GetVtxEdges(vID).Count));
 		if (ValidityOptions.bAllowNonManifoldVertices)
diff --git a/Engine/Source/Runtime/GeometryCore/Private/DynamicMesh/DynamicMesh3_Edits.cpp b/Engine/Source/Runtime/GeometryCore/Private/DynamicMesh/DynamicMesh3_Edits.cpp
--- a/Engine/Source/Runtime/GeometryCore/Private/DynamicMesh/DynamicMesh3_Edits.cpp
+++ b/Engine/Source/Runtime/GeometryCore/Private/DynamicMesh/DynamicMesh3_Edits.cpp
@@ -640,7 +640,7 @@ EMeshResult FDynamicMesh3::RemoveVertex(int vID, bool bPreserveManifold)
 	}
 
 	// Remove incident triangles
-	TArray<int> tris;
+	FDynamicMesh3::FLocalIntArray tris;
 	GetVtxTriangles(vID, tris);
 	for (int tID : tris)
 	{
@@ -2056,7 +2056,7 @@ EMeshResult UE::Geometry::FDynamicMesh3::MergeEdges(int eKeep, int eDiscard, dou
 			continue;
 		}
 
-		TArray<int> edges_v;
+		FDynamicMesh3::FLocalIntArray edges_v;
 		GetVertexEdgesList(v1, edges_v);
 		int Nedges = (int)edges_v.Num();
 		int FoundNum = 0;
diff --git a/Engine/Source/Runtime/GeometryCore/Private/Operations/EmbedSurfacePath.cpp b/Engine/Source/Runtime/GeometryCore/Private/Operations/EmbedSurfacePath.cpp
--- a/Engine/Source/Runtime/GeometryCore/Private/Operations/EmbedSurfacePath.cpp
+++ b/Engine/Source/Runtime/GeometryCore/Private/Operations/EmbedSurfacePath.cpp
@@ -912,7 +912,7 @@ bool UE::Geometry::EmbedProjectedPaths(FDynamicMesh3* Mesh, const TArrayView<con
 			{
 				return false;
 			}
-			TArray<int> TrianglesOut;
+			FDynamicMesh3::FLocalIntArray TrianglesOut;
 			Mesh->GetVtxTriangles(OutPathVertices.Last(), TrianglesOut);
 
 			check(TrianglesOut.Num());
