Improved performance of MeshRegionGraph operations:
- Region merging transfers known neighbor information, rather than recomputing it
- Less redundant validation on internal operations
- Less redundant array searching when building neighbors

diff --git a/Engine/Plugins/Runtime/GeometryProcessing/Source/DynamicMesh/Private/Parameterization/MeshRegionGraph.cpp b/Engine/Plugins/Runtime/GeometryProcessing/Source/DynamicMesh/Private/Parameterization/MeshRegionGraph.cpp
--- a/Engine/Plugins/Runtime/GeometryProcessing/Source/DynamicMesh/Private/Parameterization/MeshRegionGraph.cpp
+++ b/Engine/Plugins/Runtime/GeometryProcessing/Source/DynamicMesh/Private/Parameterization/MeshRegionGraph.cpp
@@ -196,11 +196,17 @@ bool FMeshRegionGraph::MergeRegion(int32 RemoveRgnIdx, int32 MergeToRgnIdx)
 		ensure(false);
 		return false;
 	}
+	return MergeRegionUnsafe(RemoveRgnIdx, MergeToRgnIdx);
+}
+
 
+bool FMeshRegionGraph::MergeRegionUnsafe(int32 RemoveRgnIdx, int32 MergeToRgnIdx)
+{
 	FRegion& RemoveRegion = Regions[RemoveRgnIdx];
 	FRegion& AppendToRegion = Regions[MergeToRgnIdx];
 
 	// append triangles to the new region
+	AppendToRegion.Triangles.Reserve(AppendToRegion.Triangles.Num() + RemoveRegion.Triangles.Num());
 	for (int32 tid : RemoveRegion.Triangles)
 	{
 		AppendToRegion.Triangles.Add(tid);
@@ -209,14 +215,105 @@ bool FMeshRegionGraph::MergeRegion(int32 RemoveRgnIdx, int32 MergeToRgnIdx)
 	RemoveRegion.Triangles.Reset();
 	RemoveRegion.bValid = false;
 
-	// rebuild all affected regions, which should just be all nbrs of the removed region
-	// (which will include the merge-to region)
-	// TODO: this can be more efficient. At minimum can update all nbr sets and then rebuild counts
+	// transfer boundary flags
+	if (RemoveRegion.bIsOnMeshBoundary)
+	{
+		AppendToRegion.bIsOnMeshBoundary = true;
+	}
+	if (RemoveRegion.bIsOnROIBoundary)
+	{
+		AppendToRegion.bIsOnROIBoundary = true;
+	}
+	// remove RemoveRgnIdx from the AppendToRegion's neighbors list
+	AppendToRegion.Neighbours.RemoveAll([RemoveRgnIdx](const FNeighbour& Nbr)
+		{
+			return Nbr.RegionIndex == RemoveRgnIdx;
+		});
+
+	// Mapping from RegionIndex -> Index of Region in AppendToRegion.Neighbors
+	TMap<int32, int32> AppendToRegionNbrTransferMap;
+	// Populate transfer map with relevant region keys -- the removed region neighbors, which will need updating
 	for (FNeighbour& Nbr : RemoveRegion.Neighbours)
 	{
-		BuildNeigbours(Nbr.RegionIndex);
+		if (Nbr.RegionIndex != MergeToRgnIdx)
+		{
+			AppendToRegionNbrTransferMap.Add(Nbr.RegionIndex, INDEX_NONE);
+		}
+	}
+	// Linear pass to get indices for neighbors that are already in the merge-to-region's neighbor array
+	for (int32 AppendNbrIdx = 0; AppendNbrIdx < AppendToRegion.Neighbours.Num(); ++AppendNbrIdx)
+	{
+		int32 RegionIdx = AppendToRegion.Neighbours[AppendNbrIdx].RegionIndex;
+		if (int32* FoundTransferIdx = AppendToRegionNbrTransferMap.Find(RegionIdx))
+		{
+			*FoundTransferIdx = AppendNbrIdx;
+		}
 	}
 
+	// update neighbor info for all removed neighbors -- passing counts to the merged-to region
+	for (FNeighbour& Nbr : RemoveRegion.Neighbours)
+	{
+		// if this is the merged-to region, skip it (we already removed the back-link for it above)
+		if (Nbr.RegionIndex == MergeToRgnIdx)
+		{
+			continue;
+		}
+
+		FRegion& NbrRegion = Regions[Nbr.RegionIndex];
+		TArray<FNeighbour>& NbrNbrs = NbrRegion.Neighbours;
+
+		// Find the removed and kept regions in the neighbor list
+		int32 NbrsRemoveIdx = INDEX_NONE;
+		int32 NbrsKeepIdx = INDEX_NONE;
+		for (int32 Idx = 0; Idx < NbrNbrs.Num(); ++Idx)
+		{
+			if (NbrNbrs[Idx].RegionIndex == RemoveRgnIdx)
+			{
+				NbrsRemoveIdx = Idx;
+			}
+			else if (NbrNbrs[Idx].RegionIndex == MergeToRgnIdx)
+			{
+				NbrsKeepIdx = Idx;
+			}
+		}
+		if (!ensure(NbrsRemoveIdx != INDEX_NONE)) // we should always find the removed idx, by symmetry
+		{
+			continue;
+		}
+
+		// make sure the neighbor is connected to the merged-to region & transfer counts
+		const FNeighbour& Removed = NbrNbrs[NbrsRemoveIdx];
+		// Find the relevant entry in the merged-to region's neighbors array, to symmetrically update
+		int32 NbrInMergeRegionIdx = AppendToRegionNbrTransferMap[Nbr.RegionIndex];
+		// ... or create it if needed
+		if (NbrInMergeRegionIdx == INDEX_NONE)
+		{
+			NbrInMergeRegionIdx = AppendToRegion.Neighbours.Add(
+				FNeighbour
+				{
+					.RegionIndex = Nbr.RegionIndex,
+					.Count = Removed.Count
+				}
+			);
+		}
+		else
+		{
+			AppendToRegion.Neighbours[NbrInMergeRegionIdx].Count += Removed.Count;
+		}
+		if (NbrsKeepIdx == INDEX_NONE)
+		{
+			// it wasn't previously connected to the merged-to region; can directly transfer the neighbor info
+			NbrNbrs[NbrsRemoveIdx].RegionIndex = MergeToRgnIdx;
+		}
+		else
+		{
+			// it was previously connected; can add increase the connection count by that of the removed region
+			NbrNbrs[NbrsKeepIdx].Count += Removed.Count;
+			NbrNbrs.RemoveAt(NbrsRemoveIdx);
+		}
+	}
+	RemoveRegion.Neighbours.Empty();
+
 	return true;
 }
 
@@ -295,24 +392,10 @@ bool FMeshRegionGraph::OptimizeBorders(int32 MaxRounds)
 
 void FMeshRegionGraph::BuildNeigbours(int32 RegionIdx)
 {
-	auto CountNeighbourFunc = [](FRegion& Region, int32 NbrRegionIndex)
-	{
-		for (FNeighbour& Nbr : Region.Neighbours)
-		{
-			if (Nbr.RegionIndex == NbrRegionIndex)
-			{
-				Nbr.Count++;
-				return;
-			}
-		}
-		check(false);
-	};
-
-
 	FRegion& Region = Regions[RegionIdx];
 	Region.Neighbours.Reset();
 
-	TSet<int32> NeighbourSet;
+	TMap<int32, int32> NeighborMap;
 
 	for (int32 tid : Region.Triangles)
 	{
@@ -337,15 +420,21 @@ void FMeshRegionGraph::BuildNeigbours(int32 RegionIdx)
 				{
 					if (NbrRegionIndex != RegionIndex)
 					{
-						if (NeighbourSet.Contains(NbrRegionIndex) == false)
+						int32* FoundNbrIdx = NeighborMap.Find(NbrRegionIndex);
+						if (!FoundNbrIdx)
 						{
-							NeighbourSet.Add(NbrRegionIndex);
-							FNeighbour Nbr;
-							Nbr.RegionIndex = NbrRegionIndex;
-							Region.Neighbours.Add(Nbr);
+							FNeighbour Nbr
+							{
+								.RegionIndex = NbrRegionIndex,
+								.Count = 1
+							};
+							int32 NbrIdx = Region.Neighbours.Add(Nbr);
+							NeighborMap.Add(NbrRegionIndex, NbrIdx);
+						}
+						else
+						{
+							Region.Neighbours[*FoundNbrIdx].Count++;
 						}
-
-						CountNeighbourFunc(Region, NbrRegionIndex);
 					}
 				}
 			}
diff --git a/Engine/Plugins/Runtime/GeometryProcessing/Source/DynamicMesh/Public/Parameterization/MeshRegionGraph.h b/Engine/Plugins/Runtime/GeometryProcessing/Source/DynamicMesh/Public/Parameterization/MeshRegionGraph.h
--- a/Engine/Plugins/Runtime/GeometryProcessing/Source/DynamicMesh/Public/Parameterization/MeshRegionGraph.h
+++ b/Engine/Plugins/Runtime/GeometryProcessing/Source/DynamicMesh/Public/Parameterization/MeshRegionGraph.h
@@ -153,6 +153,10 @@ protected:
 	UE_API bool MergeSmallRegionsPass(int32 SmallThreshold,
 						   TFunctionRef<float(int32 SmallRgnIdx, int32 NbrRgnIdx)> RegionSimilarityFunc);
 
+private:
+	// MergeRegion without validation, for internal use where we know the inputs are valid neighbors
+	bool MergeRegionUnsafe(int32 RemoveRgnIdx, int32 MergeToRgnIdx);
+
 };
 
 
