Added option to provide local per-vertex scaling to the cluster simplifier's target edge length, and also to automatically scale the edge length down as needed to avoid fully collapsing mesh islands

diff --git a/Engine/Plugins/Runtime/GeometryProcessing/Source/DynamicMesh/Private/Operations/MeshClusterSimplifier.cpp b/Engine/Plugins/Runtime/GeometryProcessing/Source/DynamicMesh/Private/Operations/MeshClusterSimplifier.cpp
--- a/Engine/Plugins/Runtime/GeometryProcessing/Source/DynamicMesh/Private/Operations/MeshClusterSimplifier.cpp
+++ b/Engine/Plugins/Runtime/GeometryProcessing/Source/DynamicMesh/Private/Operations/MeshClusterSimplifier.cpp
@@ -262,6 +262,20 @@ bool Simplify(const FDynamicMesh3& InMesh, FDynamicMesh3& ResultMesh, const FSim
 			VertexTags[VID] = EElemTag::Free;
 		});
 
+	TArray<float> TargetLengths;
+	if (SimplifyOptions.OptionalTargetEdgeLengthScale)
+	{
+		TargetLengths.SetNumUninitialized(InMesh.MaxVertexID());
+		ParallelFor(InMesh.MaxVertexID(), [&InMesh, &TargetLengths, &SimplifyOptions](int32 VID)
+			{
+				if (!InMesh.IsVertex(VID))
+				{
+					return;
+				}
+				TargetLengths[VID] = float(SimplifyOptions.TargetEdgeLength * SimplifyOptions.OptionalTargetEdgeLengthScale(VID));
+			}
+		);
+	}
 
 	///
 	/// Step 2. Clustering: Grow vertex clusters out to the target edge length size
@@ -280,7 +294,7 @@ bool Simplify(const FDynamicMesh3& InMesh, FDynamicMesh3& ResultMesh, const FSim
 	SourceDist.Init(FMathf::MaxReal, InMesh.MaxVertexID());
 
 	auto TagVerticesByRegionGrowth = 
-		[&Source, &SourceDist, &InMesh, &SimplifyOptions, &EdgeTags, &VertexTags]
+		[&Source, &SourceDist, &InMesh, &SimplifyOptions, &EdgeTags, &VertexTags, &TargetLengths]
 		(const TStaticArray<TArray<int32>, NUM_TAGS>& VertexIDBuckets)
 	{
 		// add all the fixed vertices as sources first, so they can't be claimed by other verts
@@ -326,11 +340,13 @@ bool Simplify(const FDynamicMesh3& InMesh, FDynamicMesh3& ResultMesh, const FSim
 						continue;
 					}
 
+					float MaxDist = TargetLengths.IsEmpty() ? (float)SimplifyOptions.TargetEdgeLength : TargetLengths[CurSourceVID];
+
 					// vertex is either a new source, or previously claimed but we need to consider growing via less-constrained edges
 
 					// helper to add candidate verts to a heap
 					HeapV.Reset();
-					auto AddCandidates = [MaxDist = SimplifyOptions.TargetEdgeLength,
+					auto AddCandidates = [MaxDist,
 						&HeapV, &InMesh, &SourceDist, &Source, &EdgeTags, &VertexTags, TagIdx]
 						(const FWalk& From)
 						{
@@ -399,6 +415,85 @@ bool Simplify(const FDynamicMesh3& InMesh, FDynamicMesh3& ResultMesh, const FSim
 		ProcessBuckets[Idx].Empty();
 	}
 
+	
+	/// Step 2.5: Optionally look for clusters that didn't make it into triangles
+	for (int32 PassIdx = 0; PassIdx < SimplifyOptions.MaxPreserveCollapsedClusterPasses; ++PassIdx)
+	{
+		// factor used to set new Target Length for collapsed clusters, as a factor of max distance
+		constexpr float TargetLengthAsMaxClusterLengthFactor = .5f;
+
+		// Look through the triangles to tag the clusters that currently will end up in the output mesh
+		// due to having a valid triangle associated
+		TArray<bool> ClusterInOutput;
+		ClusterInOutput.SetNumZeroed(Source.Num());
+		ParallelFor(InMesh.MaxTriangleID(), [&InMesh, &Source, &ClusterInOutput](int32 TID)
+			{
+				if (!InMesh.IsTriangle(TID))
+				{
+					return;
+				}
+				FIndex3i Tri = InMesh.GetTriangle(TID);
+				FIndex3i SourceTri(Source[Tri.A], Source[Tri.B], Source[Tri.C]);
+				if (SourceTri.A != SourceTri.B && SourceTri.A != SourceTri.C && SourceTri.B != SourceTri.C)
+				{
+					for (int32 SubIdx = 0; SubIdx < 3; ++SubIdx)
+					{
+						ClusterInOutput[Source[Tri[SubIdx]]] = true;
+					}
+				}
+			}
+		);
+		TArray<int32> NeedRecluster;
+		TMap<int32, int32> VIDtoNeedReclusterIdx;
+		for (int32 VID : InMesh.VertexIndicesItr())
+		{
+			if (VID == Source[VID] && !ClusterInOutput[VID])
+			{
+				VIDtoNeedReclusterIdx.Add(VID, NeedRecluster.Add(VID));
+			}
+		}
+		if (NeedRecluster.IsEmpty())
+		{
+			break;
+		}
+
+		// Create buckets for vertices that will need re-processing due to not having any triangles
+		TStaticArray<TArray<int32>, NUM_TAGS> ReProcessBuckets;
+		// Make sure target lengths exist, so we can reduce them as needed for re-clustering
+		if (TargetLengths.IsEmpty())
+		{
+			TargetLengths.Init((float)SimplifyOptions.TargetEdgeLength, InMesh.MaxVertexID());
+		}
+
+		TArray<float> MaxSourceDistInCluster;
+		MaxSourceDistInCluster.SetNumZeroed(NeedRecluster.Num());
+		for (int32 VID : InMesh.VertexIndicesItr())
+		{
+			int32 SourceVID = Source[VID];
+			if (!ClusterInOutput[SourceVID]) // vertex belongs to a cluster that we need to reprocess
+			{
+				// update the max distance for the cluster
+				float& MaxFoundClusterDist = MaxSourceDistInCluster[VIDtoNeedReclusterIdx[SourceVID]];
+				MaxFoundClusterDist = FMath::Max(MaxFoundClusterDist, SourceDist[VID]);
+				ReProcessBuckets[(int32)VertexTags[VID]].Add(VID);
+			}
+		}
+		// For all re-process vertices, set down-scaled target lengths, and clear source/sourcedist
+		for (int32 BucketIdx = 0; BucketIdx < ReProcessBuckets.Num(); ++BucketIdx)
+		{
+			for (int32 Idx = 0; Idx < ReProcessBuckets[BucketIdx].Num(); ++Idx)
+			{
+				int32 UpdateVID = ReProcessBuckets[BucketIdx][Idx];
+				int32 SourceVID = Source[UpdateVID];
+				float ClusterMaxDist = MaxSourceDistInCluster[VIDtoNeedReclusterIdx[SourceVID]];
+				TargetLengths[UpdateVID] = ClusterMaxDist * TargetLengthAsMaxClusterLengthFactor;
+				Source[UpdateVID] = INDEX_NONE;
+				SourceDist[UpdateVID] = FMathf::MaxReal;
+			}
+		}
+		// Re-run tagging
+		TagVerticesByRegionGrowth(ReProcessBuckets);
+	}
 
 	///
 	/// Step 3: Copy the cluster connectivity out to our ResultMesh
diff --git a/Engine/Plugins/Runtime/GeometryProcessing/Source/DynamicMesh/Public/Operations/MeshClusterSimplifier.h b/Engine/Plugins/Runtime/GeometryProcessing/Source/DynamicMesh/Public/Operations/MeshClusterSimplifier.h
--- a/Engine/Plugins/Runtime/GeometryProcessing/Source/DynamicMesh/Public/Operations/MeshClusterSimplifier.h
+++ b/Engine/Plugins/Runtime/GeometryProcessing/Source/DynamicMesh/Public/Operations/MeshClusterSimplifier.h
@@ -21,6 +21,18 @@ struct FSimplifyOptions
 	// Vertices within this distance are allowed to be replaced with a single vertex
 	double TargetEdgeLength = 1;
 
+	// Optionally locally scales the simplifier's TargetEdgeLength on a per-vertex basis
+	// If provided, the target edge length at a vertex VID will be TargetEdgeLength * OptionalTargetEdgeLengthScale(VID)
+	TFunction<double(int32)> OptionalTargetEdgeLengthScale = nullptr;
+
+	// If > 0, number of passes to attempt to make sure the simplification does not fully delete islands of the mesh.
+	// In areas that would be deleted (because the region is smaller than the TargetEdgeLength),
+	// the target edge length will instead be locally reduced to half the local max in-cluster distance,
+	// and clustering will be re-run.
+	// This does not guarantee all clusters will be preserved: Use a higher number of passes to reducing the 
+	// likelihood of losing mesh components/connectivity, for additional cost & more vertices in the result.
+	int32 MaxPreserveCollapsedClusterPasses = 3;
+
 	// If > 0, boundary vertices w/ incident boundary edge angle greater than this (in degrees) will be kept in the output
 	double FixBoundaryAngleTolerance = 45;
 
