Move a function that gradually increases refinement level into the GeometryProcessing library for reuse

diff --git a/Engine/Plugins/Runtime/GeometryProcessing/Source/DynamicMesh/Private/Selections/GeometrySelectionUtil.cpp b/Engine/Plugins/Runtime/GeometryProcessing/Source/DynamicMesh/Private/Selections/GeometrySelectionUtil.cpp
--- a/Engine/Plugins/Runtime/GeometryProcessing/Source/DynamicMesh/Private/Selections/GeometrySelectionUtil.cpp
+++ b/Engine/Plugins/Runtime/GeometryProcessing/Source/DynamicMesh/Private/Selections/GeometrySelectionUtil.cpp
@@ -3073,4 +3073,85 @@ bool UE::Geometry::GetTriangleSelectionFrame(
 	}
 
 	return true;
-}
\ No newline at end of file
+}
+
+void UE::Geometry::TaperPerTriangleValues(const UE::Geometry::FDynamicMesh3& Mesh,
+	const TSet<int32>& TriangleROI,
+	TFunction<int32(int32)> ValueForTriangleFunc,
+	TArray<int>& OutTriangleValues)
+{
+	using namespace UE::Geometry;
+
+	OutTriangleValues.SetNum(Mesh.MaxTriangleID());
+	for (const int32 TriangleIndex : Mesh.TriangleIndicesItr())
+	{
+		if (Mesh.IsTriangle(TriangleIndex))
+		{
+			OutTriangleValues[TriangleIndex] = ValueForTriangleFunc(TriangleIndex);
+		}
+	}
+
+	// If a triangle is on the boundary of a mesh or on the boundary of the selection region, set its level to zero
+	for (int32 TriangleID = 0; TriangleID < Mesh.MaxTriangleID(); ++TriangleID)
+	{
+		if (Mesh.IsTriangle(TriangleID))
+		{
+			bool bTriangleShouldHaveLevel = true;
+
+			const FIndex3i Tri = Mesh.GetTriangle(TriangleID);
+			for (int32 TriVertIndex = 0; TriVertIndex < 3; ++TriVertIndex)
+			{
+				const int32 VertexID = Tri[TriVertIndex];
+				if (Mesh.IsBoundaryVertex(VertexID))
+				{
+					bTriangleShouldHaveLevel = false;
+					break;
+				}
+
+				for (const int32 IncidentEdgeID : Mesh.VtxEdgesItr(VertexID))
+				{
+					const FIndex2i EdgeTriangles = Mesh.GetEdgeT(IncidentEdgeID);
+					if (!TriangleROI.Contains(EdgeTriangles[0]) || !TriangleROI.Contains(EdgeTriangles[1]))
+					{
+						bTriangleShouldHaveLevel = false;
+						break;
+					}
+				}
+			}
+
+			if (!bTriangleShouldHaveLevel)
+			{
+				OutTriangleValues[TriangleID] = 0;
+			}
+		}
+	}
+
+	// TODO: This could be slow in the worst case. Maybe swap out for a BFS
+	bool bChanged = true;
+	while (bChanged)
+	{
+		bChanged = false;
+
+		for (int32 TriID : Mesh.TriangleIndicesItr())
+		{
+			int32 LowestNeighborLevel = ValueForTriangleFunc(TriID);
+			const FIndex3i Tri = Mesh.GetTriangle(TriID);
+			for (int32 TriVertIndex = 0; TriVertIndex < 3; ++TriVertIndex)
+			{
+				for (const int32 NeighborTri : Mesh.VtxTrianglesItr(Tri[TriVertIndex]))
+				{
+					if (Mesh.IsTriangle(NeighborTri) && TriangleROI.Contains(NeighborTri))
+					{
+						LowestNeighborLevel = FMath::Min(LowestNeighborLevel, OutTriangleValues[NeighborTri]);
+					}
+				}
+			}
+
+			if (OutTriangleValues[TriID] > LowestNeighborLevel + 1 && LowestNeighborLevel + 1 <= ValueForTriangleFunc(TriID))
+			{
+				OutTriangleValues[TriID] = LowestNeighborLevel + 1;
+				bChanged = true;
+			}
+		}
+	}
+}
diff --git a/Engine/Plugins/Runtime/GeometryProcessing/Source/DynamicMesh/Public/Selections/GeometrySelectionUtil.h b/Engine/Plugins/Runtime/GeometryProcessing/Source/DynamicMesh/Public/Selections/GeometrySelectionUtil.h
--- a/Engine/Plugins/Runtime/GeometryProcessing/Source/DynamicMesh/Public/Selections/GeometrySelectionUtil.h
+++ b/Engine/Plugins/Runtime/GeometryProcessing/Source/DynamicMesh/Public/Selections/GeometrySelectionUtil.h
@@ -572,5 +572,23 @@ DYNAMICMESH_API bool GetTriangleSelectionFrame(
 	FFrame3d& SelectionFrameOut);
 
 
+/**
+ * Computes and smooths per-triangle integer values across a specified region of a mesh.
+ *
+ * Each triangle in the mesh is assigned an initial value via the user-provided `ValueForTriangleFunc`.
+ * Triangles that are on the boundary of the mesh or on the boundary of the specified `InsideTriangles` region are assigned a value of zero.
+ * Triangle values are then adjusted within the region to ensure that each triangle's value does not exceed one more than the minimum value of its neighboring triangles.
+ * This produces a smoothly varying integer field that respects the boundary constraints.
+ * Note the field varies with topological distance (triangle connectivity), not distance along the surface or Euclidean distance
+ *
+ * @param TriangleROI          Triangles to consider
+ * @param LevelForTriangleFunc Function that takes a triangle ID and returns the maximum value for that triangle
+ * @param OutTriangleLevels    Output array of triangle values
+ */
+DYNAMICMESH_API void TaperPerTriangleValues(const UE::Geometry::FDynamicMesh3& Mesh,
+	const TSet<int32>& TriangleROI,
+	TFunction<int32(int32)> ValueForTriangleFunc,
+	TArray<int>& OutTriangleValues);
+
+}
 }
-}
\ No newline at end of file
