generalize cluster mesh simplifier to work on any triangle mesh via mesh adapters

diff --git a/Engine/Plugins/Runtime/GeometryProcessing/Source/DynamicMesh/Private/Operations/MeshClusterSimplifier.cpp b/Engine/Plugins/Runtime/GeometryProcessing/Source/DynamicMesh/Private/Operations/MeshClusterSimplifier.cpp
--- a/Engine/Plugins/Runtime/GeometryProcessing/Source/DynamicMesh/Private/Operations/MeshClusterSimplifier.cpp
+++ b/Engine/Plugins/Runtime/GeometryProcessing/Source/DynamicMesh/Private/Operations/MeshClusterSimplifier.cpp
@@ -6,11 +6,29 @@
 #include "Async/ParallelFor.h"
 #include "DynamicMesh/DynamicMesh3.h"
 #include "DynamicMesh/DynamicMeshAttributeSet.h"
+#include "IndexTypes.h"
 #include "VectorTypes.h"
 
+#include <type_traits>
 
 namespace UE::MeshClusterSimplifyLocals
 {
+	using namespace UE::Geometry::MeshClusterSimplify;
+	using namespace UE::Geometry;
+
+	using EElemTag = FSimplifyOptions::EConstraintLevel;
+	using FVertexBuckets = TStaticArray<TArray<int32>, (uint32)EElemTag::MAX>;
+
+	// Constraint levels are ordered from Most to Least constrained, so we can compare their integer values to determine constraint priority
+	// @return true if constraint level of A is more constrained (lower value) vs the constraint level of B
+	template<typename TagType>
+	inline bool IsMoreConstrained(EElemTag A, TagType B)
+	{
+		// TagType should either be EElemTag or uint8
+		static_assert(std::is_same_v<TagType, EElemTag> || std::is_same_v<TagType, uint8>);
+		return (uint8)A < (uint8)B;
+	}
+
 	template<typename ElemType, int Dim, typename AttributeType>
 	void CopyAttribs(AttributeType* Result, const AttributeType* Source, TConstArrayView<int32> ResToSource, int32 Num)
 	{
@@ -23,279 +41,145 @@ namespace UE::MeshClusterSimplifyLocals
 			}
 		);
 	}
-}
-
-namespace UE::Geometry::MeshClusterSimplify
-{
-
-bool Simplify(const FDynamicMesh3& InMesh, FDynamicMesh3& ResultMesh, const FSimplifyOptions& SimplifyOptions)
-{
-	TRACE_CPUPROFILER_EVENT_SCOPE(MeshClusterSimplify::Simplify);
-
-	// we build the result mesh by incrementally copying from the input mesh, so they shouldn't be the same mesh
-	if (!ensure(&ResultMesh != &InMesh))
-	{
-		return false;
-	}
-
-	ResultMesh.Clear();
-	
-	const FDynamicMeshAttributeSet* InAttribs = InMesh.Attributes();
-
-	// We tag edges and vertices w/ the constraint level, abbreviated to EElemTag for convenience
-	using EElemTag = FSimplifyOptions::EConstraintLevel;
-	constexpr int32 NUM_TAGS = 3;
-
-	// TODO: optionally also compute some vertex curvature feature & sort by it, to favor capturing less flat parts of the input shape?
-
-	///
-	/// Step 1, Data Prep: Translate all mesh constraint options to simple per-edge and per-vertex tags, so we know what to try to especially preserve in the result
-	///
 	
-	// Compute an Edge ID -> Constraint Level mapping
-	TArray<EElemTag> EdgeTags;
-	EdgeTags.SetNumUninitialized(InMesh.MaxEdgeID());
-	ParallelFor(InMesh.MaxEdgeID(), [&InMesh, &InAttribs, &SimplifyOptions, &EdgeTags](int32 EID)
+	template<typename MeshType, typename MeshConnType>
+	static void TagVerticesAndEdges(const MeshType& InMesh, const MeshConnType& InMeshConnectivity,
+		TFunctionRef<EElemTag(int32 EID, EElemTag InitialTag)> GetEdgeConstraintLevel,
+		TFunctionRef<bool(int32 VID)> IsSeamIntersectionVertex,
+		const FSimplifyOptions& SimplifyOptions,
+		TArray<EElemTag>& OutEdgeTags, TArray<EElemTag>& OutVertexTags)
 	{
-		if (!InMesh.IsEdge(EID))
-		{
-			return;
-		}
-
-		EElemTag UseTag = EElemTag::Free;
-		if ((uint8)SimplifyOptions.PreserveEdges.Boundary < (uint8)UseTag && InMesh.IsBoundaryEdge(EID))
-		{
-			UseTag = SimplifyOptions.PreserveEdges.Boundary;
-		}
-		if ((uint8)SimplifyOptions.PreserveEdges.PolyGroup < (uint8)UseTag && InMesh.IsGroupBoundaryEdge(EID))
-		{
-			UseTag = SimplifyOptions.PreserveEdges.PolyGroup;
-		}
-
-		if (InAttribs)
-		{
-			if ((uint8)SimplifyOptions.PreserveEdges.Material < (uint8)UseTag && InAttribs->IsMaterialBoundaryEdge(EID))
-			{
-				UseTag = SimplifyOptions.PreserveEdges.Material;
-			}
-
-			if ((uint8)SimplifyOptions.PreserveEdges.UVSeam < (uint8)UseTag)
+		// Compute an Edge ID -> Constraint Level mapping
+		OutEdgeTags.SetNumUninitialized(InMeshConnectivity.MaxEdgeID());
+		ParallelFor(InMeshConnectivity.MaxEdgeID(), [&InMesh, &InMeshConnectivity, &GetEdgeConstraintLevel, &SimplifyOptions, &OutEdgeTags](int32 EID)
 			{
-				for (int32 UVLayer = 0; UVLayer < InAttribs->NumUVLayers(); ++UVLayer)
+				if (!InMeshConnectivity.IsEdge(EID))
 				{
-					if (InAttribs->GetUVLayer(UVLayer)->IsSeamEdge(EID))
-					{
-						UseTag = SimplifyOptions.PreserveEdges.UVSeam;
-						break;
-					}
-				}
-			}
-			if ((uint8)SimplifyOptions.PreserveEdges.TangentSeam < (uint8)UseTag)
-			{
-				for (int32 NormalLayer = 1; NormalLayer < InAttribs->NumNormalLayers(); ++NormalLayer)
-				{
-					if (InAttribs->GetNormalLayer(NormalLayer)->IsSeamEdge(EID))
-					{
-						UseTag = SimplifyOptions.PreserveEdges.TangentSeam;
-						break;
-					}
-				}
-			}
-			if ((uint8)SimplifyOptions.PreserveEdges.NormalSeam < (uint8)UseTag)
-			{
-				if (const FDynamicMeshNormalOverlay* Normals = InAttribs->PrimaryNormals())
-				{
-					if (Normals->IsSeamEdge(EID))
-					{
-						UseTag = SimplifyOptions.PreserveEdges.NormalSeam;
-					}
+					return;
 				}
-			}
-			if ((uint8)SimplifyOptions.PreserveEdges.ColorSeam < (uint8)UseTag)
-			{
-				if (const FDynamicMeshColorOverlay* Colors = InAttribs->PrimaryColors())
+
+				EElemTag UseTag = EElemTag::Free;
+				if (IsMoreConstrained(SimplifyOptions.PreserveEdges.Boundary, UseTag) && InMeshConnectivity.IsBoundaryEdge(EID))
 				{
-					if (Colors->IsSeamEdge(EID))
-					{
-						UseTag = SimplifyOptions.PreserveEdges.ColorSeam;
-					}
+					UseTag = SimplifyOptions.PreserveEdges.Boundary;
 				}
-			}
-		}
+				UseTag = GetEdgeConstraintLevel(EID, UseTag);
 
-		EdgeTags[EID] = UseTag;
-	});
+				OutEdgeTags[EID] = UseTag;
+			});
 
-	TArray<EElemTag> VertexTags; 
-	VertexTags.SetNumUninitialized(InMesh.MaxVertexID());
+		OutVertexTags.SetNumUninitialized(InMesh.MaxVertexID());
 
-	double CosBoundaryEdgeAngleTolerance = FMath::Cos(FMath::DegreesToRadians(FMath::Clamp(SimplifyOptions.FixBoundaryAngleTolerance, 0, 180)));
+		double CosBoundaryEdgeAngleTolerance = FMath::Cos(FMath::DegreesToRadians(FMath::Clamp(SimplifyOptions.FixBoundaryAngleTolerance, 0, 180)));
 
-	auto IsSeamIntersectionVertex = 
-		[&InAttribs, &SimplifyOptions]
-		(int32 VID)
-	{
-		if (SimplifyOptions.PreserveEdges.UVSeam == EElemTag::Constrained)
-		{
-			for (int32 Layer = 0; Layer < InAttribs->NumUVLayers(); ++Layer)
+		ParallelFor(InMesh.MaxVertexID(),
+			[&InMesh, &InMeshConnectivity, &SimplifyOptions,
+			&OutVertexTags, &OutEdgeTags,
+			&IsSeamIntersectionVertex,
+			CosBoundaryEdgeAngleTolerance]
+			(int32 VID)
 			{
-				if (InAttribs->GetUVLayer(Layer)->IsSeamIntersectionVertex(VID))
-				{
-					return true;
-				}
-			}
-		}
-		if (SimplifyOptions.PreserveEdges.NormalSeam == EElemTag::Constrained)
-		{
-			if (const FDynamicMeshNormalOverlay* Normals = InAttribs->PrimaryNormals())
-			{
-				if (Normals->IsSeamIntersectionVertex(VID))
-				{
-					return true;
-				}
-			}
-		}
-		if (SimplifyOptions.PreserveEdges.TangentSeam == EElemTag::Constrained)
-		{
-			for (int32 Layer = 1; Layer < InAttribs->NumNormalLayers(); ++Layer)
-			{
-				if (InAttribs->GetNormalLayer(Layer)->IsSeamIntersectionVertex(VID))
-				{
-					return true;
-				}
-			}
-		}
-		if (SimplifyOptions.PreserveEdges.ColorSeam == EElemTag::Constrained)
-		{
-			if (const FDynamicMeshColorOverlay* Colors = InAttribs->PrimaryColors())
-			{
-				if (Colors->IsSeamIntersectionVertex(VID))
+				if (!InMesh.IsVertex(VID))
 				{
-					return true;
+					return;
 				}
-			}
-		}
-
-		return false;
-	};
-
-	ParallelFor(InMesh.MaxVertexID(), 
-		[&InMesh, &InAttribs, &SimplifyOptions, 
-		&VertexTags, &EdgeTags, 
-		&IsSeamIntersectionVertex,
-		CosBoundaryEdgeAngleTolerance]
-		(int32 VID)
-		{
-			if (!InMesh.IsVertex(VID))
-			{
-				return;
-			}
-			int32 FixedCount = 0;
-			int32 ConstrainedCount = 0;
-			
-			FVector3d BoundaryEdgeVert[2];
-			int32 FoundBoundaryEdgeVerts = 0;
-			InMesh.EnumerateVertexEdges(VID, 
-				[&EdgeTags, VID, &FixedCount, &ConstrainedCount, &BoundaryEdgeVert,
-				&FoundBoundaryEdgeVerts, &SimplifyOptions, &InMesh]
-				(int32 EID)
-				{
-					FixedCount += int32(EdgeTags[EID] == EElemTag::Fixed);
-					if (EdgeTags[EID] == EElemTag::Constrained)
+				int32 FixedCount = 0;
+				int32 ConstrainedCount = 0;
+
+				FVector3d BoundaryEdgeVert[2];
+				int32 FoundBoundaryEdgeVerts = 0;
+				InMeshConnectivity.EnumerateVertexEdges(VID,
+					[&OutEdgeTags, VID, &FixedCount, &ConstrainedCount, &BoundaryEdgeVert,
+					&FoundBoundaryEdgeVerts, &SimplifyOptions, &InMesh, &InMeshConnectivity]
+					(int32 EID)
 					{
-						ConstrainedCount++;
-						if (SimplifyOptions.FixBoundaryAngleTolerance > 0)
+						FixedCount += int32(OutEdgeTags[EID] == EElemTag::Fixed);
+						if (OutEdgeTags[EID] == EElemTag::Constrained)
 						{
-							if (InMesh.IsBoundaryEdge(EID))
+							ConstrainedCount++;
+							if (SimplifyOptions.FixBoundaryAngleTolerance > 0)
 							{
-								if (FoundBoundaryEdgeVerts < 2)
+								if (InMeshConnectivity.IsBoundaryEdge(EID))
 								{
-									FIndex2i EdgeV = InMesh.GetEdgeV(EID);
-									int32 OtherV = EdgeV.A == VID ? EdgeV.B : EdgeV.A;
+									if (FoundBoundaryEdgeVerts < 2)
+									{
+										FIndex2i EdgeV = InMeshConnectivity.GetEdgeV(EID);
+										int32 OtherV = EdgeV.A == VID ? EdgeV.B : EdgeV.A;
 
-									BoundaryEdgeVert[FoundBoundaryEdgeVerts] = InMesh.GetVertex(OtherV);
+										BoundaryEdgeVert[FoundBoundaryEdgeVerts] = InMesh.GetVertex(OtherV);
+									}
+									FoundBoundaryEdgeVerts++;
 								}
-								FoundBoundaryEdgeVerts++;
 							}
 						}
 					}
-				}
-			);
+				);
 
 
-			if (FixedCount > 0)
-			{
-				VertexTags[VID] = EElemTag::Fixed;
-				return;
-			}
-			if (FoundBoundaryEdgeVerts == 2)
-			{
-				FVector3d CenterV = InMesh.GetVertex(VID);
-				FVector3d E1 = Normalized(BoundaryEdgeVert[0] - CenterV);
-				FVector3d E2 = Normalized(CenterV - BoundaryEdgeVert[1]);
-				if (E1.Dot(E2) < CosBoundaryEdgeAngleTolerance)
+				if (FixedCount > 0)
 				{
-					VertexTags[VID] = EElemTag::Fixed;
+					OutVertexTags[VID] = EElemTag::Fixed;
 					return;
 				}
-			}
-			if (ConstrainedCount > 0)
-			{
-				if (ConstrainedCount == 2 &&
-					// seams are a special case where we can have two constrained edges but still be at a seam intersection 
-					// (e.g. at a vertex that joins two different types of seam)
-					(!InAttribs || !IsSeamIntersectionVertex(VID))
-					)
+				if (FoundBoundaryEdgeVerts == 2)
 				{
-					// constrain vertices along contiguous constrained edge paths
-					VertexTags[VID] = EElemTag::Constrained;
-					return;
+					FVector3d CenterV = InMesh.GetVertex(VID);
+					FVector3d E1 = Normalized(BoundaryEdgeVert[0] - CenterV);
+					FVector3d E2 = Normalized(CenterV - BoundaryEdgeVert[1]);
+					if (E1.Dot(E2) < CosBoundaryEdgeAngleTolerance)
+					{
+						OutVertexTags[VID] = EElemTag::Fixed;
+						return;
+					}
 				}
-				else
+				if (ConstrainedCount > 0)
 				{
-					// fix vertices at constraint intersections
-					VertexTags[VID] = EElemTag::Fixed;
-					return;
+					if (ConstrainedCount == 2 &&
+						// seams are a special case where we can have two constrained edges but still be at a seam intersection 
+						// (e.g. at a vertex that joins two different types of seam)
+						!IsSeamIntersectionVertex(VID)
+						)
+					{
+						// constrain vertices along contiguous constrained edge paths
+						OutVertexTags[VID] = EElemTag::Constrained;
+						return;
+					}
+					else
+					{
+						// fix vertices at constraint intersections
+						OutVertexTags[VID] = EElemTag::Fixed;
+						return;
+					}
 				}
-			}
 
-			VertexTags[VID] = EElemTag::Free;
-		});
+				OutVertexTags[VID] = EElemTag::Free;
+			});
+	}
 
-	TArray<float> TargetLengths;
-	if (SimplifyOptions.OptionalTargetEdgeLengthScale)
+	template<typename MeshType>
+	static TArray<float> InitTargetLengths(const MeshType& InMesh, const FSimplifyOptions& SimplifyOptions)
 	{
-		TargetLengths.SetNumUninitialized(InMesh.MaxVertexID());
-		ParallelFor(InMesh.MaxVertexID(), [&InMesh, &TargetLengths, &SimplifyOptions](int32 VID)
-			{
-				if (!InMesh.IsVertex(VID))
+		TArray<float> TargetLengths;
+		if (SimplifyOptions.OptionalTargetEdgeLengthScale)
+		{
+			TargetLengths.SetNumUninitialized(InMesh.MaxVertexID());
+			ParallelFor(InMesh.MaxVertexID(), [&InMesh, &TargetLengths, &SimplifyOptions](int32 VID)
 				{
-					return;
+					if (!InMesh.IsVertex(VID))
+					{
+						return;
+					}
+					TargetLengths[VID] = float(SimplifyOptions.TargetEdgeLength * SimplifyOptions.OptionalTargetEdgeLengthScale(VID));
 				}
-				TargetLengths[VID] = float(SimplifyOptions.TargetEdgeLength * SimplifyOptions.OptionalTargetEdgeLengthScale(VID));
-			}
-		);
-	}
-
-	///
-	/// Step 2. Clustering: Grow vertex clusters out to the target edge length size
-	///
-	
-	// Buckets of vertices to process -- vertices that are processed sooner are more likely to be directly included in the output
-	TStaticArray<TArray<int32>, NUM_TAGS> ProcessBuckets;
-	for (int32 VID : InMesh.VertexIndicesItr())
-	{
-		ProcessBuckets[(int32)VertexTags[VID]].Add(VID);
+			);
+		}
+		return TargetLengths;
 	}
 
-	TArray<float> SourceDist;
-	TArray<int32> Source;
-	Source.Init(INDEX_NONE, InMesh.MaxVertexID());
-	SourceDist.Init(FMathf::MaxReal, InMesh.MaxVertexID());
-
-	auto TagVerticesByRegionGrowth = 
-		[&Source, &SourceDist, &InMesh, &SimplifyOptions, &EdgeTags, &VertexTags, &TargetLengths]
-		(const TStaticArray<TArray<int32>, NUM_TAGS>& VertexIDBuckets)
+	template<typename MeshType, typename MeshConnType>
+	void ClusterVerticesByRegionGrowth(const MeshType& InMesh, const MeshConnType& InMeshConnectivity, TArray<int32>& Source, TArray<float>& SourceDist, const FSimplifyOptions& SimplifyOptions,
+		const TArray<EElemTag>& EdgeTags, const TArray<EElemTag>& VertexTags, const TArray<float>& TargetLengths,
+		const FVertexBuckets& VertexIDBuckets)
 	{
 		// add all the fixed vertices as sources first, so they can't be claimed by other verts
 		for (int32 VID : VertexIDBuckets[(int32)EElemTag::Fixed])
@@ -317,7 +201,7 @@ bool Simplify(const FDynamicMesh3& InMesh, FDynamicMesh3& ResultMesh, const FSim
 		TArray<FWalk> HeapV;
 
 		// for the non-fixed vertices, progressively grow from vertices, in passes from more-constrained to less-constrained edges
-		for (uint8 TagIdx = 1; TagIdx < (uint8)NUM_TAGS; ++TagIdx)
+		for (uint8 TagIdx = 1; TagIdx < (uint8)EElemTag::MAX; ++TagIdx)
 		{
 			for (uint8 BucketIdx = 0; BucketIdx <= TagIdx; ++BucketIdx)
 			{
@@ -347,13 +231,13 @@ bool Simplify(const FDynamicMesh3& InMesh, FDynamicMesh3& ResultMesh, const FSim
 					// helper to add candidate verts to a heap
 					HeapV.Reset();
 					auto AddCandidates = [MaxDist,
-						&HeapV, &InMesh, &SourceDist, &Source, &EdgeTags, &VertexTags, TagIdx]
+						&HeapV, &InMesh, &InMeshConnectivity, &SourceDist, &Source, &EdgeTags, &VertexTags, TagIdx]
 						(const FWalk& From)
 						{
 							// expand to one-ring
-							InMesh.EnumerateVertexEdges(From.VID,
+							InMeshConnectivity.EnumerateVertexEdges(From.VID,
 								[&From, MaxDist,
-								&HeapV, &InMesh, &SourceDist, &Source, &EdgeTags, &VertexTags, TagIdx]
+								&HeapV, &InMesh, &InMeshConnectivity, &SourceDist, &Source, &EdgeTags, &VertexTags, TagIdx]
 								(int32 EID)
 								{
 									if ((uint8)EdgeTags[EID] != TagIdx)
@@ -361,10 +245,10 @@ bool Simplify(const FDynamicMesh3& InMesh, FDynamicMesh3& ResultMesh, const FSim
 										return;
 									}
 
-									FIndex2i EdgeV = InMesh.GetEdgeV(EID);
+									FIndex2i EdgeV = InMeshConnectivity.GetEdgeV(EID);
 									int32 ToVID = EdgeV.A == From.VID ? EdgeV.B : EdgeV.A;
 
-									if ((uint8)VertexTags[ToVID] < TagIdx || From.Dist >= SourceDist[ToVID])
+									if (IsMoreConstrained(VertexTags[ToVID], TagIdx) || From.Dist >= SourceDist[ToVID])
 									{
 										// vertex was already claimed by more-constrained context, or is already as close (or closer) to another source
 										return;
@@ -407,26 +291,14 @@ bool Simplify(const FDynamicMesh3& InMesh, FDynamicMesh3& ResultMesh, const FSim
 				}
 			}
 		}
-	};
-
-	TagVerticesByRegionGrowth(ProcessBuckets);
-	for (int32 Idx = 0; Idx < ProcessBuckets.Num(); ++Idx)
-	{
-		ProcessBuckets[Idx].Empty();
 	}
 
-	
-	/// Step 2.5: Optionally look for clusters that didn't make it into triangles
-	for (int32 PassIdx = 0; PassIdx < SimplifyOptions.MaxPreserveCollapsedClusterPasses; ++PassIdx)
+	template<typename MeshType>
+	static void FindClustersInOutput(const MeshType& InMesh, const TArray<int32>& Source, TArray<bool>& OutClusterInOutput)
 	{
-		// factor used to set new Target Length for collapsed clusters, as a factor of max distance
-		constexpr float TargetLengthAsMaxClusterLengthFactor = .5f;
-
-		// Look through the triangles to tag the clusters that currently will end up in the output mesh
-		// due to having a valid triangle associated
-		TArray<bool> ClusterInOutput;
-		ClusterInOutput.SetNumZeroed(Source.Num());
-		ParallelFor(InMesh.MaxTriangleID(), [&InMesh, &Source, &ClusterInOutput](int32 TID)
+		OutClusterInOutput.Reset();
+		OutClusterInOutput.SetNumZeroed(InMesh.MaxVertexID());
+		ParallelFor(InMesh.MaxTriangleID(), [&InMesh, &Source, &OutClusterInOutput](int32 TID)
 			{
 				if (!InMesh.IsTriangle(TID))
 				{
@@ -438,203 +310,386 @@ bool Simplify(const FDynamicMesh3& InMesh, FDynamicMesh3& ResultMesh, const FSim
 				{
 					for (int32 SubIdx = 0; SubIdx < 3; ++SubIdx)
 					{
-						ClusterInOutput[Source[Tri[SubIdx]]] = true;
+						OutClusterInOutput[Source[Tri[SubIdx]]] = true;
 					}
 				}
 			}
 		);
-		TArray<int32> NeedRecluster;
-		TMap<int32, int32> VIDtoNeedReclusterIdx;
-		for (int32 VID : InMesh.VertexIndicesItr())
+	}
+
+	template<typename MeshType, typename MeshConnType>
+	static void PreserveCollapsedClusters(const MeshType& InMesh, const MeshConnType& InMeshConnectivity,
+		TArray<bool>& ClusterInOutput,
+		TArray<int32>& Source, TArray<float>& SourceDist, const FSimplifyOptions& SimplifyOptions,
+		const TArray<EElemTag>& EdgeTags, const TArray<EElemTag>& VertexTags, TArray<float>& TargetLengths,
+		const FVertexBuckets& VertexIDBuckets)
+	{
+		for (int32 PassIdx = 0; PassIdx < SimplifyOptions.MaxPreserveCollapsedClusterPasses; ++PassIdx)
 		{
-			if (VID == Source[VID] && !ClusterInOutput[VID])
+			// factor used to set new Target Length for collapsed clusters, as a factor of max distance
+			constexpr float TargetLengthAsMaxClusterLengthFactor = .5f;
+			
+			TArray<int32> NeedRecluster;
+			TMap<int32, int32> VIDtoNeedReclusterIdx;
+			for (int32 VID = 0, MaxVID = InMesh.MaxVertexID(); VID < MaxVID; ++VID)
 			{
-				VIDtoNeedReclusterIdx.Add(VID, NeedRecluster.Add(VID));
+				if (InMesh.IsVertex(VID) && VID == Source[VID] && !ClusterInOutput[VID])
+				{
+					VIDtoNeedReclusterIdx.Add(VID, NeedRecluster.Add(VID));
+				}
+			}
+			if (NeedRecluster.IsEmpty())
+			{
+				break;
 			}
-		}
-		if (NeedRecluster.IsEmpty())
-		{
-			break;
-		}
 
-		// Create buckets for vertices that will need re-processing due to not having any triangles
-		TStaticArray<TArray<int32>, NUM_TAGS> ReProcessBuckets;
-		// Make sure target lengths exist, so we can reduce them as needed for re-clustering
-		if (TargetLengths.IsEmpty())
-		{
-			TargetLengths.Init((float)SimplifyOptions.TargetEdgeLength, InMesh.MaxVertexID());
-		}
+			// Create buckets for vertices that will need re-processing due to not having any triangles
+			FVertexBuckets ReProcessBuckets;
+			// Make sure target lengths exist, so we can reduce them as needed for re-clustering
+			if (TargetLengths.IsEmpty())
+			{
+				TargetLengths.Init((float)SimplifyOptions.TargetEdgeLength, InMesh.MaxVertexID());
+			}
 
-		TArray<float> MaxSourceDistInCluster;
-		MaxSourceDistInCluster.SetNumZeroed(NeedRecluster.Num());
-		for (int32 VID : InMesh.VertexIndicesItr())
-		{
-			int32 SourceVID = Source[VID];
-			if (!ClusterInOutput[SourceVID]) // vertex belongs to a cluster that we need to reprocess
+			TArray<float> MaxSourceDistInCluster;
+			MaxSourceDistInCluster.SetNumZeroed(NeedRecluster.Num());
+			for (int32 VID = 0, MaxVID = InMesh.MaxVertexID(); VID < MaxVID; ++VID)
 			{
-				// update the max distance for the cluster
-				float& MaxFoundClusterDist = MaxSourceDistInCluster[VIDtoNeedReclusterIdx[SourceVID]];
-				MaxFoundClusterDist = FMath::Max(MaxFoundClusterDist, SourceDist[VID]);
-				ReProcessBuckets[(int32)VertexTags[VID]].Add(VID);
+				if (!InMesh.IsVertex(VID))
+				{
+					continue;
+				}
+				int32 SourceVID = Source[VID];
+				if (!ClusterInOutput[SourceVID]) // vertex belongs to a cluster that we need to reprocess
+				{
+					// update the max distance for the cluster
+					float& MaxFoundClusterDist = MaxSourceDistInCluster[VIDtoNeedReclusterIdx[SourceVID]];
+					MaxFoundClusterDist = FMath::Max(MaxFoundClusterDist, SourceDist[VID]);
+					ReProcessBuckets[(int32)VertexTags[VID]].Add(VID);
+				}
+			}
+			// For all re-process vertices, set down-scaled target lengths, and clear source/sourcedist
+			for (int32 BucketIdx = 0; BucketIdx < ReProcessBuckets.Num(); ++BucketIdx)
+			{
+				for (int32 Idx = 0; Idx < ReProcessBuckets[BucketIdx].Num(); ++Idx)
+				{
+					int32 UpdateVID = ReProcessBuckets[BucketIdx][Idx];
+					int32 SourceVID = Source[UpdateVID];
+					float ClusterMaxDist = MaxSourceDistInCluster[VIDtoNeedReclusterIdx[SourceVID]];
+					TargetLengths[UpdateVID] = ClusterMaxDist * TargetLengthAsMaxClusterLengthFactor;
+					Source[UpdateVID] = INDEX_NONE;
+					SourceDist[UpdateVID] = FMathf::MaxReal;
+				}
 			}
+			// Re-run tagging
+			ClusterVerticesByRegionGrowth(InMesh, InMeshConnectivity, Source, SourceDist, SimplifyOptions, EdgeTags, VertexTags, TargetLengths, ReProcessBuckets);
+			// Update tracking of clusters in output
+			FindClustersInOutput(InMesh, Source, ClusterInOutput);
 		}
-		// For all re-process vertices, set down-scaled target lengths, and clear source/sourcedist
-		for (int32 BucketIdx = 0; BucketIdx < ReProcessBuckets.Num(); ++BucketIdx)
+	}
+
+	template<typename InMeshType, typename InMeshConnType, typename ResultMeshType>
+	void BuildSimplifiedMeshFromClusters(ResultMeshType& ResultMesh, TArray<int32>& FromResVID, TArray<int32>& ResultToSourceTri,
+		TArray<bool>& ClusterInOutput,
+		bool& bOutResultHasDuplicateVertices,
+		const InMeshType& InMesh, const InMeshConnType& InMeshConnectivity, TArray<int32>& Source, TArray<float>& SourceDist, const FSimplifyOptions& SimplifyOptions,
+		const TArray<EElemTag>& EdgeTags, TArray<EElemTag>& VertexTags, const TArray<float>& TargetLengths,
+		FVertexBuckets& ProcessBuckets)
+	{
+		TArray<int32> ToResVID;
+
+		// If simplification introduces non-manifold edges, we can often recover by fixing more vertices and re-attempting the build.
+		// After MeshBuildAttempts tries, if still failing, we stop adding vertices and just duplicate vertices to add the non-manifold triangles.
+		// TODO: We could potentially analyze the cluster connectivity more carefully handle more degenerate cluster connectivity, more robustly.
+		//		(if so -- it may be better to do so by analyzing the graph before building the ResultMesh, rather than this rebuilding approach!)
+		int32 MeshBuildAttempts = 2;
+
+		while (MeshBuildAttempts-- > 0)
 		{
-			for (int32 Idx = 0; Idx < ReProcessBuckets[BucketIdx].Num(); ++Idx)
+			// clear mesh outputs
+			ToResVID.Reset();
+			FromResVID.Reset();
+			ResultToSourceTri.Reset();
+			ResultMesh.Clear();
+
+			bool bAllowDegenerate = MeshBuildAttempts <= 0;
+			// Array of vertex IDs to set to 'fixed' on a rebuild attempt
+			TArray<int32> SourceVIDToFix;
+
+			ToResVID.Init(INDEX_NONE, InMesh.MaxVertexID());
+			for (int32 VID = 0; VID < Source.Num(); ++VID)
 			{
-				int32 UpdateVID = ReProcessBuckets[BucketIdx][Idx];
-				int32 SourceVID = Source[UpdateVID];
-				float ClusterMaxDist = MaxSourceDistInCluster[VIDtoNeedReclusterIdx[SourceVID]];
-				TargetLengths[UpdateVID] = ClusterMaxDist * TargetLengthAsMaxClusterLengthFactor;
-				Source[UpdateVID] = INDEX_NONE;
-				SourceDist[UpdateVID] = FMathf::MaxReal;
+				if (Source[VID] == VID && ClusterInOutput[VID])
+				{
+					ToResVID[VID] = ResultMesh.AppendVertex(InMesh.GetVertex(VID));
+					// we need the reverse mapping if we're transferring seams
+					if (SimplifyOptions.bTransferAttributes)
+					{
+						FromResVID.Add(VID);
+					}
+				}
 			}
+
+			for (int32 TID = 0, MaxTID = InMesh.MaxTriangleID(); TID < MaxTID; ++TID)
+			{
+				if (!InMesh.IsTriangle(TID))
+				{
+					continue;
+				}
+				FIndex3i Tri = InMesh.GetTriangle(TID);
+				FIndex3i SourceTri(Source[Tri.A], Source[Tri.B], Source[Tri.C]);
+				if (SourceTri.A != SourceTri.B && SourceTri.A != SourceTri.C && SourceTri.B != SourceTri.C)
+				{
+					FIndex3i ResTri(ToResVID[SourceTri.A], ToResVID[SourceTri.B], ToResVID[SourceTri.C]);
+					int32 ResultTID = ResultMesh.AppendTriangle(ResTri);
+					if (ResultTID == FDynamicMesh3::NonManifoldID)
+					{
+						if (bAllowDegenerate)
+						{
+							// TODO: only duplicate vertices on the non-manifold edge(s)
+							FIndex3i ExtraTri;
+							ExtraTri.A = ResultMesh.AppendVertex(ResultMesh.GetVertex(ResTri.A));
+							FromResVID.Add(SourceTri.A);
+							ExtraTri.B = ResultMesh.AppendVertex(ResultMesh.GetVertex(ResTri.B));
+							FromResVID.Add(SourceTri.B);
+							ExtraTri.C = ResultMesh.AppendVertex(ResultMesh.GetVertex(ResTri.C));
+							FromResVID.Add(SourceTri.C);
+							ResultTID = ResultMesh.AppendTriangle(ExtraTri);
+							bOutResultHasDuplicateVertices = true;
+						}
+						else
+						{
+							// Non-manifold edges can often be resolved by adding an extra vertex --
+							// mark the vertex with largest SourceDist for inclusion in the result mesh
+							int32 BestSubIdx = INDEX_NONE;
+							float BestDist = 0;
+							for (int32 SubIdx = 0; SubIdx < 3; ++SubIdx)
+							{
+								if (SourceDist[Tri[SubIdx]] > BestDist)
+								{
+									BestDist = SourceDist[Tri[SubIdx]];
+									BestSubIdx = SubIdx;
+								}
+							}
+							if (BestSubIdx != INDEX_NONE)
+							{
+								SourceVIDToFix.Add(Tri[BestSubIdx]);
+							}
+						}
+					}
+					if ((SimplifyOptions.bTransferAttributes || SimplifyOptions.bTransferGroups) && ResultTID >= 0)
+					{
+						checkSlow(ResultTID == ResultToSourceTri.Num()); // ResultMesh starts empty and should be compact
+						ResultToSourceTri.Add(TID);
+					}
+				}
+			}
+
+			// We marked some new vertices for inclusion in the result; tag them and re-try
+			if (!bAllowDegenerate && SourceVIDToFix.Num() > 0)
+			{
+				for (int32 VID : SourceVIDToFix)
+				{
+					VertexTags[VID] = EElemTag::Fixed;
+				}
+				ProcessBuckets[0] = MoveTemp(SourceVIDToFix);
+				ClusterVerticesByRegionGrowth(InMesh, InMeshConnectivity, Source, SourceDist, SimplifyOptions, EdgeTags, VertexTags, TargetLengths, ProcessBuckets);
+				FindClustersInOutput(InMesh, Source, ClusterInOutput);
+				continue;
+			}
+
+			// Accept the result mesh triangulation
+			break;
 		}
-		// Re-run tagging
-		TagVerticesByRegionGrowth(ReProcessBuckets);
 	}
+}
 
-	///
-	/// Step 3: Copy the cluster connectivity out to our ResultMesh
-	/// 
+namespace UE::Geometry::MeshClusterSimplify
+{
 
-	TArray<int32> ToResVID, FromResVID;
-	TArray<int32> ResultToSourceTri;
+bool Simplify(const FDynamicMesh3& InMesh, FDynamicMesh3& ResultMesh, const FSimplifyOptions& SimplifyOptions)
+{
+	TRACE_CPUPROFILER_EVENT_SCOPE(MeshClusterSimplify::Simplify);
 
-	// If simplification introduces non-manifold edges, we can often recover by fixing more vertices and re-attempting the build.
-	// After MeshBuildAttempts tries, if still failing, we stop adding vertices and just duplicate vertices to add the non-manifold triangles.
-	// TODO: We could potentially analyze the cluster connectivity more carefully handle more degenerate cluster connectivity, more robustly.
-	//		(if so -- it may be better to do so by analyzing the graph before building the ResultMesh, rather than this rebuilding approach!)
-	int32 MeshBuildAttempts = 2;
-	bool bResultHasDuplicateVertices = false;
+	using namespace UE::MeshClusterSimplifyLocals;
 
-	while (MeshBuildAttempts-- > 0)
+	// we build the result mesh by incrementally copying from the input mesh, so they shouldn't be the same mesh
+	if (!ensure(&ResultMesh != &InMesh))
 	{
-		// clear mesh outputs
-		ToResVID.Reset();
-		FromResVID.Reset();
-		ResultToSourceTri.Reset();
-		ResultMesh.Clear();
+		return false;
+	}
+
+	ResultMesh.Clear();
+	
+	const FDynamicMeshAttributeSet* InAttribs = InMesh.Attributes();
 
+	// We tag edges and vertices w/ the constraint level, abbreviated to EElemTag for convenience
+	using EElemTag = FSimplifyOptions::EConstraintLevel;
 
-		bool bAllowDegenerate = MeshBuildAttempts <= 0;
-		// Array of vertex IDs to set to 'fixed' on a rebuild attempt
-		TArray<int32> SourceVIDToFix;
+	// TODO: optionally also compute some vertex curvature feature & sort by it, to favor capturing less flat parts of the input shape?
 
-		ToResVID.Init(INDEX_NONE, InMesh.MaxVertexID());
-		for (int32 VID = 0; VID < Source.Num(); ++VID)
+	///
+	/// Step 1, Data Prep: Translate all mesh constraint options to simple per-edge and per-vertex tags, so we know what to try to especially preserve in the result
+	///
+	
+	auto GetEdgeConstraintLevelFromAttributes =
+		[InAttribs, &InMesh, &SimplifyOptions]
+		(int32 EID, EElemTag UseTag) -> EElemTag
 		{
-			if (Source[VID] == VID)
+			if (IsMoreConstrained(SimplifyOptions.PreserveEdges.PolyGroup, UseTag) && InMesh.IsGroupBoundaryEdge(EID))
 			{
-				ToResVID[VID] = ResultMesh.AppendVertex(InMesh.GetVertex(VID));
-				// we need the reverse mapping if we're transferring seams
-				if (SimplifyOptions.bTransferAttributes)
-				{
-					FromResVID.Add(VID);
-				}
+				UseTag = SimplifyOptions.PreserveEdges.PolyGroup;
 			}
-		}
 
-		for (int32 TID : InMesh.TriangleIndicesItr())
-		{
-			FIndex3i Tri = InMesh.GetTriangle(TID);
-			FIndex3i SourceTri(Source[Tri.A], Source[Tri.B], Source[Tri.C]);
-			if (SourceTri.A != SourceTri.B && SourceTri.A != SourceTri.C && SourceTri.B != SourceTri.C)
+			if (InAttribs)
 			{
-				FIndex3i ResTri(ToResVID[SourceTri.A], ToResVID[SourceTri.B], ToResVID[SourceTri.C]);
-				int32 ResultTID = ResultMesh.AppendTriangle(ResTri);
-				if (ResultTID == FDynamicMesh3::NonManifoldID)
+				if (IsMoreConstrained(SimplifyOptions.PreserveEdges.Material, UseTag) && InAttribs->IsMaterialBoundaryEdge(EID))
+				{
+					UseTag = SimplifyOptions.PreserveEdges.Material;
+				}
+
+				if (IsMoreConstrained(SimplifyOptions.PreserveEdges.UVSeam, UseTag))
 				{
-					if (bAllowDegenerate)
+					for (int32 UVLayer = 0; UVLayer < InAttribs->NumUVLayers(); ++UVLayer)
 					{
-						// TODO: only duplicate vertices on the non-manifold edge(s)
-						FIndex3i ExtraTri;
-						ExtraTri.A = ResultMesh.AppendVertex(ResultMesh.GetVertex(ResTri.A));
-						FromResVID.Add(SourceTri.A);
-						ExtraTri.B = ResultMesh.AppendVertex(ResultMesh.GetVertex(ResTri.B));
-						FromResVID.Add(SourceTri.B);
-						ExtraTri.C = ResultMesh.AppendVertex(ResultMesh.GetVertex(ResTri.C));
-						FromResVID.Add(SourceTri.C);
-						ResultTID = ResultMesh.AppendTriangle(ExtraTri);
-						bResultHasDuplicateVertices = true;
+						if (InAttribs->GetUVLayer(UVLayer)->IsSeamEdge(EID))
+						{
+							UseTag = SimplifyOptions.PreserveEdges.UVSeam;
+							break;
+						}
 					}
-					else
+				}
+				if (IsMoreConstrained(SimplifyOptions.PreserveEdges.TangentSeam, UseTag))
+				{
+					for (int32 NormalLayer = 1; NormalLayer < InAttribs->NumNormalLayers(); ++NormalLayer)
 					{
-						// Non-manifold edges can often be resolved by adding an extra vertex --
-						// mark the vertex with largest SourceDist for inclusion in the result mesh
-						int32 BestSubIdx = INDEX_NONE;
-						float BestDist = 0;
-						for (int32 SubIdx = 0; SubIdx < 3; ++SubIdx)
+						if (InAttribs->GetNormalLayer(NormalLayer)->IsSeamEdge(EID))
 						{
-							if (SourceDist[Tri[SubIdx]] > BestDist)
-							{
-								BestDist = SourceDist[Tri[SubIdx]];
-								BestSubIdx = SubIdx;
-							}
+							UseTag = SimplifyOptions.PreserveEdges.TangentSeam;
+							break;
 						}
-						if (BestSubIdx != INDEX_NONE)
+					}
+				}
+				if (IsMoreConstrained(SimplifyOptions.PreserveEdges.NormalSeam, UseTag))
+				{
+					if (const FDynamicMeshNormalOverlay* Normals = InAttribs->PrimaryNormals())
+					{
+						if (Normals->IsSeamEdge(EID))
 						{
-							SourceVIDToFix.Add(Tri[BestSubIdx]);
+							UseTag = SimplifyOptions.PreserveEdges.NormalSeam;
 						}
 					}
 				}
-				if ((SimplifyOptions.bTransferAttributes || SimplifyOptions.bTransferGroups) && ResultTID >= 0)
+				if (IsMoreConstrained(SimplifyOptions.PreserveEdges.ColorSeam, UseTag))
 				{
-					checkSlow(ResultTID == ResultToSourceTri.Num()); // ResultMesh starts empty and should be compact
-					ResultToSourceTri.Add(TID);
+					if (const FDynamicMeshColorOverlay* Colors = InAttribs->PrimaryColors())
+					{
+						if (Colors->IsSeamEdge(EID))
+						{
+							UseTag = SimplifyOptions.PreserveEdges.ColorSeam;
+						}
+					}
 				}
 			}
-		}
 
-		// We marked some new vertices for inclusion in the result; tag them and re-try
-		if (!bAllowDegenerate && SourceVIDToFix.Num() > 0)
+			return UseTag;
+		};
+	auto IsSeamIntersectionVertex =
+		[&InAttribs, &SimplifyOptions]
+		(int32 VID) -> bool
 		{
-			for (int32 VID : SourceVIDToFix)
+			if (!InAttribs)
 			{
-				VertexTags[VID] = EElemTag::Fixed;
+				return false;
 			}
-			ProcessBuckets[0] = MoveTemp(SourceVIDToFix);
-			TagVerticesByRegionGrowth(ProcessBuckets);
-			continue;
-		}
 
-		// Accept the result mesh triangulation
-		break;
-	}
+			if (SimplifyOptions.PreserveEdges.UVSeam == EElemTag::Constrained)
+			{
+				for (int32 Layer = 0; Layer < InAttribs->NumUVLayers(); ++Layer)
+				{
+					if (InAttribs->GetUVLayer(Layer)->IsSeamIntersectionVertex(VID))
+					{
+						return true;
+					}
+				}
+			}
+			if (SimplifyOptions.PreserveEdges.NormalSeam == EElemTag::Constrained)
+			{
+				if (const FDynamicMeshNormalOverlay* Normals = InAttribs->PrimaryNormals())
+				{
+					if (Normals->IsSeamIntersectionVertex(VID))
+					{
+						return true;
+					}
+				}
+			}
+			if (SimplifyOptions.PreserveEdges.TangentSeam == EElemTag::Constrained)
+			{
+				for (int32 Layer = 1; Layer < InAttribs->NumNormalLayers(); ++Layer)
+				{
+					if (InAttribs->GetNormalLayer(Layer)->IsSeamIntersectionVertex(VID))
+					{
+						return true;
+					}
+				}
+			}
+			if (SimplifyOptions.PreserveEdges.ColorSeam == EElemTag::Constrained)
+			{
+				if (const FDynamicMeshColorOverlay* Colors = InAttribs->PrimaryColors())
+				{
+					if (Colors->IsSeamIntersectionVertex(VID))
+					{
+						return true;
+					}
+				}
+			}
 
-	// We don't need the ToResVID mapping anymore.
-	// Empty it to make it clear it is not valid (since we don't remap it if we compact, below)
-	ToResVID.Empty();
+			return false;
+		};
+	
+	TArray<EElemTag> EdgeTags, VertexTags;
+	TagVerticesAndEdges(InMesh, InMesh, 
+		GetEdgeConstraintLevelFromAttributes, IsSeamIntersectionVertex,
+		SimplifyOptions, EdgeTags, VertexTags);
 
-	// Remove any unused vertices and compact + fix maps accordingly
-	// Note we do this before transferring attributes, when it is cheaper,
-	// though this does require updating mappings to the source mesh accordingly
-	if (ResultMesh.RemoveUnusedVertices() > 0)
+	TArray<float> TargetLengths = InitTargetLengths(InMesh, SimplifyOptions);
+
+	///
+	/// Step 2. Clustering: Grow vertex clusters out to the target edge length size
+	///
+	
+	// Buckets of vertices to process -- vertices that are processed sooner are more likely to be directly included in the output
+	FVertexBuckets ProcessBuckets;
+	for (int32 VID : InMesh.VertexIndicesItr())
 	{
-		// Triangles are already compact by construction, so don't need to remap
-		check(ResultMesh.IsCompactT());
+		ProcessBuckets[(int32)VertexTags[VID]].Add(VID);
+	}
 
-		// Compact away unused vertices
-		FCompactMaps CompactMaps;
-		ResultMesh.CompactInPlace(&CompactMaps);
+	TArray<float> SourceDist;
+	TArray<int32> Source;
+	Source.Init(INDEX_NONE, InMesh.MaxVertexID());
+	SourceDist.Init(FMathf::MaxReal, InMesh.MaxVertexID());
 
-		if (SimplifyOptions.bTransferAttributes)
-		{
-			// remap FromResVID, which maps Result VID -> Source VID
-			TArray<int32> NewFromResVID;
-			NewFromResVID.SetNumUninitialized(ResultMesh.MaxVertexID());
-			for (int32 PrevVID = 0, PrevNum = CompactMaps.NumVertexMappings(); PrevVID < PrevNum; ++PrevVID)
-			{
-				int32 NewVID = CompactMaps.GetVertexMapping(PrevVID);
-				if (NewVID != IndexConstants::InvalidID)
-				{
-					NewFromResVID[NewVID] = FromResVID[PrevVID];
-				}
-			}
-			FromResVID = MoveTemp(NewFromResVID);
-		}
+	ClusterVerticesByRegionGrowth(InMesh, InMesh, Source, SourceDist, SimplifyOptions, EdgeTags, VertexTags, TargetLengths, ProcessBuckets);
+	for (int32 Idx = 0; Idx < ProcessBuckets.Num(); ++Idx)
+	{
+		ProcessBuckets[Idx].Empty();
 	}
+	
+	/// Step 2.5: Determine which clusters are in the output, and optionally try reduce target lengths to preserve missing clusters
+	TArray<bool> ClusterInOutput;
+	FindClustersInOutput(InMesh, Source, ClusterInOutput);
+	PreserveCollapsedClusters(InMesh, InMesh, ClusterInOutput, Source, SourceDist, SimplifyOptions, EdgeTags, VertexTags, TargetLengths, ProcessBuckets);
+
+	///
+	/// Step 3: Copy the cluster connectivity out to our ResultMesh
+	///
+
+	TArray<int32> FromResVID, ResultToSourceTri;
+	bool bResultHasDuplicateVertices = false;
+	BuildSimplifiedMeshFromClusters<FDynamicMesh3, FDynamicMesh3, FDynamicMesh3>(ResultMesh, FromResVID, ResultToSourceTri, ClusterInOutput, bResultHasDuplicateVertices,
+		InMesh, InMesh, Source, SourceDist, SimplifyOptions, EdgeTags, VertexTags, TargetLengths,
+		ProcessBuckets);
 
 	///
 	/// Step 4: After accepting the final ResultMesh triangulation, copy the input mesh's attributes (UVs, materials, etc) over as well
@@ -870,4 +925,99 @@ bool Simplify(const FDynamicMesh3& InMesh, FDynamicMesh3& ResultMesh, const FSim
 	return true;
 }
 
+bool Simplify(const FTriangleMeshAdapterd& InMesh, FResultMeshAdapter& OutSimplifiedMesh, const FSimplifyOptions& SimplifyOptions)
+{
+	TRACE_CPUPROFILER_EVENT_SCOPE(MeshClusterSimplify::Simplify);
+
+	using namespace UE::MeshClusterSimplifyLocals;
+
+	// We tag edges and vertices w/ the constraint level, abbreviated to EElemTag for convenience
+	using EElemTag = FSimplifyOptions::EConstraintLevel;
+
+
+
+	///
+	/// Step 1, Data Prep: Translate all mesh constraint options to simple per-edge and per-vertex tags, so we know what to try to especially preserve in the result
+	///
+	
+	using FMeshConnectivity = TTriangleMeshAdapterEdgeConnectivity<TTriangleMeshAdapter<double>>;
+	FMeshConnectivity MeshConnectivity(&InMesh);
+
+	auto TagEdgeFromAttributes = [](int32 EID, EElemTag InitialTag)->EElemTag
+		{
+			return InitialTag;
+		};
+	auto IsSeamIntersectionVertex = [](int32 VID)->bool
+		{
+			return false;
+		};
+
+	TArray<EElemTag> EdgeTags, VertexTags;
+	TagVerticesAndEdges(InMesh, MeshConnectivity,
+		TagEdgeFromAttributes, IsSeamIntersectionVertex,
+		SimplifyOptions, EdgeTags, VertexTags);
+
+	TArray<float> TargetLengths = InitTargetLengths(InMesh, SimplifyOptions);
+
+	///
+	/// Step 2. Clustering: Grow vertex clusters out to the target edge length size
+	///
+
+	// Buckets of vertices to process -- vertices that are processed sooner are more likely to be directly included in the output
+	FVertexBuckets ProcessBuckets;
+	for (int32 VID = 0, MaxVID = InMesh.MaxVertexID(); VID < MaxVID; ++VID)
+	{
+		if (InMesh.IsVertex(VID))
+		{
+			ProcessBuckets[(int32)VertexTags[VID]].Add(VID);
+		}
+	}
+
+	TArray<float> SourceDist;
+	TArray<int32> Source;
+	Source.Init(INDEX_NONE, InMesh.MaxVertexID());
+	SourceDist.Init(FMathf::MaxReal, InMesh.MaxVertexID());
+
+	ClusterVerticesByRegionGrowth(InMesh, MeshConnectivity, Source, SourceDist, SimplifyOptions, EdgeTags, VertexTags, TargetLengths, ProcessBuckets);
+	for (int32 Idx = 0; Idx < ProcessBuckets.Num(); ++Idx)
+	{
+		ProcessBuckets[Idx].Empty();
+	}
+
+
+	/// Step 2.5: Optionally look for clusters that didn't make it into triangles
+	TArray<bool> ClusterInOutput;
+	FindClustersInOutput(InMesh, Source, ClusterInOutput);
+	PreserveCollapsedClusters(InMesh, MeshConnectivity, ClusterInOutput, Source, SourceDist, SimplifyOptions, EdgeTags, VertexTags, TargetLengths, ProcessBuckets);
+
+	///
+	/// Step 3: Copy the cluster connectivity out to our ResultMesh
+	///
+
+	TArray<int32> FromResVID, ResultToSourceTri;
+	bool bResultHasDuplicateVertices = false;
+	BuildSimplifiedMeshFromClusters<TTriangleMeshAdapter<double>, FMeshConnectivity, FResultMeshAdapter>(OutSimplifiedMesh, FromResVID, ResultToSourceTri, ClusterInOutput,
+		bResultHasDuplicateVertices,
+		InMesh, MeshConnectivity, Source, SourceDist, SimplifyOptions, EdgeTags, VertexTags, TargetLengths,
+		ProcessBuckets);
+
+	//
+	// Step 4: Transfer attributes out to the result mesh
+	//
+
+	if (SimplifyOptions.bTransferAttributes)
+	{
+		if (OutSimplifiedMesh.TransferPerVertexAttributes)
+		{
+			OutSimplifiedMesh.TransferPerVertexAttributes(FromResVID);
+		}
+		if (OutSimplifiedMesh.TransferPerTriangleAttributes)
+		{
+			OutSimplifiedMesh.TransferPerTriangleAttributes(ResultToSourceTri);
+		}
+	}
+
+	return true;
+}
+
 } // namespace UE::Geometry
diff --git a/Engine/Plugins/Runtime/GeometryProcessing/Source/DynamicMesh/Public/Operations/MeshClusterSimplifier.h b/Engine/Plugins/Runtime/GeometryProcessing/Source/DynamicMesh/Public/Operations/MeshClusterSimplifier.h
--- a/Engine/Plugins/Runtime/GeometryProcessing/Source/DynamicMesh/Public/Operations/MeshClusterSimplifier.h
+++ b/Engine/Plugins/Runtime/GeometryProcessing/Source/DynamicMesh/Public/Operations/MeshClusterSimplifier.h
@@ -5,6 +5,7 @@
 #include "GeometryTypes.h"
 #include "IndexTypes.h"
 #include "Clustering/FaceNormalClustering.h"
+#include "MeshAdapter.h"
 
 class FProgressCancel;
 
@@ -36,7 +37,8 @@ struct FSimplifyOptions
 	// If > 0, boundary vertices w/ incident boundary edge angle greater than this (in degrees) will be kept in the output
 	double FixBoundaryAngleTolerance = 45;
 
-	// Constraint options control what simplifications are allowed
+	// Constraint levels control what simplifications are allowed
+	// Ordered from most-constrained to least-constrained
 	enum class EConstraintLevel : uint8
 	{
 		// Fixed vertices/edges will generally be preserved in the output, as they will each be given their own cluster
@@ -45,7 +47,9 @@ struct FSimplifyOptions
 		// A vertex at an intersection of more than two constrained edges will be automatically preserved as 'Fixed'
 		Constrained,
 		// No constraints / ok to simplify as much as possible
-		Free
+		Free,
+		// The count of total constraint levels, used for iterating over the levels. (Not itself a valid constraint level)
+		MAX
 	};
 
 	struct FPreserveFeatures
@@ -106,5 +110,63 @@ struct FSimplifyOptions
  */
 bool DYNAMICMESH_API Simplify(const FDynamicMesh3& InMesh, FDynamicMesh3& OutSimplifiedMesh, const FSimplifyOptions& SimplifyOptions);
 
+/**
+ * Adapter for producing an output simplified triangle mesh in any target triangle mesh structure
+ */
+struct FResultMeshAdapter
+{
+	// Add a vertex to the result mesh, and return its ID
+	TFunction<int32(FVector3d)> AppendVertex;
+
+	// Add a triangle to the result mesh, and return its ID
+	// Note this method should follow the same error conventions as FDynamicMesh3, so
+	//  if it refuses to add a non-manifold triangle, it should return FDynamicMesh3::NonManifoldID
+	TFunction<int32(FIndex3i)> AppendTriangle;
+
+	// Clear the current mesh from the results
+	TFunction<void()> Clear;
+
+	// Accessor for vertices
+	TFunction<FVector3d(int32)> GetVertex;
+	// Accessor for triangles
+	TFunction<FIndex3i(int32)> GetTriangle;
+
+	// Initialize the adapter from a mesh type that has the same basic interface
+	// Note that for attributes to transfer, the optional attribute transfer methods must be implemented separately
+	template<typename MeshType>
+	void Init(MeshType* Mesh)
+	{
+		AppendVertex = [Mesh](FVector3d V) -> int32
+			{
+				return Mesh->AppendVertex(V);
+			};
+		AppendTriangle = [Mesh](FIndex3i T) -> int32
+			{
+				return Mesh->AppendTriangle(T);
+			};
+		Clear = [Mesh]() -> void
+			{
+				Mesh->Clear();
+			};
+		GetVertex = [Mesh](int32 VID) -> FVector3d
+			{
+				return Mesh->GetVertex(VID);
+			};
+		GetTriangle = [Mesh](int32 TID) -> FIndex3i
+			{
+				return Mesh->GetTriangle(TID);
+			};
+	}
+
+	// Attribute transfer methods; if not provided, attributes will not be transferred
+	TFunction<void(TConstArrayView<int32> ResultToSourceVertexID)> TransferPerVertexAttributes;
+	TFunction<void(TConstArrayView<int32> ResultToSourceTriangleID)> TransferPerTriangleAttributes;
+	// TODO: also provide a generic transfer method for attributes that may have seams (such as UVs)
+};
+
+// Simplify an input mesh provided by the general triangle mesh adapter, placing the result in an output mesh via a similar adapter interface
+// Note: The output mesh must not be aliased with input mesh data
+bool DYNAMICMESH_API Simplify(const FTriangleMeshAdapterd& InMesh, FResultMeshAdapter& OutSimplifiedMesh, const FSimplifyOptions& SimplifyOptions);
+
 } // end namespace UE::Geometry
 } // end namespace UE
