Added path operations

diff --git a/Engine/Plugins/Runtime/GeometryProcessing/Source/GeometryAlgorithms/Private/Curve/ClipperUtils.h b/Engine/Plugins/Runtime/GeometryProcessing/Source/GeometryAlgorithms/Private/Curve/ClipperUtils.h
--- a/Engine/Plugins/Runtime/GeometryProcessing/Source/GeometryAlgorithms/Private/Curve/ClipperUtils.h
+++ b/Engine/Plugins/Runtime/GeometryProcessing/Source/GeometryAlgorithms/Private/Curve/ClipperUtils.h
@@ -44,6 +44,12 @@ namespace UE::Geometry::Private
 	 
 	template <typename InputType, typename RealType>
 	static void ConvertPolyTreeToPolygons(const Clipper2Lib::PolyPath<InputType>* InPaths, TArray<UE::Geometry::TGeneralPolygon2<RealType>>& OutPolygons, const TVector2<RealType>& InMin = {} , const RealType& InRange = {}, int ParentPolygonIdx = -1);
+
+	template<typename InputType, typename RealType>
+	TArray<TVector2<RealType>> ConvertPathToUEPath(const Clipper2Lib::Path<InputType>& InPath, const TVector2<RealType>& InMin, const RealType& InRange);
+
+	template<typename InputType, typename RealType>
+	TArray<TArray<TVector2<RealType>>> ConvertPathsToUEPaths(const Clipper2Lib::Paths<InputType>& InPaths, const TVector2<RealType>& InMin, const RealType& InRange);
 }
 
 #include "ClipperUtils.inl"
diff --git a/Engine/Plugins/Runtime/GeometryProcessing/Source/GeometryAlgorithms/Private/Curve/ClipperUtils.inl b/Engine/Plugins/Runtime/GeometryProcessing/Source/GeometryAlgorithms/Private/Curve/ClipperUtils.inl
--- a/Engine/Plugins/Runtime/GeometryProcessing/Source/GeometryAlgorithms/Private/Curve/ClipperUtils.inl
+++ b/Engine/Plugins/Runtime/GeometryProcessing/Source/GeometryAlgorithms/Private/Curve/ClipperUtils.inl
@@ -151,4 +151,37 @@ namespace UE::Geometry::Private
 			ConvertPolyTreeToPolygons<InputType, RealType>(ChildPath, OutPolygons, InMin, InRange, ParentPolygonIdx);
 		}
 	}
+
+	template<typename InputType, typename RealType>
+	TArray<TVector2<RealType>> ConvertPathToUEPath(const Clipper2Lib::Path<InputType>& InPath, const TVector2<RealType>& InMin, const RealType& InRange)
+	{
+		TArray<TVector2<RealType>> Path;
+		Path.Reserve(InPath.size());
+		for (const Clipper2Lib::Point<InputType>& Point : InPath)
+		{
+			if constexpr (TIsIntegral<InputType>::Value)
+			{
+				Path.Add(UnpackVector<InputType, RealType>(Point, InMin, InRange));
+			}
+			else if constexpr (TIsFloatingPoint<InputType>::Value)
+			{
+				Path.Add({ static_cast<RealType>(Point.x), static_cast<RealType>(Point.y) });
+			}
+		}
+
+		return Path;
+	}
+
+	template<typename InputType, typename RealType>
+	TArray<TArray<TVector2<RealType>>> ConvertPathsToUEPaths(const Clipper2Lib::Paths<InputType>& InPaths, const TVector2<RealType>& InMin, const RealType& InRange)
+	{
+		TArray<TArray<TVector2<RealType>>> Paths;
+		Paths.Reserve(InPaths.size());
+		for (const Clipper2Lib::Path<InputType>& InPath : InPaths)
+		{
+			Paths.Add(ConvertPathToUEPath<InputType, RealType>(InPath, InMin, InRange));
+		}
+
+		return Paths;
+	}
 }
diff --git a/Engine/Plugins/Runtime/GeometryProcessing/Source/GeometryAlgorithms/Private/Curve/PolygonIntersectionUtils.cpp b/Engine/Plugins/Runtime/GeometryProcessing/Source/GeometryAlgorithms/Private/Curve/PolygonIntersectionUtils.cpp
--- a/Engine/Plugins/Runtime/GeometryProcessing/Source/GeometryAlgorithms/Private/Curve/PolygonIntersectionUtils.cpp
+++ b/Engine/Plugins/Runtime/GeometryProcessing/Source/GeometryAlgorithms/Private/Curve/PolygonIntersectionUtils.cpp
@@ -82,6 +82,46 @@ namespace UE::Geometry::Private
 
 		return bExecuteResult;
 	}
+
+	// Clip paths in the subject array by the polygons in the clip array
+	template <typename RealType>
+	static bool ClipPaths(
+		const Clipper2Lib::ClipType& InClipType,
+		const TArray<TArrayView<FVector2D>>& SubjPaths,
+		TArrayView<const FGeneralPolygon2d> ClipPolygons,
+		TArray<TArray<FVector2D>>& ResultOut)
+	{
+		TAxisAlignedBox2<RealType> InputBounds;
+		for (TArrayView<FVector2D> Path : SubjPaths)
+		{
+			InputBounds.Contain(Path);
+		}
+		for (const FGeneralPolygon2d& Polygon : ClipPolygons)
+		{
+			InputBounds.Contain(Polygon.GetOuter().Bounds());
+		}
+		RealType InputRange = (InputBounds.Extents() * 2).GetMax();
+
+		Clipper2Lib::Clipper64 Clipper;
+
+		Clipper.AddOpenSubject(ConvertPolygonsToPaths<RealType, IntegralType>(SubjPaths, InputBounds.Min, InputRange));
+
+		for (const FGeneralPolygon2d& Polygon : ClipPolygons)
+		{
+			Clipper.AddClip(ConvertGeneralizedPolygonToPath<RealType, IntegralType>(Polygon, InputBounds.Min, InputRange));
+		}
+
+		Clipper2Lib::PolyTree<IntegralType> ResultTree;
+		Clipper2Lib::Paths64 ResultOpenPaths;
+		const bool bExecuteResult = Clipper.Execute(InClipType, Clipper2Lib::FillRule::NonZero, ResultTree, ResultOpenPaths);
+
+		if (bExecuteResult)
+		{
+			ResultOut = ConvertPathsToUEPaths(ResultOpenPaths, InputBounds.Min, InputRange);
+		}
+
+		return bExecuteResult;
+	}
 };
 
 
@@ -113,6 +153,14 @@ bool PolygonsExclusiveOr(TArrayView<const FGeneralPolygon2d> SubjPolygons, TArra
 {
 	return Private::ClipArrays<double>(Clipper2Lib::ClipType::Xor, SubjPolygons, ClipPolygons, ResultOut);
 }
+bool PathsIntersection(const TArray<TArrayView<FVector2d>>& SubjPaths, TArrayView<const FGeneralPolygon2d> ClipPolygons, TArray<TArray<FVector2d>>& PathsOut)
+{
+	return Private::ClipPaths<double>(Clipper2Lib::ClipType::Intersection, SubjPaths, ClipPolygons, PathsOut);
+}
+bool PathsDifference(const TArray<TArrayView<FVector2d>>& SubjPaths, TArrayView<const FGeneralPolygon2d> ClipPolygons, TArray<TArray<FVector2d>>& PathsOut)
+{
+	return Private::ClipPaths<double>(Clipper2Lib::ClipType::Difference, SubjPaths, ClipPolygons, PathsOut);
+}
 } // namespace UE::Geometry
 
 
diff --git a/Engine/Plugins/Runtime/GeometryProcessing/Source/GeometryAlgorithms/Public/Curve/PolygonIntersectionUtils.h b/Engine/Plugins/Runtime/GeometryProcessing/Source/GeometryAlgorithms/Public/Curve/PolygonIntersectionUtils.h
--- a/Engine/Plugins/Runtime/GeometryProcessing/Source/GeometryAlgorithms/Public/Curve/PolygonIntersectionUtils.h
+++ b/Engine/Plugins/Runtime/GeometryProcessing/Source/GeometryAlgorithms/Public/Curve/PolygonIntersectionUtils.h
@@ -82,4 +82,10 @@ namespace UE::Geometry
 
 	/// Populate ResultOut with the 'exclusive or' of SubjPolygons with ClipPolygons
 	GEOMETRYALGORITHMS_API bool PolygonsExclusiveOr(TArrayView<const FGeneralPolygon2d> SubjPolygons, TArrayView<const FGeneralPolygon2d> ClipPolygons, TArray<FGeneralPolygon2d>& ResultOut);
+
+	/// Populate PathsOut with the intersected SubjPaths with ClipPolygons
+	GEOMETRYALGORITHMS_API bool PathsIntersection(const TArray<TArrayView<FVector2d>>& Paths, TArrayView<const FGeneralPolygon2d> ClipPolygons, TArray<TArray<FVector2d>>& PathsOut);
+
+	/// Populate PathsOut with the clipped SubjPaths with ClipPolygons
+	GEOMETRYALGORITHMS_API bool PathsDifference(const TArray<TArrayView<FVector2d>>& Paths, TArrayView<const FGeneralPolygon2d> ClipPolygons, TArray<TArray<FVector2d>>& PathsOut);
 }
