Remove the arbitrary target length option from the mesh cluster simplifier, as it breaks the assumptions of the simplifier & so can give bad results, and it wasn't currently in use.

diff --git a/Engine/Plugins/Runtime/GeometryProcessing/Source/DynamicMesh/Private/Operations/MeshClusterSimplifier.cpp b/Engine/Plugins/Runtime/GeometryProcessing/Source/DynamicMesh/Private/Operations/MeshClusterSimplifier.cpp
--- a/Engine/Plugins/Runtime/GeometryProcessing/Source/DynamicMesh/Private/Operations/MeshClusterSimplifier.cpp
+++ b/Engine/Plugins/Runtime/GeometryProcessing/Source/DynamicMesh/Private/Operations/MeshClusterSimplifier.cpp
@@ -160,19 +160,15 @@ namespace UE::MeshClusterSimplifyLocals
 	static TArray<float> InitTargetLengths(const MeshType& InMesh, const FSimplifyOptions& SimplifyOptions)
 	{
 		TArray<float> TargetLengths;
-		if (SimplifyOptions.OptionalTargetEdgeLengthScale)
-		{
-			TargetLengths.SetNumUninitialized(InMesh.MaxVertexID());
-			ParallelFor(InMesh.MaxVertexID(), [&InMesh, &TargetLengths, &SimplifyOptions](int32 VID)
-				{
-					if (!InMesh.IsVertex(VID))
-					{
-						return;
-					}
-					TargetLengths[VID] = float(SimplifyOptions.TargetEdgeLength * SimplifyOptions.OptionalTargetEdgeLengthScale(VID));
-				}
-			);
-		}
+		// Currently we always initialize to empty target lengths, which is treated as all vertices requesting SimplifyOptions.TargetEdgeLength.
+		// 
+		// Note that arbitrary target lengths may lead to bad results w/ the current algorithm. To support arbitrary lengths,
+		// the algorithm needs to prevent clusters from larger-target-length source vertices from running through or around those w/ 
+		// smaller target lengths. This is additionally complicated by the multi-pass strategy for handling constrained 
+		// edges. So for now, we just don't support arbitrary initial target lengths.
+		// 
+		// Note this is not as problematic for the lengths we currently automatically set when running 'preserve collapsed cluster' passes,
+		// because in those cases we tend to set locally consistent lengths in the regions where we locally re-run the clustering.
 		return TargetLengths;
 	}
 
diff --git a/Engine/Plugins/Runtime/GeometryProcessing/Source/DynamicMesh/Public/Operations/MeshClusterSimplifier.h b/Engine/Plugins/Runtime/GeometryProcessing/Source/DynamicMesh/Public/Operations/MeshClusterSimplifier.h
--- a/Engine/Plugins/Runtime/GeometryProcessing/Source/DynamicMesh/Public/Operations/MeshClusterSimplifier.h
+++ b/Engine/Plugins/Runtime/GeometryProcessing/Source/DynamicMesh/Public/Operations/MeshClusterSimplifier.h
@@ -22,10 +22,6 @@ struct FSimplifyOptions
 	// Vertices within this distance are allowed to be replaced with a single vertex
 	double TargetEdgeLength = 1;
 
-	// Optionally locally scales the simplifier's TargetEdgeLength on a per-vertex basis
-	// If provided, the target edge length at a vertex VID will be TargetEdgeLength * OptionalTargetEdgeLengthScale(VID)
-	TFunction<double(int32)> OptionalTargetEdgeLengthScale = nullptr;
-
 	// If > 0, number of passes to attempt to make sure the simplification does not fully delete islands of the mesh.
 	// In areas that would be deleted (because the region is smaller than the TargetEdgeLength),
 	// the target edge length will instead be locally reduced to half the local max in-cluster distance,
