FFD Lattice: use the input initial bounds rather than transforming them to the origin

diff --git a/Engine/Plugins/Runtime/GeometryProcessing/Source/DynamicMesh/Private/Operations/FFDLattice.cpp b/Engine/Plugins/Runtime/GeometryProcessing/Source/DynamicMesh/Private/Operations/FFDLattice.cpp
--- a/Engine/Plugins/Runtime/GeometryProcessing/Source/DynamicMesh/Private/Operations/FFDLattice.cpp
+++ b/Engine/Plugins/Runtime/GeometryProcessing/Source/DynamicMesh/Private/Operations/FFDLattice.cpp
@@ -70,14 +70,8 @@ FFFDLattice::FFFDLattice(const FVector3i& InDims, const FDynamicMesh3& Mesh, flo
 		InitialBounds = FAxisAlignedBox3d(FVector3d::Zero(), FVector3d::Zero());
 	}
 
-	// InitLattice() moves InitialBounds so that the min is at the origin, then applies the Padding parameter, so save the initial box center here
-	const FVector3d InputCenter = InitialBounds.Center();
-
 	InitLattice(Padding);
 
-	const FVector3d NewExtents = InitialBounds.Extents();
-	Transform.SetTranslation(InputCenter - NewExtents);
-
 	ComputeInitialEmbedding(Mesh);
 }
 
@@ -108,8 +102,8 @@ void FFFDLattice::InitLattice(float Padding)
 	const float ClampedPadding = FMath::Clamp(Padding, 0.0f, 5.0f);
 	const double MaxDiagonal = MaxElement(InitialBounds.Diagonal());
 
-	InitialBounds.Min = FVector3d(-0.5 * ClampedPadding * MaxDiagonal);
-	InitialBounds.Max = 2.0 * Extents + 0.5 * ClampedPadding * MaxDiagonal;
+	InitialBounds.Min -= FVector3d(0.5 * ClampedPadding * MaxDiagonal);
+	InitialBounds.Max += FVector3d(0.5 * ClampedPadding * MaxDiagonal);
 
 	const FVector3d Diag = InitialBounds.Diagonal();
 	CellSize = Diag / FVector3d(Dimensions - 1);
@@ -131,7 +125,7 @@ void FFFDLattice::GenerateInitialLatticePositions(TArray<FVector3d>& OutLatticeP
 				int PointID = ControlPointIndexFromCoordinates(i, j, k);
 				double Z = CellSize.Z * k;
 
-				OutLatticePositions[PointID] = Transform.TransformPosition(FVector3d{ X,Y,Z });
+				OutLatticePositions[PointID] = Transform.TransformPosition(InitialBounds.Min + FVector3d{ X,Y,Z });
 			}
 		}
 	}
@@ -173,7 +167,7 @@ void FFFDLattice::GenerateLatticeEdges(TArray<FVector2i>& OutLatticeEdges) const
 
 FVector3d FFFDLattice::ComputeTrilinearWeights(const FVector3d& QueryPoint, FVector3i& GridCoordinates) const
 {
-	const FVector3d GridPoint = QueryPoint / CellSize;
+	const FVector3d GridPoint = (QueryPoint - InitialBounds.Min) / CellSize;
 	GridCoordinates = FVector3i(GridPoint);
 	FVector3d Weights = GridPoint - FVector3d(GridCoordinates);
 
