Group Topology: ensure that the mesh and the current selection are synced

- this can happen if the mesh topology has been modified, but the selection topology hasn't

#rnx

diff --git a/Engine/Plugins/Runtime/GeometryProcessing/Source/DynamicMesh/Private/GroupTopology.cpp b/Engine/Plugins/Runtime/GeometryProcessing/Source/DynamicMesh/Private/GroupTopology.cpp
--- a/Engine/Plugins/Runtime/GeometryProcessing/Source/DynamicMesh/Private/GroupTopology.cpp
+++ b/Engine/Plugins/Runtime/GeometryProcessing/Source/DynamicMesh/Private/GroupTopology.cpp
@@ -970,18 +970,32 @@ FAxisAlignedBox3d FGroupTopology::GetSelectionBounds(
 	}
 
 	FAxisAlignedBox3d Bounds = FAxisAlignedBox3d::Empty();
+	
+	auto ContainVertex = [MeshPtr = Mesh, &Bounds, TransformFunc](int VertexID)
+	{
+		if (MeshPtr->IsVertex(VertexID))
+		{
+			Bounds.Contain( TransformFunc(MeshPtr->GetVertex(VertexID)) );
+		}		
+	};
 
 	for (int32 CornerID : Selection.SelectedCornerIDs)
 	{
-		Bounds.Contain(TransformFunc(Mesh->GetVertex(GetCornerVertexID(CornerID))));
+		if (Corners.IsValidIndex(CornerID))
+		{
+			ContainVertex(GetCornerVertexID(CornerID));
+		}
 	}
 
 	for (int32 EdgeID : Selection.SelectedEdgeIDs)
 	{
-		const FGroupEdge& Edge = Edges[EdgeID];
-		for (int32 vid : Edge.Span.Vertices)
+		if (Edges.IsValidIndex(EdgeID))
 		{
-			Bounds.Contain(TransformFunc(Mesh->GetVertex(vid)));
+			const FGroupEdge& Edge = Edges[EdgeID];
+			for (int32 vid : Edge.Span.Vertices)
+			{
+				ContainVertex(vid);
+			}
 		}
 	}
 
@@ -989,10 +1003,13 @@ FAxisAlignedBox3d FGroupTopology::GetSelectionBounds(
 	{
 		for (int32 TriangleID : GetGroupTriangles(GroupID))
 		{
-			FIndex3i TriVerts = Mesh->GetTriangle(TriangleID);
-			Bounds.Contain(TransformFunc(Mesh->GetVertex(TriVerts.A)));
-			Bounds.Contain(TransformFunc(Mesh->GetVertex(TriVerts.B)));
-			Bounds.Contain(TransformFunc(Mesh->GetVertex(TriVerts.C)));
+			if (Mesh->IsTriangle(TriangleID))
+			{
+				const FIndex3i TriVerts = Mesh->GetTriangle(TriangleID);
+				ContainVertex(TriVerts.A);
+				ContainVertex(TriVerts.B);
+				ContainVertex(TriVerts.C);
+			}
 		}
 	}
 
