Use the cluster simplifier for convex decomposition's pre-pass simplifier (since the main use case is to make the decomposition faster)

diff --git a/Engine/Plugins/Runtime/GeometryProcessing/Source/DynamicMesh/Private/ShapeApproximation/MeshSimpleShapeApproximation.cpp b/Engine/Plugins/Runtime/GeometryProcessing/Source/DynamicMesh/Private/ShapeApproximation/MeshSimpleShapeApproximation.cpp
--- a/Engine/Plugins/Runtime/GeometryProcessing/Source/DynamicMesh/Private/ShapeApproximation/MeshSimpleShapeApproximation.cpp
+++ b/Engine/Plugins/Runtime/GeometryProcessing/Source/DynamicMesh/Private/ShapeApproximation/MeshSimpleShapeApproximation.cpp
@@ -11,10 +11,10 @@
 #include "Implicit/SweepingMeshSDF.h"
 #include "ShapeApproximation/ShapeDetection3.h"
 #include "MeshQueries.h"
+#include "Operations/MeshClusterSimplifier.h"
 #include "Operations/MeshConvexHull.h"
 #include "Operations/MeshProjectionHull.h"
 #include "Util/ProgressCancel.h"
-#include "MeshSimplification.h"
 
 #define LOCTEXT_NAMESPACE "MeshSimpleShapeApproximation"
 
@@ -374,9 +374,15 @@ void FMeshSimpleShapeApproximation::Generate_ConvexHullDecompositions(FSimpleSha
 			if (bDecompositionPreSimplifyWithEdgeLength)
 			{
 				// Run pre-simplification to the target edge length
-				UE::Geometry::FVolPresMeshSimplification Simplifier(&ProcessMesh);
-				Simplifier.CollapseMode = UE::Geometry::FVolPresMeshSimplification::ESimplificationCollapseModes::MinimalExistingVertexError;
-				Simplifier.SimplifyToEdgeLength(DecompositionPreSimplifyEdgeLength);
+				MeshClusterSimplify::FSimplifyOptions SimplifyOptions;
+				SimplifyOptions.TargetEdgeLength = DecompositionPreSimplifyEdgeLength;
+				// Only consider shape for the decomposition -- do not constrain seams, do not transfer attributes/groups
+				SimplifyOptions.bTransferAttributes = false;
+				SimplifyOptions.bTransferGroups = false;
+				SimplifyOptions.PreserveEdges.SetSeamConstraints(MeshClusterSimplify::FSimplifyOptions::EConstraintLevel::Free);
+				FDynamicMesh3 ResultMesh;
+				MeshClusterSimplify::Simplify(ProcessMesh, ResultMesh, SimplifyOptions);
+				ProcessMesh = ResultMesh;
 			}
 		};
 		FConvexDecomposition3 Decomposition(SourceMesh, PreprocessOptions);
