move dynamic mesh mikkt support out of the calculate tangents op so it can be used by geometry script + fix handling of unset elements (normals, uvs) in the mesh overlay

diff --git a/Engine/Plugins/Runtime/GeometryProcessing/Source/GeometryAlgorithms/GeometryAlgorithms.Build.cs b/Engine/Plugins/Runtime/GeometryProcessing/Source/GeometryAlgorithms/GeometryAlgorithms.Build.cs
--- a/Engine/Plugins/Runtime/GeometryProcessing/Source/GeometryAlgorithms/GeometryAlgorithms.Build.cs
+++ b/Engine/Plugins/Runtime/GeometryProcessing/Source/GeometryAlgorithms/GeometryAlgorithms.Build.cs
@@ -40,5 +40,7 @@ public class GeometryAlgorithms : ModuleRules
 		}
 
 		bDisableAutoRTFMInstrumentation = true;
+
+		AddEngineThirdPartyPrivateStaticDependencies(Target, "MikkTSpace");
 	}
 }
diff --git a/Engine/Plugins/Runtime/GeometryProcessing/Source/GeometryAlgorithms/Private/DynamicMeshMikkTWrapper.cpp b/Engine/Plugins/Runtime/GeometryProcessing/Source/GeometryAlgorithms/Private/DynamicMeshMikkTWrapper.cpp
new file mode 100644
--- /dev/null
+++ b/Engine/Plugins/Runtime/GeometryProcessing/Source/GeometryAlgorithms/Private/DynamicMeshMikkTWrapper.cpp
@@ -0,0 +1,169 @@
+// Copyright Epic Games, Inc. All Rights Reserved.
+
+
+#include "DynamicMeshMikkTWrapper.h"
+
+#include "DynamicMesh/DynamicMesh3.h"
+
+// on platforms that build the mikkt library (win/mac/linux) its build.cs will define this:
+#if WITH_MIKKTSPACE
+#include "mikktspace.h"
+#endif //WITH_MIKKTSPACE
+
+#if WITH_MIKKTSPACE
+namespace DynamicMeshMikkTWrapper
+{
+	using namespace UE::Geometry;
+
+	namespace DynamicMeshMikkTInterface
+	{
+		struct FDynamicMeshInfo
+		{
+			const FDynamicMesh3* Mesh;
+			const FDynamicMeshNormalOverlay* NormalOverlay;
+			const FDynamicMeshUVOverlay* UVOverlay;
+			FMeshTangentsd* TangentsOut;
+		};
+
+
+		int MikkGetNumFaces(const SMikkTSpaceContext* Context)
+		{
+			const FDynamicMeshInfo* MeshInfo = (const FDynamicMeshInfo*)(Context->m_pUserData);
+			return MeshInfo->Mesh->TriangleCount();
+		}
+
+		int MikkGetNumVertsOfFace(const SMikkTSpaceContext* Context, const int FaceIdx)
+		{
+			const FDynamicMeshInfo* MeshInfo = (const FDynamicMeshInfo*)(Context->m_pUserData);
+			return MeshInfo->Mesh->IsTriangle(FaceIdx) ? 3 : 0;
+		}
+
+		void MikkGetPosition(const SMikkTSpaceContext* Context, float Position[3], const int FaceIdx, const int VertIdx)
+		{
+			const FDynamicMeshInfo* MeshInfo = (const FDynamicMeshInfo*)(Context->m_pUserData);
+			FIndex3i Triangle = MeshInfo->Mesh->GetTriangle(FaceIdx);
+			FVector3d VertexPos = MeshInfo->Mesh->GetVertex(Triangle[VertIdx]);
+			Position[0] = (float)VertexPos.X;
+			Position[1] = (float)VertexPos.Y;
+			Position[2] = (float)VertexPos.Z;
+		}
+
+		void MikkGetNormal(const SMikkTSpaceContext* Context, float Normal[3], const int FaceIdx, const int VertIdx)
+		{
+			const FDynamicMeshInfo* MeshInfo = (const FDynamicMeshInfo*)(Context->m_pUserData);
+			FIndex3i NormTriangle = MeshInfo->NormalOverlay->GetTriangle(FaceIdx);
+			FVector3f Normalf;
+			if (NormTriangle.A == IndexConstants::InvalidID)
+			{
+				Normalf = FVector3f::ZAxisVector;
+			}
+			else
+			{
+				MeshInfo->NormalOverlay->GetElement(NormTriangle[VertIdx], Normalf);
+			}
+			Normal[0] = Normalf.X;
+			Normal[1] = Normalf.Y;
+			Normal[2] = Normalf.Z;
+		}
+
+		void MikkGetTexCoord(const SMikkTSpaceContext* Context, float UV[2], const int FaceIdx, const int VertIdx)
+		{
+			const FDynamicMeshInfo* MeshInfo = (const FDynamicMeshInfo*)(Context->m_pUserData);
+			FIndex3i UVTriangle = MeshInfo->UVOverlay->GetTriangle(FaceIdx);
+			FVector2f UVf;
+			if (UVTriangle.A == IndexConstants::InvalidID)
+			{
+				UV[0] = UV[1] = 0.f;
+			}
+			else
+			{
+				MeshInfo->UVOverlay->GetElement(UVTriangle[VertIdx], UVf);
+				UV[0] = UVf.X;
+				UV[1] = UVf.Y;
+			}
+		}
+
+		void MikkSetTSpaceBasic(const SMikkTSpaceContext* Context, const float Tangent[3], const float BitangentSign, const int FaceIdx, const int VertIdx)
+		{
+			FDynamicMeshInfo* MeshInfo = (FDynamicMeshInfo*)(Context->m_pUserData);
+			FVector3d Tangentd(Tangent[0], Tangent[1], Tangent[2]);
+			FIndex3i NormTriangle = MeshInfo->NormalOverlay->GetTriangle(FaceIdx);
+			FVector3d Normald = (NormTriangle.A != IndexConstants::InvalidID) ? (FVector3d)MeshInfo->NormalOverlay->GetElement(NormTriangle[VertIdx]) : FVector3d::ZAxisVector;
+			FVector3d Bitangentd = VectorUtil::Bitangent(Normald, Tangentd, -(double)BitangentSign);
+
+			MeshInfo->TangentsOut->SetPerTriangleTangent(FaceIdx, VertIdx, Tangentd, Bitangentd);
+		}
+
+
+		void MikkSetTSpace(const SMikkTSpaceContext* Context, const float Tangent[3], const float BiTangent[3],
+			const float MagS, const float MagT, const tbool bIsOrientationPreserving,
+			const int FaceIdx, const int VertIdx)
+		{
+			FDynamicMeshInfo* MeshInfo = (FDynamicMeshInfo*)(Context->m_pUserData);
+			FVector3d Tangentd(Tangent[0], Tangent[1], Tangent[2]);
+			FVector3d Bitangentd(BiTangent[0], BiTangent[1], BiTangent[2]);
+			MeshInfo->TangentsOut->SetPerTriangleTangent(FaceIdx, VertIdx, Tangentd, Bitangentd);
+		}
+
+	}
+
+	bool IsSupported()
+	{
+		return true;
+	}
+	bool ComputeTangents(FMeshTangentsd& MeshTangents, int32 TargetUVLayer)
+	{
+		if (!MeshTangents.GetMesh())
+		{
+			return false;
+		}
+
+		const FDynamicMesh3& SourceMesh = *MeshTangents.GetMesh();
+		if (!SourceMesh.HasAttributes() || SourceMesh.Attributes()->NumNormalLayers() == 0 || TargetUVLayer >= SourceMesh.Attributes()->NumUVLayers())
+		{
+			return false;
+		}
+
+		MeshTangents.InitializeTriVertexTangents(true);
+
+		DynamicMeshMikkTInterface::FDynamicMeshInfo MeshInfo;
+		MeshInfo.Mesh = &SourceMesh;
+		MeshInfo.NormalOverlay = SourceMesh.Attributes()->PrimaryNormals();
+		MeshInfo.UVOverlay = SourceMesh.Attributes()->GetUVLayer(TargetUVLayer);
+		MeshInfo.TangentsOut = &MeshTangents;
+
+		SMikkTSpaceInterface MikkTInterface;
+		MikkTInterface.m_getNormal = DynamicMeshMikkTInterface::MikkGetNormal;
+		MikkTInterface.m_getNumFaces = DynamicMeshMikkTInterface::MikkGetNumFaces;
+		MikkTInterface.m_getNumVerticesOfFace = DynamicMeshMikkTInterface::MikkGetNumVertsOfFace;
+		MikkTInterface.m_getPosition = DynamicMeshMikkTInterface::MikkGetPosition;
+		MikkTInterface.m_getTexCoord = DynamicMeshMikkTInterface::MikkGetTexCoord;
+		MikkTInterface.m_setTSpaceBasic = DynamicMeshMikkTInterface::MikkSetTSpaceBasic;
+		MikkTInterface.m_setTSpace = nullptr;
+
+		SMikkTSpaceContext MikkTContext;
+		MikkTContext.m_pInterface = &MikkTInterface;
+		MikkTContext.m_pUserData = (void*)(&MeshInfo);
+		MikkTContext.m_bIgnoreDegenerates = false;
+
+		// execute mikkt
+		genTangSpaceDefault(&MikkTContext);
+
+		return true;
+	}
+}
+#else
+
+namespace DynamicMeshMikkTWrapper
+{
+	bool IsSupported()
+	{
+		return false;
+	}
+	bool ComputeTangents(FMeshTangentsd& ComputedTangents, int32 UVLayer)
+	{
+		return false;
+	}
+}
+
+#endif
diff --git a/Engine/Plugins/Runtime/GeometryProcessing/Source/GeometryAlgorithms/Public/DynamicMeshMikkTWrapper.h b/Engine/Plugins/Runtime/GeometryProcessing/Source/GeometryAlgorithms/Public/DynamicMeshMikkTWrapper.h
new file mode 100644
--- /dev/null
+++ b/Engine/Plugins/Runtime/GeometryProcessing/Source/GeometryAlgorithms/Public/DynamicMeshMikkTWrapper.h
@@ -0,0 +1,21 @@
+// Copyright Epic Games, Inc. All Rights Reserved.
+
+#pragma once
+
+#include "DynamicMesh/MeshTangents.h"
+
+namespace DynamicMeshMikkTWrapper
+{	
+	/** @return true if the current platform supports MikkT tangent calculations */
+	bool GEOMETRYALGORITHMS_API IsSupported();
+
+	/**
+	 * Compute MikkT tangents for the given MeshTangents
+	 * 
+	 * @param MeshTangents Tangents to fill in. Must have a valid target mesh set.
+	 * @param UVLayer The UV layer tangents should reference
+	 * @return true if tangents were successfully computed (i.e. platform supported mikkt + requested mesh and UV layer were valid)
+	 */
+	bool GEOMETRYALGORITHMS_API ComputeTangents(UE::Geometry::FMeshTangentsd& MeshTangents, int32 UVLayer);
+}
+
diff --git a/Engine/Plugins/Runtime/GeometryScripting/Source/GeometryScriptingCore/Private/MeshNormalsFunctions.cpp b/Engine/Plugins/Runtime/GeometryScripting/Source/GeometryScriptingCore/Private/MeshNormalsFunctions.cpp
--- a/Engine/Plugins/Runtime/GeometryScripting/Source/GeometryScriptingCore/Private/MeshNormalsFunctions.cpp
+++ b/Engine/Plugins/Runtime/GeometryScripting/Source/GeometryScriptingCore/Private/MeshNormalsFunctions.cpp
@@ -5,6 +5,7 @@
 #include "DynamicMesh/DynamicMesh3.h"
 #include "DynamicMesh/MeshNormals.h"
 #include "DynamicMesh/MeshTangents.h"
+#include "DynamicMeshMikkTWrapper.h"
 #include "DynamicMeshEditor.h" // for sharp edge creation/deletion via overlay seam methods
 #include "Operations/RepairOrientation.h"
 #include "DynamicMesh/DynamicMeshAABBTree3.h"
@@ -472,6 +473,29 @@ UDynamicMesh* UGeometryScriptLibrary_MeshNormalsFunctions::ComputeTangents(
 			EditMesh.Attributes()->EnableTangents();
 		}
 
+		if (Options.Type == EGeometryScriptTangentTypes::StandardMikkT)
+		{
+			if (DynamicMeshMikkTWrapper::IsSupported())
+			{
+				FMeshTangentsd Tangents(&EditMesh);
+				if (DynamicMeshMikkTWrapper::ComputeTangents(Tangents, Options.UVLayer))
+				{
+					Tangents.CopyToOverlays(EditMesh);
+				}
+				else
+				{
+					UE::Geometry::AppendError(Debug, EGeometryScriptErrorType::OperationFailed, LOCTEXT("ComputeTangents_MikkTFailed", "ComputeTangents: Standard MikkT failed to compute a result"));
+				}
+				return;
+			}
+			else
+			{
+				UE::Geometry::AppendWarning(Debug, EGeometryScriptErrorType::InvalidInputs, LOCTEXT("ComputeTangents_StandardMikkTEditorOnly", "ComputeTangents: Standard MikkT not supported by current platform; falling back to Fast MikkT."));
+				Options.Type = EGeometryScriptTangentTypes::FastMikkT;
+				// continue through below to compute FastMikkT
+			}
+		}
+
 		FComputeTangentsOptions TangentOptions;
 		TangentOptions.bAveraged = (Options.Type == EGeometryScriptTangentTypes::FastMikkT);
 
diff --git a/Engine/Plugins/Runtime/GeometryScripting/Source/GeometryScriptingCore/Public/GeometryScript/MeshNormalsFunctions.h b/Engine/Plugins/Runtime/GeometryScripting/Source/GeometryScriptingCore/Public/GeometryScript/MeshNormalsFunctions.h
--- a/Engine/Plugins/Runtime/GeometryScripting/Source/GeometryScriptingCore/Public/GeometryScript/MeshNormalsFunctions.h
+++ b/Engine/Plugins/Runtime/GeometryScripting/Source/GeometryScriptingCore/Public/GeometryScript/MeshNormalsFunctions.h
@@ -49,7 +49,9 @@ UENUM(BlueprintType)
 enum class EGeometryScriptTangentTypes : uint8
 {
 	FastMikkT = 0,
-	PerTriangle = 1
+	PerTriangle = 1,
+	// Standard MikkT-algorithm tangents. Available for windows/mac/linux platforms; will fall back to FastMikkT if unavailable.
+	StandardMikkT
 };
 
 
diff --git a/Engine/Plugins/Runtime/MeshModelingToolset/Source/ModelingOperatorsEditorOnly/ModelingOperatorsEditorOnly.Build.cs b/Engine/Plugins/Runtime/MeshModelingToolset/Source/ModelingOperatorsEditorOnly/ModelingOperatorsEditorOnly.Build.cs
--- a/Engine/Plugins/Runtime/MeshModelingToolset/Source/ModelingOperatorsEditorOnly/ModelingOperatorsEditorOnly.Build.cs
+++ b/Engine/Plugins/Runtime/MeshModelingToolset/Source/ModelingOperatorsEditorOnly/ModelingOperatorsEditorOnly.Build.cs
@@ -17,7 +17,7 @@ public class ModelingOperatorsEditorOnly : ModuleRules
 				"GeometryCore",
 				"DynamicMesh",
 				"MeshConversion",
-				"GeometryAlgorithms", // required for constrained Delaunay triangulation
+				"GeometryAlgorithms", // required for constrained Delaunay triangulation, and MikkT
 				"MeshUtilitiesCommon", // required by uvlayoutop
 				// ... add other public dependencies that you statically link with here ...
 			}
diff --git a/Engine/Plugins/Runtime/MeshModelingToolset/Source/ModelingOperatorsEditorOnly/Private/ParameterizationOps/CalculateTangentsOp.cpp b/Engine/Plugins/Runtime/MeshModelingToolset/Source/ModelingOperatorsEditorOnly/Private/ParameterizationOps/CalculateTangentsOp.cpp
--- a/Engine/Plugins/Runtime/MeshModelingToolset/Source/ModelingOperatorsEditorOnly/Private/ParameterizationOps/CalculateTangentsOp.cpp
+++ b/Engine/Plugins/Runtime/MeshModelingToolset/Source/ModelingOperatorsEditorOnly/Private/ParameterizationOps/CalculateTangentsOp.cpp
@@ -1,16 +1,11 @@
 // Copyright Epic Games, Inc. All Rights Reserved.
 
 #include "ParameterizationOps/CalculateTangentsOp.h"
+#include "DynamicMeshMikkTWrapper.h"
 #include "VectorUtil.h"
 
 #include UE_INLINE_GENERATED_CPP_BY_NAME(CalculateTangentsOp)
 
-#define WITH_MIKKTSPACE 1
-
-#if WITH_MIKKTSPACE
-#include "mikktspace.h"
-#endif //WITH_MIKKTSPACE
-
 using namespace UE::Geometry;
 
 void FCalculateTangentsOp::CalculateResult(FProgressCancel* Progress)
@@ -96,124 +91,7 @@ void FCalculateTangentsOp::CopyFromSource(FProgressCancel* Progress, TUniquePtr<
 }
 
 
-
-#if WITH_MIKKTSPACE
-namespace DynamicMeshMikkTInterface
-{
-	struct FDynamicMeshInfo
-	{
-		const FDynamicMesh3* Mesh;
-		const FDynamicMeshNormalOverlay* NormalOverlay;
-		const FDynamicMeshUVOverlay* UVOverlay;
-		FMeshTangentsd* TangentsOut;
-	};
-
-
-	int MikkGetNumFaces(const SMikkTSpaceContext* Context)
-	{
-		const FDynamicMeshInfo* MeshInfo = (const FDynamicMeshInfo*)(Context->m_pUserData);
-		return MeshInfo->Mesh->TriangleCount();
-	}
-
-	int MikkGetNumVertsOfFace(const SMikkTSpaceContext* Context, const int FaceIdx)
-	{
-		const FDynamicMeshInfo* MeshInfo = (const FDynamicMeshInfo*)(Context->m_pUserData);
-		return MeshInfo->Mesh->IsTriangle(FaceIdx) ? 3 : 0;
-	}
-
-	void MikkGetPosition(const SMikkTSpaceContext* Context, float Position[3], const int FaceIdx, const int VertIdx)
-	{
-		const FDynamicMeshInfo* MeshInfo = (const FDynamicMeshInfo*)(Context->m_pUserData);
-		FIndex3i Triangle = MeshInfo->Mesh->GetTriangle(FaceIdx);
-		FVector3d VertexPos = MeshInfo->Mesh->GetVertex(Triangle[VertIdx]);
-		Position[0] = (float)VertexPos.X;
-		Position[1] = (float)VertexPos.Y;
-		Position[2] = (float)VertexPos.Z;
-	}
-
-	void MikkGetNormal(const SMikkTSpaceContext* Context, float Normal[3], const int FaceIdx, const int VertIdx)
-	{
-		const FDynamicMeshInfo* MeshInfo = (const FDynamicMeshInfo*)(Context->m_pUserData);
-		FIndex3i NormTriangle = MeshInfo->NormalOverlay->GetTriangle(FaceIdx);
-		FVector3f Normalf;
-		MeshInfo->NormalOverlay->GetElement(NormTriangle[VertIdx], Normalf);
-		Normal[0] = Normalf.X;
-		Normal[1] = Normalf.Y;
-		Normal[2] = Normalf.Z;
-	}
-
-	void MikkGetTexCoord(const SMikkTSpaceContext* Context, float UV[2], const int FaceIdx, const int VertIdx)
-	{
-		const FDynamicMeshInfo* MeshInfo = (const FDynamicMeshInfo*)(Context->m_pUserData);
-		FIndex3i UVTriangle = MeshInfo->UVOverlay->GetTriangle(FaceIdx);
-		FVector2f UVf;
-		MeshInfo->UVOverlay->GetElement(UVTriangle[VertIdx], UVf);
-		UV[0] = UVf.X;
-		UV[1] = UVf.Y;
-	}
-
-	void MikkSetTSpaceBasic(const SMikkTSpaceContext* Context, const float Tangent[3], const float BitangentSign, const int FaceIdx, const int VertIdx)
-	{
-		FDynamicMeshInfo* MeshInfo = (FDynamicMeshInfo*)(Context->m_pUserData);
-		FVector3d Tangentd(Tangent[0], Tangent[1], Tangent[2]);
-
-		FIndex3i NormTriangle = MeshInfo->NormalOverlay->GetTriangle(FaceIdx);
-		FVector3d Normald = (FVector3d)MeshInfo->NormalOverlay->GetElement(NormTriangle[VertIdx]);
-		FVector3d Bitangentd = VectorUtil::Bitangent(Normald, Tangentd, -(double)BitangentSign);
-
-		MeshInfo->TangentsOut->SetPerTriangleTangent(FaceIdx, VertIdx, Tangentd, Bitangentd);
-	}
-
-
-	void MikkSetTSpac(const SMikkTSpaceContext* Context, const float Tangent[3], const float BiTangent[3],
-		const float MagS, const float MagT, const tbool bIsOrientationPreserving,
-		const int FaceIdx, const int VertIdx)
-	{
-		FDynamicMeshInfo* MeshInfo = (FDynamicMeshInfo*)(Context->m_pUserData);
-		FVector3d Tangentd(Tangent[0], Tangent[1], Tangent[2]);
-		FVector3d Bitangentd(BiTangent[0], BiTangent[1], BiTangent[2]);
-		MeshInfo->TangentsOut->SetPerTriangleTangent(FaceIdx, VertIdx, Tangentd, Bitangentd);
-	}
-
-}
-
-
-
 void FCalculateTangentsOp::CalculateMikkT(FProgressCancel* Progress, TUniquePtr<FMeshTangentsd>& Tangents)
 {
-	Tangents->InitializeTriVertexTangents(true);
-
-	DynamicMeshMikkTInterface::FDynamicMeshInfo MeshInfo;
-	MeshInfo.Mesh = SourceMesh.Get();
-	MeshInfo.NormalOverlay = SourceMesh->Attributes()->PrimaryNormals();
-	MeshInfo.UVOverlay = SourceMesh->Attributes()->GetUVLayer(TargetUVLayer);
-	MeshInfo.TangentsOut = Tangents.Get();
-
-	SMikkTSpaceInterface MikkTInterface;
-	MikkTInterface.m_getNormal = DynamicMeshMikkTInterface::MikkGetNormal;
-	MikkTInterface.m_getNumFaces = DynamicMeshMikkTInterface::MikkGetNumFaces;
-	MikkTInterface.m_getNumVerticesOfFace = DynamicMeshMikkTInterface::MikkGetNumVertsOfFace;
-	MikkTInterface.m_getPosition = DynamicMeshMikkTInterface::MikkGetPosition;
-	MikkTInterface.m_getTexCoord = DynamicMeshMikkTInterface::MikkGetTexCoord;
-	MikkTInterface.m_setTSpaceBasic = DynamicMeshMikkTInterface::MikkSetTSpaceBasic;
-	MikkTInterface.m_setTSpace = nullptr;
-
-	SMikkTSpaceContext MikkTContext;
-	MikkTContext.m_pInterface = &MikkTInterface;
-	MikkTContext.m_pUserData = (void*)(&MeshInfo);
-	MikkTContext.m_bIgnoreDegenerates = false;
-	//MikkTContext.m_bIgnoreDegenerates = true;
-	//MikkTContext.m_bIgnoreDegenerates = bIgnoreDegenerateTriangles;
-
-	// execute mikkt
-	genTangSpaceDefault(&MikkTContext);
+	DynamicMeshMikkTWrapper::ComputeTangents(*Tangents, TargetUVLayer);
 }
-
-#else
-
-void FCalculateTangentsOp::CalculateMikkT(FProgressCancel* Progress, TUniquePtr<FMeshTangentsd>& Tangents)
-{
-	check(false);
-}
-
-#endif
diff --git a/Engine/Source/Runtime/GeometryCore/Public/DynamicMesh/MeshTangents.h b/Engine/Source/Runtime/GeometryCore/Public/DynamicMesh/MeshTangents.h
--- a/Engine/Source/Runtime/GeometryCore/Public/DynamicMesh/MeshTangents.h
+++ b/Engine/Source/Runtime/GeometryCore/Public/DynamicMesh/MeshTangents.h
@@ -106,6 +106,11 @@ public:
 
 	const TArray<int32>& GetDegenerateTris() const { return AllDegenerateTris; }
 
+	// Get the target mesh
+	const FDynamicMesh3* GetMesh() const
+	{
+		return Mesh; 
+	}
 
 	//
 	// Per-Triangle-Vertex Tangents
