fix trianglecount -> max triangle ID in dynamic mesh mikkt wrapper -- could result in non-compact meshes missing tangents for triangles w/ IDs > triangle count
+ guard against negative uv layer being requested

diff --git a/Engine/Plugins/Runtime/GeometryProcessing/Source/GeometryAlgorithms/Private/DynamicMeshMikkTWrapper.cpp b/Engine/Plugins/Runtime/GeometryProcessing/Source/GeometryAlgorithms/Private/DynamicMeshMikkTWrapper.cpp
--- a/Engine/Plugins/Runtime/GeometryProcessing/Source/GeometryAlgorithms/Private/DynamicMeshMikkTWrapper.cpp
+++ b/Engine/Plugins/Runtime/GeometryProcessing/Source/GeometryAlgorithms/Private/DynamicMeshMikkTWrapper.cpp
@@ -29,7 +29,7 @@ namespace DynamicMeshMikkTWrapper
 		int MikkGetNumFaces(const SMikkTSpaceContext* Context)
 		{
 			const FDynamicMeshInfo* MeshInfo = (const FDynamicMeshInfo*)(Context->m_pUserData);
-			return MeshInfo->Mesh->TriangleCount();
+			return MeshInfo->Mesh->MaxTriangleID();
 		}
 
 		int MikkGetNumVertsOfFace(const SMikkTSpaceContext* Context, const int FaceIdx)
@@ -119,7 +119,7 @@ namespace DynamicMeshMikkTWrapper
 		}
 
 		const FDynamicMesh3& SourceMesh = *MeshTangents.GetMesh();
-		if (!SourceMesh.HasAttributes() || SourceMesh.Attributes()->NumNormalLayers() == 0 || TargetUVLayer >= SourceMesh.Attributes()->NumUVLayers())
+		if (!SourceMesh.HasAttributes() || SourceMesh.Attributes()->NumNormalLayers() == 0 || TargetUVLayer < 0 || TargetUVLayer >= SourceMesh.Attributes()->NumUVLayers())
 		{
 			return false;
 		}
