guard against failed factorization in the geometry sparse solver wrappers

diff --git a/Engine/Plugins/Runtime/GeometryProcessing/Source/DynamicMesh/Private/Solvers/Internal/ConstrainedPoissonSolver.h b/Engine/Plugins/Runtime/GeometryProcessing/Source/DynamicMesh/Private/Solvers/Internal/ConstrainedPoissonSolver.h
--- a/Engine/Plugins/Runtime/GeometryProcessing/Source/DynamicMesh/Private/Solvers/Internal/ConstrainedPoissonSolver.h
+++ b/Engine/Plugins/Runtime/GeometryProcessing/Source/DynamicMesh/Private/Solvers/Internal/ConstrainedPoissonSolver.h
@@ -199,7 +199,13 @@ public:
 	bool Solve(const FSOAPositions& SourceVector, FSOAPositions& SolutionVector) const
 	{
 		checkSlow(bMatrixSolverDirty == false);
-		FConstrainedSolverTimeLogger Timmer(TEXT("Post-setup solve time"));
+
+		if (MatrixSolver->bSucceeded() == false)  // check that the factorization succeeded 
+		{
+			return false;
+		}
+
+		FConstrainedSolverTimeLogger Timer(TEXT("Post-setup solve time"));
 
 		FSparseMatrixD& SymmetricMatrix = *SymmetricMatrixPtr;
 		// Set up the source vector
@@ -224,8 +230,13 @@ public:
 	bool Solve(FSOAPositions& SolutionVector) const
 	{
 		checkSlow(bMatrixSolverDirty == false);
-		
-		FConstrainedSolverTimeLogger Timmer(TEXT("Post-setup solve time"));
+
+		if (MatrixSolver->bSucceeded() == false)  // check that the factorization succeeded 
+		{
+			return false;
+		}
+
+		FConstrainedSolverTimeLogger Timer(TEXT("Post-setup solve time"));
 
 		FSparseMatrixD& SymmetricMatrix = *SymmetricMatrixPtr;
 		FSOAPositions RHSVector((int32)SymmetricMatrix.cols());
diff --git a/Engine/Plugins/Runtime/GeometryProcessing/Source/DynamicMesh/Private/Solvers/Internal/MatrixSolver.h b/Engine/Plugins/Runtime/GeometryProcessing/Source/DynamicMesh/Private/Solvers/Internal/MatrixSolver.h
--- a/Engine/Plugins/Runtime/GeometryProcessing/Source/DynamicMesh/Private/Solvers/Internal/MatrixSolver.h
+++ b/Engine/Plugins/Runtime/GeometryProcessing/Source/DynamicMesh/Private/Solvers/Internal/MatrixSolver.h
@@ -117,7 +117,8 @@ private:
 		MatrixSolver.analyzePattern(SparseMatrix);
 		MatrixSolver.factorize(SparseMatrix);
 
-		bSetup = true;	
+		// if the factorize succeeded, info() status should return 'success'
+		bSetup = (MatrixSolver.info() == Eigen::ComputationInfo::Success);
 	}
 private:
 
@@ -251,7 +252,7 @@ private:
 		MatrixSolver.analyzePattern(SparseMatrix);
 		MatrixSolver.factorize(SparseMatrix);
 
-		bSetup = true;
+		bSetup = (MatrixSolver.info() == Eigen::ComputationInfo::Success);
 
 		int32 n = Eigen::nbThreads();
 	}
diff --git a/Engine/Plugins/Runtime/GeometryProcessing/Source/DynamicMesh/Private/Solvers/Internal/MeshDiffusionIntegrator.cpp b/Engine/Plugins/Runtime/GeometryProcessing/Source/DynamicMesh/Private/Solvers/Internal/MeshDiffusionIntegrator.cpp
--- a/Engine/Plugins/Runtime/GeometryProcessing/Source/DynamicMesh/Private/Solvers/Internal/MeshDiffusionIntegrator.cpp
+++ b/Engine/Plugins/Runtime/GeometryProcessing/Source/DynamicMesh/Private/Solvers/Internal/MeshDiffusionIntegrator.cpp
@@ -136,6 +136,11 @@ void FMeshDiffusionIntegrator::Integrate_BackwardEuler(const EMatrixSolverType M
 
 	MatrixSolver->SetUp(SparseMatrix, bIsSymmetric);
 
+	if (MatrixSolver->bSucceeded() == false)  // check that the factorization succeeded 
+	{
+		return;
+	}
+
 	// We are going to solve the system 
 	FSOAPositions Source(InternalVertexCount);
 
