Fixes for Mesh Bevel --
1. Handle bevels where the bevel edge is not a group boundary edge
2. Detect + handle fallback case where a non-subdivided quad is used so there is no quad strip to process
3. Stop processing earlier when encountering bad caes (invalid edge IDs, or non-boundary edge IDs, in the stitch-up pass)

diff --git a/Engine/Plugins/Runtime/GeometryProcessing/Source/DynamicMesh/Private/Operations/MeshBevel.cpp b/Engine/Plugins/Runtime/GeometryProcessing/Source/DynamicMesh/Private/Operations/MeshBevel.cpp
--- a/Engine/Plugins/Runtime/GeometryProcessing/Source/DynamicMesh/Private/Operations/MeshBevel.cpp
+++ b/Engine/Plugins/Runtime/GeometryProcessing/Source/DynamicMesh/Private/Operations/MeshBevel.cpp
@@ -961,109 +961,156 @@ void FMeshBevel::BuildTerminatorVertex(FBevelVertex& Vertex, const FDynamicMesh3
 	// "other" split edge that we pick
 	int32 RingSplitEdgeID = -1;
 
-	// NOTE: code below assumes we have polygroups on the mesh. In theory we could also support not having polygroups,
-	// by (eg) picking an arbitrary edge "furthest" from IncomingEdgeID in the exponential map
+	// NOTE: code below follows polygroups on the mesh when available; otherwise, just tries to pick a 'furthest' edge
 
 	// Find the ordered set of triangles that are not in either of the groups connected to IncomingEdgeID. Eg imagine
 	// a cube corner, if we want to bevel IncomingEdgeID along one of the cube edges, we want to add the new edge
 	// in the "furthest" face (perpendicular to the edge)
 	FIndex2i IncomingEdgeT = Mesh.GetEdgeT(IncomingEdgeID);
-	FIndex2i IncomingEdgeGroups(Mesh.GetTriangleGroup(IncomingEdgeT.A), Mesh.GetTriangleGroup(IncomingEdgeT.B));
 
-	// Find the first index of Vertex.SortedTriangles that we want to remove.
-	// This should not be able to fail but if it does, we will just start at 0 
-	// and maybe end up with some bevel failures below
-	int32 NumTriangles = Vertex.SortedTriangles.Num();
-	int32 StartIndex = 0;
-	for (int32 k = 0; k < NumTriangles; ++k)
+	// Logic to find the RingSplitEdgeID (in its own control block just to clarify local variables logical scope)
 	{
-		int32 gid = Mesh.GetTriangleGroup(Vertex.SortedTriangles[k]);
-		if (IncomingEdgeGroups.Contains(gid))
+		FIndex2i IncomingEdgeGroups(Mesh.GetTriangleGroup(IncomingEdgeT.A), Mesh.GetTriangleGroup(IncomingEdgeT.B));
+
+		// Find the first index of Vertex.SortedTriangles that we want to remove.
+		// This should not be able to fail but if it does, we will just start at 0 
+		// and maybe end up with some bevel failures below
+		int32 NumTriangles = Vertex.SortedTriangles.Num();
+		int32 StartIndex = 0;
+		for (int32 k = 0; k < NumTriangles; ++k)
 		{
-			StartIndex = k;
-			break;
+			int32 gid = Mesh.GetTriangleGroup(Vertex.SortedTriangles[k]);
+			if (IncomingEdgeGroups.Contains(gid))
+			{
+				StartIndex = k;
+				break;
+			}
 		}
-	}
 
-	TArray<int32> OtherGroupTris;	// sorted wedge of triangles that are not in either of the groups connected to incoming edge
-	TArray<int32> OtherGroups;		// list of group IDs encountered, in-order
-	for ( int32 k = 0; k < NumTriangles; ++k )
-	{
-		// Vertex.SortedTriangles was ordered, ie sequential triangles were connected, but if we filter
-		// out some triangles it may no longer be sequential in OtherGroupTris, leading to badness. 
-		// But the two groups we are removing should be contiguous, so if we start there, then the
-		// remaining tris should be contiguous. StartIndex found above should give us that triangle.
-		int32 ShiftedIndex = (StartIndex + k) % NumTriangles;
-		int32 tid = Vertex.SortedTriangles[ShiftedIndex];
-		int32 gid = Mesh.GetTriangleGroup(tid);
-		if (IncomingEdgeGroups.Contains(gid) == false)
+		TArray<int32> OtherGroupTris;	// sorted wedge of triangles that are not in either of the groups connected to incoming edge
+		TArray<int32> OtherGroups;		// list of group IDs encountered, in-order
+		for (int32 k = 0; k < NumTriangles; ++k)
 		{
-			OtherGroupTris.Add(tid);
-			OtherGroups.AddUnique(gid);
+			// Vertex.SortedTriangles was ordered, ie sequential triangles were connected, but if we filter
+			// out some triangles it may no longer be sequential in OtherGroupTris, leading to badness. 
+			// But the two groups we are removing should be contiguous, so if we start there, then the
+			// remaining tris should be contiguous. StartIndex found above should give us that triangle.
+			int32 ShiftedIndex = (StartIndex + k) % NumTriangles;
+			int32 tid = Vertex.SortedTriangles[ShiftedIndex];
+			int32 gid = Mesh.GetTriangleGroup(tid);
+			if (IncomingEdgeGroups.Contains(gid) == false)
+			{
+				OtherGroupTris.Add(tid);
+				OtherGroups.AddUnique(gid);
+			}
 		}
-	}
-	int32 NumRemainingTris = OtherGroupTris.Num();
+		int32 NumRemainingTris = OtherGroupTris.Num();
 
-	// Determine which edge to split at in the "other" group triangles. If we only have one group
-	// then we can try to pick the "middlest" edge. The worst case is when there is only one triangle,
-	// then we are picking a not-very-good edge no matter what (potentially we should do an edge split or
-	// face poke in that situation). If we have multiple groups then we probably want to pick one
-	// of the group-boundary edges inside the triangle-span, ideally the "middlest" but currently we
-	// are just picking one arbitrarily...
-	if (OtherGroups.Num() == 1)
-	{
-		Vertex.NewGroupID = OtherGroups[0];
-		if (OtherGroupTris.Num() == 1)
+		// If we didn't find other groups at all, then pick the outgoing edge
+		// that is most aligned to the incoming edge
+		if (OtherGroups.Num() == 0)
 		{
-			FIndex3i TriEdges = Mesh.GetTriEdges(OtherGroupTris[0]);
-			for (int32 j = 0; j < 3; ++j)
+			FIndex2i IncomingEdgeV = Mesh.GetEdgeV(IncomingEdgeID);
+			int32 OtherVID = IncomingEdgeV.OtherElement(Vertex.VertexID);
+			if (!ensure(OtherVID != IndexConstants::InvalidID))
 			{
-				if (Mesh.GetEdgeV(TriEdges[j]).Contains(Vertex.VertexID))
+				// incoming edge is not attached to the expected vertex, fail out
+				return;
+			}
+			FVector3d CenterPos = Mesh.GetVertex(Vertex.VertexID);
+			FVector3d IncomingPos = Mesh.GetVertex(OtherVID);
+			FVector3d IncomingDir = (CenterPos - IncomingPos);
+			IncomingDir.Normalize();
+			// Note: If the incoming edge is degenerate, this will just pick the first different edge.
+			// We could instead try to use a different heuristic in this case.
+			double BestAlignmentScore = -FMathd::MaxReal;
+			Mesh.EnumerateVertexEdges(Vertex.VertexID, 
+				[&BestAlignmentScore, &RingSplitEdgeID, &Mesh, &IncomingDir, &CenterPos, &Vertex, IncomingEdgeID]
+				(int32 EID)
 				{
-					RingSplitEdgeID = TriEdges[j];
-					break;
+					if (EID == IncomingEdgeID)
+					{
+						return;
+					}
+					FIndex2i OutEdgeV = Mesh.GetEdgeV(EID);
+					int32 OutVID = OutEdgeV.OtherElement(Vertex.VertexID);
+					if (ensure(OutVID != IndexConstants::InvalidID))
+					{
+						FVector3d OutgoingPos = Mesh.GetVertex(OutVID);
+						FVector3d OutgoingDir = OutgoingPos - CenterPos;
+						OutgoingDir.Normalize();
+						double AlignmentScore = OutgoingDir.Dot(IncomingDir);
+						if (AlignmentScore > BestAlignmentScore)
+						{
+							RingSplitEdgeID = EID;
+							BestAlignmentScore = AlignmentScore;
+						}
+					}
 				}
-			}
+			);
 		}
-		else if (OtherGroupTris.Num() == 2)
-		{
-			RingSplitEdgeID = FindSharedEdgeInTriangles(Mesh, OtherGroupTris[0], OtherGroupTris[1]);
-		}
-		else
+		// If we found edges on other groups, use them to guide the split --
+		// Determine which edge to split at in the "other" group triangles. If we only have one group
+		// then we can try to pick the "middlest" edge. The worst case is when there is only one triangle,
+		// then we are picking a not-very-good edge no matter what (potentially we should do an edge split or
+		// face poke in that situation). If we have multiple groups then we probably want to pick one
+		// of the group-boundary edges inside the triangle-span, ideally the "middlest" but currently we
+		// are just picking one arbitrarily...
+		else if (OtherGroups.Num() == 1)
 		{
-			// try using the 'middlest' triangle as the 'middlest' edge
-			// TODO: should compute opening angles here and pick the edge closest to the middle of the angular span!!
-			int32 j = OtherGroupTris.Num() / 2;
-			RingSplitEdgeID = FindSharedEdgeInTriangles(Mesh, OtherGroupTris[j], OtherGroupTris[j+1]);
-
-			// If the OtherGroupTris list ended up being not contiguous (see above for how that can happen), then
-			// it's possible that (j) and (j+1) are not connected and the share-edge search will fail. In that case
-			// we will just linear-search for two connected triangles. If this fails then this vertex will not be bevelled.
-			if (RingSplitEdgeID == -1)
+			Vertex.NewGroupID = OtherGroups[0];
+			if (OtherGroupTris.Num() == 1)
 			{
-				for (int32 k = 0; k < NumRemainingTris; ++k)
+				FIndex3i TriEdges = Mesh.GetTriEdges(OtherGroupTris[0]);
+				for (int32 j = 0; j < 3; ++j)
 				{
-					RingSplitEdgeID = FindSharedEdgeInTriangles(Mesh, OtherGroupTris[k], OtherGroupTris[(k+1)%NumRemainingTris]);
-					if (RingSplitEdgeID != -1)
+					if (Mesh.GetEdgeV(TriEdges[j]).Contains(Vertex.VertexID))
 					{
+						RingSplitEdgeID = TriEdges[j];
 						break;
 					}
 				}
 			}
+			else if (OtherGroupTris.Num() == 2)
+			{
+				RingSplitEdgeID = FindSharedEdgeInTriangles(Mesh, OtherGroupTris[0], OtherGroupTris[1]);
+			}
+			else
+			{
+				// try using the 'middlest' triangle as the 'middlest' edge
+				// TODO: should compute opening angles here and pick the edge closest to the middle of the angular span!!
+				int32 j = OtherGroupTris.Num() / 2;
+				RingSplitEdgeID = FindSharedEdgeInTriangles(Mesh, OtherGroupTris[j], OtherGroupTris[j + 1]);
+
+				// If the OtherGroupTris list ended up being not contiguous (see above for how that can happen), then
+				// it's possible that (j) and (j+1) are not connected and the share-edge search will fail. In that case
+				// we will just linear-search for two connected triangles. If this fails then this vertex will not be bevelled.
+				if (RingSplitEdgeID == -1)
+				{
+					for (int32 k = 0; k < NumRemainingTris; ++k)
+					{
+						RingSplitEdgeID = FindSharedEdgeInTriangles(Mesh, OtherGroupTris[k], OtherGroupTris[(k + 1) % NumRemainingTris]);
+						if (RingSplitEdgeID != -1)
+						{
+							break;
+						}
+					}
+				}
+			}
 		}
-	}
-	else
-	{
-		// Search for two adjacent triangles in different groups that have a shared edge. This search
-		// may need to wrap around if we got unlucky in the triangle ordering in OtherGroupTris
-		for (int32 k = 0; k < OtherGroupTris.Num(); ++k)
+		else
 		{
-			int32 TriangleA = OtherGroupTris[k], TriangleB = OtherGroupTris[(k+1)%NumRemainingTris];
-			if (Mesh.GetTriangleGroup(TriangleA) != Mesh.GetTriangleGroup(TriangleB))
+			// Search for two adjacent triangles in different groups that have a shared edge. This search
+			// may need to wrap around if we got unlucky in the triangle ordering in OtherGroupTris
+			for (int32 k = 0; k < OtherGroupTris.Num(); ++k)
 			{
-				RingSplitEdgeID = FindSharedEdgeInTriangles(Mesh, TriangleA, TriangleB);
-				Vertex.NewGroupID = -1;		// allocate a new group for this triangle, this is usually what one would want
-				break;
+				int32 TriangleA = OtherGroupTris[k], TriangleB = OtherGroupTris[(k + 1) % NumRemainingTris];
+				if (Mesh.GetTriangleGroup(TriangleA) != Mesh.GetTriangleGroup(TriangleB))
+				{
+					RingSplitEdgeID = FindSharedEdgeInTriangles(Mesh, TriangleA, TriangleB);
+					Vertex.NewGroupID = -1;		// allocate a new group for this triangle, this is usually what one would want
+					break;
+				}
 			}
 		}
 	}
@@ -1789,7 +1836,7 @@ void FMeshBevel::AppendTerminatorVertexTriangle(FDynamicMesh3& Mesh, FBevelVerte
 		int32 FarVertexID = SplitEdgeV.OtherElement(Vertex.VertexID);
 
 		int32 QuadEdgeID = Mesh.FindEdge(Vertex.Wedges[0].WedgeVertex, Vertex.Wedges[1].WedgeVertex);
-		if (Mesh.IsEdge(QuadEdgeID))
+		if (Mesh.IsEdge(QuadEdgeID) && Mesh.IsBoundaryEdge(QuadEdgeID))
 		{
 			FIndex2i QuadEdgeV = Mesh.GetOrientedBoundaryEdgeV(QuadEdgeID);
 			// should have computed this GroupID in initial setup
@@ -1826,7 +1873,13 @@ void FMeshBevel::AppendTerminatorVertexPairQuad(FDynamicMesh3& Mesh, FBevelVerte
 	// does not seem like we need to do anything w/ the TerminatorInfo here, we can get everything from wedges
 	int32 QuadEdgeID0 = Mesh.FindEdge(Vertex0.Wedges[0].WedgeVertex, Vertex0.Wedges[1].WedgeVertex);
 	int32 QuadEdgeID1 = Mesh.FindEdge(Vertex1.Wedges[0].WedgeVertex, Vertex1.Wedges[1].WedgeVertex);
-	MESH_BEVEL_DEBUG_CHECK(Mesh.IsEdge(QuadEdgeID0) && Mesh.IsEdge(QuadEdgeID1));
+	if (!MESH_BEVEL_DEBUG_ENSURE(
+		Mesh.IsEdge(QuadEdgeID0) && Mesh.IsEdge(QuadEdgeID1) &&
+		Mesh.IsBoundaryEdge(QuadEdgeID0) && Mesh.IsBoundaryEdge(QuadEdgeID1)))
+	{
+		// if we failed to find an open boundary on edge, we can't build the connecting quad
+		return;
+	}
 	FIndex2i QuadEdgeV0 = Mesh.GetOrientedBoundaryEdgeV(QuadEdgeID0);
 	FIndex2i QuadEdgeV1 = Mesh.GetOrientedBoundaryEdgeV(QuadEdgeID1);
 
@@ -1957,7 +2010,10 @@ void FMeshBevel::AppendLoopQuads(FDynamicMesh3& Mesh, FBevelLoop& Loop)
 			FIndex2i EdgeV0 = Mesh.GetOrientedBoundaryEdgeV(EdgeID0);
 			FIndex2i EdgeV1 = Mesh.GetOrientedBoundaryEdgeV(EdgeID1);
 			QuadTris.A = Mesh.AppendTriangle(EdgeV0.B, EdgeV0.A, EdgeV1.B, NewGroupID);
-			QuadTris.B = Mesh.AppendTriangle(EdgeV1.B, EdgeV1.A, EdgeV0.B, NewGroupID);
+			if (MESH_BEVEL_DEBUG_ENSURE(!EdgeV1.Contains(EdgeV0.B)))
+			{
+				QuadTris.B = Mesh.AppendTriangle(EdgeV1.B, EdgeV1.A, EdgeV0.B, NewGroupID);
+			}
 		}
 		Loop.StripQuads.Add(QuadTris);
 	}
@@ -2946,12 +3002,20 @@ void FMeshBevel::AppendTerminatorVertexTriangles_Multi(FDynamicMesh3& Mesh, FBev
 	int32 RingSplitEdgeID = Vertex.TerminatorInfo.A;
 	if (Mesh.IsEdge(RingSplitEdgeID))
 	{
-		FIndex2i SplitEdgeV = Mesh.GetEdgeV(RingSplitEdgeID);
-		int32 FarVertexID = SplitEdgeV.OtherElement(Vertex.VertexID);
-
 		int32 BevelEdgeIndex = Vertex.IncomingBevelEdgeIndices[0];
 		const FBevelEdge& IncomingEdge = Edges[BevelEdgeIndex];
 
+		// If we failed to create a strip quad patch (due to failure earlier in the pipeline)
+		// try to append a non-multi terminator instead
+		if (IncomingEdge.StripQuadPatch.IsEmpty())
+		{
+			AppendTerminatorVertexTriangle(Mesh, Vertex);
+			return;
+		}
+
+		FIndex2i SplitEdgeV = Mesh.GetEdgeV(RingSplitEdgeID);
+		int32 FarVertexID = SplitEdgeV.OtherElement(Vertex.VertexID);
+
 		int32 WedgeVertexA = Vertex.Wedges[0].WedgeVertex;
 		int32 WedgeVertexB = Vertex.Wedges[1].WedgeVertex;
 
@@ -2960,10 +3024,17 @@ void FMeshBevel::AppendTerminatorVertexTriangles_Multi(FDynamicMesh3& Mesh, FBev
 		check(ColumnIndex == 0 || ColumnIndex == IncomingEdge.StripQuadPatch.NumVertexCols() - 1);
 		TArray<int32> QuadStripEdgeVerts;
 		IncomingEdge.StripQuadPatch.GetVertexColumn(ColumnIndex, QuadStripEdgeVerts);
-		check(QuadStripEdgeVerts.Contains(WedgeVertexB));
+		if (!MESH_BEVEL_DEBUG_ENSURE(QuadStripEdgeVerts.Contains(WedgeVertexB)))
+		{
+			return;
+		}
 
 		// make sure it is oriented consistently
 		int32 FirstEdgeID = Mesh.FindEdge(QuadStripEdgeVerts[0], QuadStripEdgeVerts[1]);
+		if (!MESH_BEVEL_DEBUG_ENSURE(FirstEdgeID != IndexConstants::InvalidID && Mesh.IsBoundaryEdge(FirstEdgeID)))
+		{
+			return;
+		}
 		FIndex2i QuadEdgeV = Mesh.GetOrientedBoundaryEdgeV(FirstEdgeID);
 		if (QuadEdgeV.A != QuadStripEdgeVerts[0])
 		{
@@ -2993,23 +3064,43 @@ void FMeshBevel::AppendTerminatorVertexPairQuad_Multi(FDynamicMesh3& Mesh, FBeve
 	MESH_BEVEL_DEBUG_CHECK(Vertex0.VertexType == EBevelVertexType::TerminatorVertex);
 	MESH_BEVEL_DEBUG_CHECK(Vertex1.VertexType == EBevelVertexType::TerminatorVertex);
 
+
+
 	// This is a variant of AppendTerminatorVertexTriangle that handles the case where basically two
 	// Terminator Vertices are directly connected by a non-beveled mesh edge that was used as the ring-split-edge.
 	// Since both sides were opened, we have a quad-shaped hole instead of a triangle-shaped hole, with a 
 	// (subdivided, in the multi-case) quad-edge at each end. So this hole can be filled with a quad-patch that
 	// stitches together the existing edges (which must have already been meshed)
-
-	auto CollectTerminatorVtxInfo = [this](FDynamicMesh3& Mesh, FBevelVertex& Vertex, TArray<int32>& QuadStripEdgeVertsOut)
+	bool bFallbackToQuad = false, bInvalidInfo = false;
+	auto CollectTerminatorVtxInfo = [this, &bFallbackToQuad, &bInvalidInfo](FDynamicMesh3& Mesh, FBevelVertex& Vertex, TArray<int32>& QuadStripEdgeVertsOut) -> void
 	{
 		int32 BevelEdgeIndex = Vertex.IncomingBevelEdgeIndices[0];
 		const FBevelEdge& IncomingEdge = Edges[BevelEdgeIndex];
+		if (IncomingEdge.StripQuadPatch.IsEmpty())
+		{
+			bFallbackToQuad = true;
+			return;
+		}
 
 		int32 ColumnIndex = IncomingEdge.StripQuadPatch.FindColumnIndex(Vertex.Wedges[0].WedgeVertex);
-		check(ColumnIndex == 0 || ColumnIndex == IncomingEdge.StripQuadPatch.NumVertexCols() - 1);
+		if (!(ColumnIndex == 0 || ColumnIndex == IncomingEdge.StripQuadPatch.NumVertexCols() - 1))
+		{
+			bInvalidInfo = true;
+			return;
+		}
 		IncomingEdge.StripQuadPatch.GetVertexColumn(ColumnIndex, QuadStripEdgeVertsOut);
-		check(QuadStripEdgeVertsOut.Contains(Vertex.Wedges[1].WedgeVertex));
+		if (!(QuadStripEdgeVertsOut.Contains(Vertex.Wedges[1].WedgeVertex)))
+		{
+			bInvalidInfo = true;
+			return;
+		}
 
 		int32 FirstEdgeID = Mesh.FindEdge(QuadStripEdgeVertsOut[0], QuadStripEdgeVertsOut[1]);
+		if (FirstEdgeID == IndexConstants::InvalidID || !Mesh.IsBoundaryEdge(FirstEdgeID))
+		{
+			bInvalidInfo = true;
+			return;
+		}
 		FIndex2i QuadEdgeV = Mesh.GetOrientedBoundaryEdgeV(FirstEdgeID);
 		if (QuadEdgeV.A != QuadStripEdgeVertsOut[0])
 		{
@@ -3020,6 +3111,17 @@ void FMeshBevel::AppendTerminatorVertexPairQuad_Multi(FDynamicMesh3& Mesh, FBeve
 	TArray<int32> QuadStripEdgeVerts0, QuadStripEdgeVerts1;
 	CollectTerminatorVtxInfo(Mesh, Vertex0, QuadStripEdgeVerts0);
 	CollectTerminatorVtxInfo(Mesh, Vertex1, QuadStripEdgeVerts1);
+	if (bInvalidInfo)
+	{
+		MESH_BEVEL_DEBUG_CHECK(false);
+		return;
+	}
+	if (bFallbackToQuad)
+	{
+		// no quad-strip fallback case
+		AppendTerminatorVertexPairQuad(Mesh, Vertex0, Vertex1);
+		return;
+	}
 	if (QuadStripEdgeVerts0.Num() != QuadStripEdgeVerts1.Num() || QuadStripEdgeVerts0.Num() == 0)
 	{
 		// this could happen if the quadstrip failed on one or both sides? will leave a hole.
