make another mesh bevel check() into a soft failure / return

diff --git a/Engine/Plugins/Runtime/GeometryProcessing/Source/DynamicMesh/Private/Operations/MeshBevel.cpp b/Engine/Plugins/Runtime/GeometryProcessing/Source/DynamicMesh/Private/Operations/MeshBevel.cpp
--- a/Engine/Plugins/Runtime/GeometryProcessing/Source/DynamicMesh/Private/Operations/MeshBevel.cpp
+++ b/Engine/Plugins/Runtime/GeometryProcessing/Source/DynamicMesh/Private/Operations/MeshBevel.cpp
@@ -3021,7 +3021,10 @@ void FMeshBevel::AppendTerminatorVertexTriangles_Multi(FDynamicMesh3& Mesh, FBev
 
 		// figure out which side of the incoming-edge quad-strip the terminator connects to
 		int32 ColumnIndex = IncomingEdge.StripQuadPatch.FindColumnIndex(WedgeVertexA);
-		check(ColumnIndex == 0 || ColumnIndex == IncomingEdge.StripQuadPatch.NumVertexCols() - 1);
+		if (!MESH_BEVEL_DEBUG_ENSURE(ColumnIndex == 0 || ColumnIndex == IncomingEdge.StripQuadPatch.NumVertexCols() - 1))
+		{
+			return;
+		}
 		TArray<int32> QuadStripEdgeVerts;
 		IncomingEdge.StripQuadPatch.GetVertexColumn(ColumnIndex, QuadStripEdgeVerts);
 		if (!MESH_BEVEL_DEBUG_ENSURE(QuadStripEdgeVerts.Contains(WedgeVertexB)))
