add target edge length support to the mesh cluster simplifier

#FailedValidations Preflight

diff --git a/Engine/Plugins/Runtime/GeometryProcessing/Source/DynamicMesh/Private/Operations/MeshClusterSimplifier.cpp b/Engine/Plugins/Runtime/GeometryProcessing/Source/DynamicMesh/Private/Operations/MeshClusterSimplifier.cpp
--- a/Engine/Plugins/Runtime/GeometryProcessing/Source/DynamicMesh/Private/Operations/MeshClusterSimplifier.cpp
+++ b/Engine/Plugins/Runtime/GeometryProcessing/Source/DynamicMesh/Private/Operations/MeshClusterSimplifier.cpp
@@ -3,6 +3,7 @@
 //#include "Operations/UniformTessellate.h"
 #include "Operations/MeshClusterSimplifier.h"
 
+#include "Algo/MinElement.h"
 #include "Async/ParallelFor.h"
 #include "DynamicMesh/DynamicMesh3.h"
 #include "DynamicMesh/DynamicMeshAttributeSet.h"
@@ -157,24 +158,36 @@ namespace UE::MeshClusterSimplifyLocals
 	}
 
 	template<typename MeshType>
-	static TArray<float> InitTargetLengths(const MeshType& InMesh, const FSimplifyOptions& SimplifyOptions)
+	static TArray<float> InitTargetLengths(const MeshType& InMesh, float& OutMinTargetLength, const FSimplifyOptions& SimplifyOptions)
 	{
 		TArray<float> TargetLengths;
-		// Currently we always initialize to empty target lengths, which is treated as all vertices requesting SimplifyOptions.TargetEdgeLength.
-		// 
-		// Note that arbitrary target lengths may lead to bad results w/ the current algorithm. To support arbitrary lengths,
-		// the algorithm needs to prevent clusters from larger-target-length source vertices from running through or around those w/ 
-		// smaller target lengths. This is additionally complicated by the multi-pass strategy for handling constrained 
-		// edges. So for now, we just don't support arbitrary initial target lengths.
-		// 
-		// Note this is not as problematic for the lengths we currently automatically set when running 'preserve collapsed cluster' passes,
-		// because in those cases we tend to set locally consistent lengths in the regions where we locally re-run the clustering.
+		OutMinTargetLength = FMathf::MaxReal;
+		if (SimplifyOptions.OptionalTargetEdgeLengthScale)
+		{
+			TargetLengths.SetNumUninitialized(InMesh.MaxVertexID());
+			TArray<float> MinLengthContexts;
+			ParallelForWithTaskContext(MinLengthContexts, InMesh.MaxVertexID(),
+				[](int32, int32) -> float {return FMathf::MaxReal;},
+				[&InMesh, &TargetLengths, &SimplifyOptions](float& MinLength, int32 VID)
+				{
+					if (!InMesh.IsVertex(VID))
+					{
+						return;
+					}
+					float TargetLength = float(SimplifyOptions.TargetEdgeLength * SimplifyOptions.OptionalTargetEdgeLengthScale(VID));
+					MinLength = FMath::Min(MinLength, TargetLength);
+					TargetLengths[VID] = TargetLength;
+				}
+			);
+			OutMinTargetLength = *Algo::MinElement(MinLengthContexts);
+		}
+
 		return TargetLengths;
 	}
 
 	template<typename MeshType, typename MeshConnType>
 	void ClusterVerticesByRegionGrowth(const MeshType& InMesh, const MeshConnType& InMeshConnectivity, TArray<int32>& Source, TArray<float>& SourceDist, const FSimplifyOptions& SimplifyOptions,
-		const TArray<EElemTag>& EdgeTags, const TArray<EElemTag>& VertexTags, const TArray<float>& TargetLengths,
+		const TArray<EElemTag>& EdgeTags, const TArray<EElemTag>& VertexTags, const TArray<float>& TargetLengths, float MinTargetLength,
 		const FVertexBuckets& VertexIDBuckets)
 	{
 		// add all the fixed vertices as sources first, so they can't be claimed by other verts
@@ -196,10 +209,13 @@ namespace UE::MeshClusterSimplifyLocals
 		};
 		TArray<FWalk> HeapV;
 
-		// for the non-fixed vertices, progressively grow from vertices, in passes from more-constrained to less-constrained edges
+		// process buckets in two phases:
+		// first, starting from fixed/constrained verts, and disallowing free->(fixed/constrained) walks
+		// second, starting-from and walking-to remaining unclaimed free vertices
+		uint8 BucketIdx = 0;
 		for (uint8 TagIdx = 1; TagIdx < (uint8)EElemTag::MAX; ++TagIdx)
 		{
-			for (uint8 BucketIdx = 0; BucketIdx <= TagIdx; ++BucketIdx)
+			for (; BucketIdx <= TagIdx; ++BucketIdx)
 			{
 				const TArray<int32>& CurBucket = VertexIDBuckets[BucketIdx];
 				for (int32 InBucketIdx = 0; InBucketIdx < CurBucket.Num(); ++InBucketIdx)
@@ -208,43 +224,49 @@ namespace UE::MeshClusterSimplifyLocals
 
 					int32& CurSourceVID = Source[GrowFromVID];
 					float& CurSourceDist = SourceDist[GrowFromVID];
-					// the vertex is unclaimed, claim it as a new source/kept vertex
-					if (CurSourceVID == INDEX_NONE)
+					// the vertex is fixed or unclaimed, claim it as a new source/kept vertex
+					if (BucketIdx == (uint8)EElemTag::Fixed || CurSourceVID == INDEX_NONE)
 					{
 						CurSourceVID = GrowFromVID;
 						CurSourceDist = 0.f;
 					}
-					// if the vertex was claimed by another source in the current tag pass, no need to process it further
-					else if (CurSourceVID != GrowFromVID && (uint8)VertexTags[GrowFromVID] == TagIdx)
+					// the vertex was claimed by another source, no need to process it further
+					else
 					{
 						continue;
 					}
 
-					float MaxDist = TargetLengths.IsEmpty() ? (float)SimplifyOptions.TargetEdgeLength : TargetLengths[CurSourceVID];
+					float CurrentMaxDist = TargetLengths.IsEmpty() ? (float)SimplifyOptions.TargetEdgeLength : TargetLengths[CurSourceVID];
 
-					// vertex is either a new source, or previously claimed but we need to consider growing via less-constrained edges
+					// if we're on the unconstrained (free) vertex bucket, don't need to check vertex/edge tags anymore
+					bool bIsFreePass = TagIdx == (uint8)EElemTag::Free;
 
 					// helper to add candidate verts to a heap
 					HeapV.Reset();
-					auto AddCandidates = [MaxDist,
-						&HeapV, &InMesh, &InMeshConnectivity, &SourceDist, &Source, &EdgeTags, &VertexTags, TagIdx]
-						(const FWalk& From)
+					auto AddCandidates = [
+						&HeapV, &InMesh, &InMeshConnectivity, &SourceDist, &Source, &EdgeTags, &VertexTags, TagIdx, bIsFreePass]
+						(const FWalk& From, float UseMaxDist)
 						{
+							uint8 FromTagIdx = bIsFreePass ? (uint8)EElemTag::Free : (uint8)VertexTags[From.VID];
 							// expand to one-ring
 							InMeshConnectivity.EnumerateVertexEdges(From.VID,
-								[&From, MaxDist,
-								&HeapV, &InMesh, &InMeshConnectivity, &SourceDist, &Source, &EdgeTags, &VertexTags, TagIdx]
+								[&From, UseMaxDist,
+								&HeapV, &InMesh, &InMeshConnectivity, &SourceDist, &Source, &EdgeTags, &VertexTags, TagIdx, FromTagIdx, bIsFreePass]
 								(int32 EID)
 								{
-									if ((uint8)EdgeTags[EID] != TagIdx)
-									{
-										return;
-									}
-
 									FIndex2i EdgeV = InMeshConnectivity.GetEdgeV(EID);
 									int32 ToVID = EdgeV.A == From.VID ? EdgeV.B : EdgeV.A;
-
-									if (IsMoreConstrained(VertexTags[ToVID], TagIdx) || From.Dist >= SourceDist[ToVID])
+									
+									uint8 UseFromTagIdx = FromTagIdx;
+									// in free pass we're only visiting free verts which only have free edges by construction
+									// but in constrained pass, we need to also test the tag on the edge
+									if (!bIsFreePass)
+									{
+										uint8 EdgeTagIdx = (uint8)EdgeTags[EID];
+										// can't go from (fixed/constrained)->constrained through an unconstrained edge
+										UseFromTagIdx = FMath::Max(EdgeTagIdx, FromTagIdx);
+									}
+									if (IsMoreConstrained(VertexTags[ToVID], UseFromTagIdx) || From.Dist >= SourceDist[ToVID])
 									{
 										// vertex was already claimed by more-constrained context, or is already as close (or closer) to another source
 										return;
@@ -253,7 +275,7 @@ namespace UE::MeshClusterSimplifyLocals
 									FVector3d Pos = InMesh.GetVertex(ToVID);
 									FVector3d FromPos = InMesh.GetVertex(From.VID);
 									float NewDist = From.Dist + (float)FVector3d::Dist(Pos, FromPos);
-									if (NewDist < MaxDist && NewDist < SourceDist[ToVID])
+									if (NewDist < UseMaxDist && NewDist < SourceDist[ToVID])
 									{
 										// Viable candidate distance; add to heap
 										HeapV.HeapPush(FWalk{ ToVID, NewDist });
@@ -264,13 +286,28 @@ namespace UE::MeshClusterSimplifyLocals
 
 					// initialize the heap w/ the neighbors of the initial grow-from vertex
 					FWalk Start{ GrowFromVID, CurSourceDist };
-					AddCandidates(Start);
+					AddCandidates(Start, CurrentMaxDist);
+
+					// In uniform target length case, or if we're already at the min target length, can skip the MaxDist updates
+					bool bNeedToUpdateMaxDist = !TargetLengths.IsEmpty() && CurrentMaxDist > MinTargetLength;
 
 					while (!HeapV.IsEmpty())
 					{
 						FWalk CurWalk;
 						HeapV.HeapPop(CurWalk, EAllowShrinking::No);
 
+						// If we have smaller TargetLengths elsewhere, test if we are walking over
+						// a smaller length and clamp max dist + stop walking if needed
+						if (bNeedToUpdateMaxDist)
+						{
+							CurrentMaxDist = FMath::Min(CurrentMaxDist, TargetLengths[CurWalk.VID]);
+							
+							if (CurWalk.Dist >= CurrentMaxDist)
+							{
+								break;
+							}
+						}
+
 						// we already got to this vert from another place
 						if (SourceDist[CurWalk.VID] <= CurWalk.Dist)
 						{
@@ -282,7 +319,7 @@ namespace UE::MeshClusterSimplifyLocals
 						Source[CurWalk.VID] = CurSourceVID;
 
 						// search its (current-tag-level) edges for more verts to claim
-						AddCandidates(CurWalk);
+						AddCandidates(CurWalk, CurrentMaxDist);
 					}
 				}
 			}
@@ -317,7 +354,7 @@ namespace UE::MeshClusterSimplifyLocals
 	static void PreserveCollapsedClusters(const MeshType& InMesh, const MeshConnType& InMeshConnectivity,
 		TArray<bool>& ClusterInOutput,
 		TArray<int32>& Source, TArray<float>& SourceDist, const FSimplifyOptions& SimplifyOptions,
-		const TArray<EElemTag>& EdgeTags, const TArray<EElemTag>& VertexTags, TArray<float>& TargetLengths,
+		const TArray<EElemTag>& EdgeTags, const TArray<EElemTag>& VertexTags, TArray<float>& TargetLengths, float MinTargetLength,
 		const FVertexBuckets& VertexIDBuckets)
 	{
 		for (int32 PassIdx = 0; PassIdx < SimplifyOptions.MaxPreserveCollapsedClusterPasses; ++PassIdx)
@@ -378,7 +415,7 @@ namespace UE::MeshClusterSimplifyLocals
 				}
 			}
 			// Re-run tagging
-			ClusterVerticesByRegionGrowth(InMesh, InMeshConnectivity, Source, SourceDist, SimplifyOptions, EdgeTags, VertexTags, TargetLengths, ReProcessBuckets);
+			ClusterVerticesByRegionGrowth(InMesh, InMeshConnectivity, Source, SourceDist, SimplifyOptions, EdgeTags, VertexTags, TargetLengths, MinTargetLength, ReProcessBuckets);
 			// Update tracking of clusters in output
 			FindClustersInOutput(InMesh, Source, ClusterInOutput);
 		}
@@ -389,7 +426,7 @@ namespace UE::MeshClusterSimplifyLocals
 		TArray<bool>& ClusterInOutput,
 		bool& bOutResultHasDuplicateVertices,
 		const InMeshType& InMesh, const InMeshConnType& InMeshConnectivity, TArray<int32>& Source, TArray<float>& SourceDist, const FSimplifyOptions& SimplifyOptions,
-		const TArray<EElemTag>& EdgeTags, TArray<EElemTag>& VertexTags, const TArray<float>& TargetLengths,
+		const TArray<EElemTag>& EdgeTags, TArray<EElemTag>& VertexTags, const TArray<float>& TargetLengths, float MinTargetLength,
 		FVertexBuckets& ProcessBuckets)
 	{
 		TArray<int32> ToResVID;
@@ -489,7 +526,7 @@ namespace UE::MeshClusterSimplifyLocals
 					VertexTags[VID] = EElemTag::Fixed;
 				}
 				ProcessBuckets[0] = MoveTemp(SourceVIDToFix);
-				ClusterVerticesByRegionGrowth(InMesh, InMeshConnectivity, Source, SourceDist, SimplifyOptions, EdgeTags, VertexTags, TargetLengths, ProcessBuckets);
+				ClusterVerticesByRegionGrowth(InMesh, InMeshConnectivity, Source, SourceDist, SimplifyOptions, EdgeTags, VertexTags, TargetLengths, MinTargetLength, ProcessBuckets);
 				FindClustersInOutput(InMesh, Source, ClusterInOutput);
 				continue;
 			}
@@ -648,7 +685,8 @@ bool Simplify(const FDynamicMesh3& InMesh, FDynamicMesh3& ResultMesh, const FSim
 		GetEdgeConstraintLevelFromAttributes, IsSeamIntersectionVertex,
 		SimplifyOptions, EdgeTags, VertexTags);
 
-	TArray<float> TargetLengths = InitTargetLengths(InMesh, SimplifyOptions);
+	float MinTargetLength;
+	TArray<float> TargetLengths = InitTargetLengths(InMesh, MinTargetLength, SimplifyOptions);
 
 	///
 	/// Step 2. Clustering: Grow vertex clusters out to the target edge length size
@@ -666,7 +704,7 @@ bool Simplify(const FDynamicMesh3& InMesh, FDynamicMesh3& ResultMesh, const FSim
 	Source.Init(INDEX_NONE, InMesh.MaxVertexID());
 	SourceDist.Init(FMathf::MaxReal, InMesh.MaxVertexID());
 
-	ClusterVerticesByRegionGrowth(InMesh, InMesh, Source, SourceDist, SimplifyOptions, EdgeTags, VertexTags, TargetLengths, ProcessBuckets);
+	ClusterVerticesByRegionGrowth(InMesh, InMesh, Source, SourceDist, SimplifyOptions, EdgeTags, VertexTags, TargetLengths, MinTargetLength, ProcessBuckets);
 	for (int32 Idx = 0; Idx < ProcessBuckets.Num(); ++Idx)
 	{
 		ProcessBuckets[Idx].Empty();
@@ -675,7 +713,7 @@ bool Simplify(const FDynamicMesh3& InMesh, FDynamicMesh3& ResultMesh, const FSim
 	/// Step 2.5: Determine which clusters are in the output, and optionally try reduce target lengths to preserve missing clusters
 	TArray<bool> ClusterInOutput;
 	FindClustersInOutput(InMesh, Source, ClusterInOutput);
-	PreserveCollapsedClusters(InMesh, InMesh, ClusterInOutput, Source, SourceDist, SimplifyOptions, EdgeTags, VertexTags, TargetLengths, ProcessBuckets);
+	PreserveCollapsedClusters(InMesh, InMesh, ClusterInOutput, Source, SourceDist, SimplifyOptions, EdgeTags, VertexTags, TargetLengths, MinTargetLength, ProcessBuckets);
 
 	///
 	/// Step 3: Copy the cluster connectivity out to our ResultMesh
@@ -684,7 +722,7 @@ bool Simplify(const FDynamicMesh3& InMesh, FDynamicMesh3& ResultMesh, const FSim
 	TArray<int32> FromResVID, ResultToSourceTri;
 	bool bResultHasDuplicateVertices = false;
 	BuildSimplifiedMeshFromClusters<FDynamicMesh3, FDynamicMesh3, FDynamicMesh3>(ResultMesh, FromResVID, ResultToSourceTri, ClusterInOutput, bResultHasDuplicateVertices,
-		InMesh, InMesh, Source, SourceDist, SimplifyOptions, EdgeTags, VertexTags, TargetLengths,
+		InMesh, InMesh, Source, SourceDist, SimplifyOptions, EdgeTags, VertexTags, TargetLengths, MinTargetLength,
 		ProcessBuckets);
 
 	///
@@ -953,7 +991,8 @@ bool Simplify(const FTriangleMeshAdapterd& InMesh, FResultMeshAdapter& OutSimpli
 		TagEdgeFromAttributes, IsSeamIntersectionVertex,
 		SimplifyOptions, EdgeTags, VertexTags);
 
-	TArray<float> TargetLengths = InitTargetLengths(InMesh, SimplifyOptions);
+	float MinTargetLength;
+	TArray<float> TargetLengths = InitTargetLengths(InMesh, MinTargetLength, SimplifyOptions);
 
 	///
 	/// Step 2. Clustering: Grow vertex clusters out to the target edge length size
@@ -974,7 +1013,7 @@ bool Simplify(const FTriangleMeshAdapterd& InMesh, FResultMeshAdapter& OutSimpli
 	Source.Init(INDEX_NONE, InMesh.MaxVertexID());
 	SourceDist.Init(FMathf::MaxReal, InMesh.MaxVertexID());
 
-	ClusterVerticesByRegionGrowth(InMesh, MeshConnectivity, Source, SourceDist, SimplifyOptions, EdgeTags, VertexTags, TargetLengths, ProcessBuckets);
+	ClusterVerticesByRegionGrowth(InMesh, MeshConnectivity, Source, SourceDist, SimplifyOptions, EdgeTags, VertexTags, TargetLengths, MinTargetLength, ProcessBuckets);
 	for (int32 Idx = 0; Idx < ProcessBuckets.Num(); ++Idx)
 	{
 		ProcessBuckets[Idx].Empty();
@@ -984,7 +1023,7 @@ bool Simplify(const FTriangleMeshAdapterd& InMesh, FResultMeshAdapter& OutSimpli
 	/// Step 2.5: Optionally look for clusters that didn't make it into triangles
 	TArray<bool> ClusterInOutput;
 	FindClustersInOutput(InMesh, Source, ClusterInOutput);
-	PreserveCollapsedClusters(InMesh, MeshConnectivity, ClusterInOutput, Source, SourceDist, SimplifyOptions, EdgeTags, VertexTags, TargetLengths, ProcessBuckets);
+	PreserveCollapsedClusters(InMesh, MeshConnectivity, ClusterInOutput, Source, SourceDist, SimplifyOptions, EdgeTags, VertexTags, TargetLengths, MinTargetLength, ProcessBuckets);
 
 	///
 	/// Step 3: Copy the cluster connectivity out to our ResultMesh
@@ -994,7 +1033,7 @@ bool Simplify(const FTriangleMeshAdapterd& InMesh, FResultMeshAdapter& OutSimpli
 	bool bResultHasDuplicateVertices = false;
 	BuildSimplifiedMeshFromClusters<TTriangleMeshAdapter<double>, FMeshConnectivity, FResultMeshAdapter>(OutSimplifiedMesh, FromResVID, ResultToSourceTri, ClusterInOutput,
 		bResultHasDuplicateVertices,
-		InMesh, MeshConnectivity, Source, SourceDist, SimplifyOptions, EdgeTags, VertexTags, TargetLengths,
+		InMesh, MeshConnectivity, Source, SourceDist, SimplifyOptions, EdgeTags, VertexTags, TargetLengths, MinTargetLength,
 		ProcessBuckets);
 
 	//
diff --git a/Engine/Plugins/Runtime/GeometryProcessing/Source/DynamicMesh/Public/Operations/MeshClusterSimplifier.h b/Engine/Plugins/Runtime/GeometryProcessing/Source/DynamicMesh/Public/Operations/MeshClusterSimplifier.h
--- a/Engine/Plugins/Runtime/GeometryProcessing/Source/DynamicMesh/Public/Operations/MeshClusterSimplifier.h
+++ b/Engine/Plugins/Runtime/GeometryProcessing/Source/DynamicMesh/Public/Operations/MeshClusterSimplifier.h
@@ -22,6 +22,10 @@ struct FSimplifyOptions
 	// Vertices within this distance are allowed to be replaced with a single vertex
 	double TargetEdgeLength = 1;
 
+	// Optionally locally scales the simplifier's TargetEdgeLength on a per-vertex basis
+	// If provided, the target edge length at a vertex VID will be TargetEdgeLength * OptionalTargetEdgeLengthScale(VID)
+	TFunction<double(int32)> OptionalTargetEdgeLengthScale = nullptr;
+
 	// If > 0, number of passes to attempt to make sure the simplification does not fully delete islands of the mesh.
 	// In areas that would be deleted (because the region is smaller than the TargetEdgeLength),
 	// the target edge length will instead be locally reduced to half the local max in-cluster distance,
