[bughawk] adaptive tessellation / geometry script

diff --git a/Engine/Plugins/Runtime/GeometryProcessing/Source/DynamicMesh/Private/Tessellation/AdaptiveDisplacement.cpp b/Engine/Plugins/Runtime/GeometryProcessing/Source/DynamicMesh/Private/Tessellation/AdaptiveDisplacement.cpp
--- a/Engine/Plugins/Runtime/GeometryProcessing/Source/DynamicMesh/Private/Tessellation/AdaptiveDisplacement.cpp
+++ b/Engine/Plugins/Runtime/GeometryProcessing/Source/DynamicMesh/Private/Tessellation/AdaptiveDisplacement.cpp
@@ -61,6 +61,9 @@ struct FDynamicMeshDisplacementPolicy
 		const FIndex3i Triangle = Mesh.GetTriangle(TriIndex);
 		const int LocalVtxIdx = IndexUtil::FindTriIndex(VertexIndex, Triangle);
 
+		// vertex is guaranteed to be referenced by triangle
+		check(LocalVtxIdx != -1);
+
 		const FIndex3i UVTri = UVOverlay->GetTriangle(TriIndex);
 		const FVector2f UV = TransformUVs(UVOverlay->GetElement(UVTri[LocalVtxIdx]));
 
diff --git a/Engine/Plugins/Runtime/GeometryProcessing/Source/DynamicMesh/Public/Tessellation/AdaptiveDisplacement.h b/Engine/Plugins/Runtime/GeometryProcessing/Source/DynamicMesh/Public/Tessellation/AdaptiveDisplacement.h
--- a/Engine/Plugins/Runtime/GeometryProcessing/Source/DynamicMesh/Public/Tessellation/AdaptiveDisplacement.h
+++ b/Engine/Plugins/Runtime/GeometryProcessing/Source/DynamicMesh/Public/Tessellation/AdaptiveDisplacement.h
@@ -17,11 +17,11 @@ class FDisplacementMap;
  * Apply adaptive displacement on given mesh.
  * 
  * @param Mesh               mesh to be displaced (requires per-vertex normals to be present)
- * @param UVOverlay          UVs on which to evaluate displacement map
+ * @param UVOverlay          UVs on which to evaluate displacement map (required)
  * @param DisplacementMap    displacement map
  * @param Options            control sample rate, error, etc.
  * @param FeatureSensitivity value in [0, 1] to control additional adaptive refinement at sharp features
- * @param MeshContraints     constraints to preserve
+ * @param MeshConstraints    constraints to preserve
  * @param ActiveTriangles    optional per-triangle selection
  * @param bApplyPostOptimization whether to apply some post regularization
  * @return                   whether the adaptive displacement succeeded
diff --git a/Engine/Plugins/Runtime/GeometryScripting/Source/GeometryScriptingCore/Private/MeshDeformFunctions.cpp b/Engine/Plugins/Runtime/GeometryScripting/Source/GeometryScriptingCore/Private/MeshDeformFunctions.cpp
--- a/Engine/Plugins/Runtime/GeometryScripting/Source/GeometryScriptingCore/Private/MeshDeformFunctions.cpp
+++ b/Engine/Plugins/Runtime/GeometryScripting/Source/GeometryScriptingCore/Private/MeshDeformFunctions.cpp
@@ -555,11 +555,12 @@ void ApplyAdaptiveDisplacement(
 			ActiveTriangles[TID] = true;
 		});
 
+		// Keep boundary of active triangle region constrained.
 		FEdgeConstraint NoFlipConstraint(EEdgeRefineFlags::FullyConstrained);
 		for (int EID : Mesh.EdgeIndicesItr())
 		{
 			const FDynamicMesh3::FEdge& Edge = Mesh.GetEdgeRef(EID);
-			if (ActiveTriangles[Edge.Tri[0]] == false || ActiveTriangles[Edge.Tri[1]] == false)
+			if (Edge.Tri[1] != IndexConstants::InvalidID && (ActiveTriangles[Edge.Tri[0]] != ActiveTriangles[Edge.Tri[1]]))
 			{
 				MeshConstraints.SetOrUpdateEdgeConstraint(EID, NoFlipConstraint);
 			}
diff --git a/Engine/Source/Runtime/GeometryCore/Public/Tessellation/AdaptiveMetrics.h b/Engine/Source/Runtime/GeometryCore/Public/Tessellation/AdaptiveMetrics.h
--- a/Engine/Source/Runtime/GeometryCore/Public/Tessellation/AdaptiveMetrics.h
+++ b/Engine/Source/Runtime/GeometryCore/Public/Tessellation/AdaptiveMetrics.h
@@ -69,11 +69,6 @@ inline UE::Math::TVector2<T> CalculateErrorBounds(
 	const FVector2f MagnitudeBounds, // bounds 
 	const int32 DisplacementBoundsRefinements)
 {
-	if (MagnitudeBounds.X >= MagnitudeBounds.Y) 
-	{
-		return UE::Math::TVector2<T>(T(0.));
-	}
-
 	T MinBarycentric0 = FMath::Min3(Barycentrics[0].X, Barycentrics[1].X, Barycentrics[2].X);
 	T MaxBarycentric0 = FMath::Max3(Barycentrics[0].X, Barycentrics[1].X, Barycentrics[2].X);
 
diff --git a/Engine/Source/Runtime/GeometryCore/Public/Tessellation/FeatureAdaptive.h b/Engine/Source/Runtime/GeometryCore/Public/Tessellation/FeatureAdaptive.h
--- a/Engine/Source/Runtime/GeometryCore/Public/Tessellation/FeatureAdaptive.h
+++ b/Engine/Source/Runtime/GeometryCore/Public/Tessellation/FeatureAdaptive.h
@@ -33,7 +33,7 @@ public:
 
 	[[nodiscard]] FORCEINLINE bool NeedsCheck() const
 	{
-		return (Threshold * FeatureSize < IncircleRadius * UE_PI); // Early-out (also if IncircleRadius is NaN)
+		return (Threshold * FeatureSize < IncircleRadius * UE_PI); // Early-out (also if IncircleRadius is NaN or zero)
 	}
 
 	// check whether the edge should be refined refined
@@ -45,6 +45,7 @@ public:
 		const Math::TVector<T> NT = (TriVertices[(LocalEdge + 1) % 3] - TriVertices[LocalEdge]) ^ (FarVertex - TriVertices[LocalEdge]);
 		const T LSqNT = NT.SquaredLength();
 
+		// this expression evaluates to false for degenerate triangles
 		return ((TriNormal | NT) < FMath::Sqrt(LSqN * LSqNT) * FMath::Cos(Threshold * FeatureSize / IncircleRadius));
 	}
 
