Add optional scaling functions for the QEM simplifier's metrics -- quadric error, target length, and geometry error

diff --git a/Engine/Plugins/Runtime/GeometryProcessing/Source/DynamicMesh/Private/MeshSimplification.cpp b/Engine/Plugins/Runtime/GeometryProcessing/Source/DynamicMesh/Private/MeshSimplification.cpp
--- a/Engine/Plugins/Runtime/GeometryProcessing/Source/DynamicMesh/Private/MeshSimplification.cpp
+++ b/Engine/Plugins/Runtime/GeometryProcessing/Source/DynamicMesh/Private/MeshSimplification.cpp
@@ -229,6 +229,11 @@ void TMeshSimplification<QuadricErrorType>::InitializeQueue()
 		FQuadricErrorType Q = AssembleEdgeQuadric(edge);
 		FVector3d opt = OptimalPoint(eid, Q, edge.Vert.A, edge.Vert.B);
 		EdgeErrors[eid] = { (float)Q.Evaluate(opt), eid };
+		if (CustomQuadricErrorScaleF)
+		{
+			double ErrorScale = CustomQuadricErrorScaleF(*Mesh, edge.Vert.A, edge.Vert.B);
+			EdgeErrors[eid].error *= (float)(ErrorScale * ErrorScale);
+		}
 		EdgeQuadrics[eid] = QEdge(eid, Q, opt);
 	}
 
@@ -545,6 +550,11 @@ void TMeshSimplification<QuadricErrorType>::UpdateNeighborhood(const FDynamicMes
 			QuadricErrorType Q = AssembleEdgeQuadric(ne);
 			FVector3d opt = OptimalPoint(eid, Q, ne.Vert.A, ne.Vert.B);
 			float err = (float)Q.Evaluate(opt);
+			if (CustomQuadricErrorScaleF)
+			{
+				double ErrorScale = CustomQuadricErrorScaleF(*Mesh, ne.Vert.A, ne.Vert.B);
+				err *= (float)(ErrorScale * ErrorScale);
+			}
 			EdgeQuadrics[eid] = QEdge(eid, Q, opt);
 			if (EdgeQueue.Contains(eid))
 			{
@@ -585,6 +595,11 @@ void TMeshSimplification<QuadricErrorType>::UpdateNeighborhood(const FDynamicMes
 
 			FVector3d opt = OptimalPoint(eid, Q, edgeData.Vert[0], edgeData.Vert[1]);
 			float err = (float)Q.Evaluate(opt);
+			if (CustomQuadricErrorScaleF)
+			{
+				double ErrorScale = CustomQuadricErrorScaleF(*Mesh, edgeData.Vert.A, edgeData.Vert.B);
+				err *= (float)(ErrorScale * ErrorScale);
+			}
 			EdgeQuadrics[eid] = QEdge(eid, Q, opt);
 			if (EdgeQueue.Contains(eid))
 			{
@@ -1125,7 +1140,14 @@ void TMeshSimplification<QuadricErrorType>::FastCollapsePass(double fMinEdgeLeng
 			}
 
 			Mesh->GetEdgeV(eid, va, vb);
-			if (DistanceSquared(va, vb) > min_sqr)
+			double EdgeLenSq = DistanceSquared(va, vb);
+			if (CustomEdgeLengthScaleF)
+			{
+				FIndex2i EdgeV = Mesh->GetEdgeV(eid);
+				double Scale = CustomEdgeLengthScaleF(*Mesh, EdgeV.A, EdgeV.B);
+				EdgeLenSq *= Scale * Scale;
+			}
+			if (EdgeLenSq > min_sqr)
 			{
 				continue;
 			}
@@ -1293,6 +1315,11 @@ ESimplificationResult TMeshSimplification<QuadricErrorType>::CollapseEdge(int ed
 	FVector3d vA = Mesh->GetVertex(a);
 	FVector3d vB = Mesh->GetVertex(b);
 	double edge_len_sqr = (vA - vB).SquaredLength();
+	if (CustomEdgeLengthScaleF)
+	{
+		double Scale = CustomEdgeLengthScaleF(*Mesh, a, b);
+		edge_len_sqr *= Scale * Scale;
+	}
 	if (edge_len_sqr > MinEdgeLength * MinEdgeLength)
 	{
 		return ESimplificationResult::Ignored_EdgeTooLong;
@@ -1528,6 +1555,18 @@ bool TMeshSimplification<QuadricErrorType>::CheckIfCollapseWithinGeometricTolera
 		if (ProjectionTarget() != nullptr)
 		{
 			double ToleranceSqr = GeometricErrorTolerance * GeometricErrorTolerance;
+			if (CustomGeometricErrorScaleF)
+			{
+				double Scale = CustomGeometricErrorScaleF(*Mesh, vKeep, vRemove);
+				double ErrorScaleSq = Scale * Scale;
+				// if error is scaled to zero, it will always pass tolerance
+				if (ErrorScaleSq <= 0)
+				{
+					return true;
+				}
+				// We divide the tolerance; this is equivalent to scaling the errors below but a bit faster / simpler
+				ToleranceSqr /= ErrorScaleSq;
+			}
 
 			// test new position to see if it is within geometric tolerance of projection surface
 			FVector3d TargetPos = ProjectionTarget()->Project(NewPosition);
diff --git a/Engine/Plugins/Runtime/GeometryProcessing/Source/DynamicMesh/Public/MeshSimplification.h b/Engine/Plugins/Runtime/GeometryProcessing/Source/DynamicMesh/Public/MeshSimplification.h
--- a/Engine/Plugins/Runtime/GeometryProcessing/Source/DynamicMesh/Public/MeshSimplification.h
+++ b/Engine/Plugins/Runtime/GeometryProcessing/Source/DynamicMesh/Public/MeshSimplification.h
@@ -101,6 +101,11 @@ public:
 	/** Tolerance to use in geometric error checking */
 	double GeometricErrorTolerance = 0.0;
 
+	// Custom scales to apply to *measured* errors and edge lengths. Note this is in effect the inverse of scaling the tolerances / target lengths.
+	TFunction<double(const FDynamicMesh3&, int VertexA, int VertexB)> CustomGeometricErrorScaleF;
+	TFunction<double(const FDynamicMesh3&, int VertexA, int VertexB)> CustomQuadricErrorScaleF;
+	TFunction<double(const FDynamicMesh3&, int VertexA, int VertexB)> CustomEdgeLengthScaleF;
+
 
 	TMeshSimplification(FDynamicMesh3* m) : FMeshRefinerBase(m)
 	{
