Move the direction normalization for the line regularization into the regularization method, and have it automatically fall back to position regularization

diff --git a/Engine/Plugins/Runtime/GeometryProcessing/Source/DynamicMesh/Private/MeshSimplification.cpp b/Engine/Plugins/Runtime/GeometryProcessing/Source/DynamicMesh/Private/MeshSimplification.cpp
--- a/Engine/Plugins/Runtime/GeometryProcessing/Source/DynamicMesh/Private/MeshSimplification.cpp
+++ b/Engine/Plugins/Runtime/GeometryProcessing/Source/DynamicMesh/Private/MeshSimplification.cpp
@@ -177,14 +177,14 @@ void TMeshSimplification<QuadricErrorType>::InitializeVertexQuadrics()
 			Q.Add(triAreas[tid], triQuadrics[tid]);
 		}
 
-		ApplyRegularizationToVertexQuadric(Q, vid, AreaSum, bNeedsNormals, Normal);
+		ApplyRegularizationToVertexQuadric(Q, vid, AreaSum, Normal);
 
 		vertQuadrics[vid] = Q;
 	}
 }
 
 template <typename QuadricErrorType>
-void TMeshSimplification<QuadricErrorType>::ApplyRegularizationToVertexQuadric(FQuadricErrorType& Q, int32 VID, double AreaSum, bool bNeedsNormals, FVector3d Normal)
+void TMeshSimplification<QuadricErrorType>::ApplyRegularizationToVertexQuadric(FQuadricErrorType& Q, int32 VID, double AreaSum, FVector3d Normal)
 {
 	if (RegularizeWeight > 0.)
 	{
@@ -194,9 +194,8 @@ void TMeshSimplification<QuadricErrorType>::ApplyRegularizationToVertexQuadric(F
 			UseRegularizeWeight = RegularizeWeight * CustomRegularizationScaleF(*Mesh, VID);
 		}
 		FVector3d VertPos = Mesh->GetVertex(VID);
-		if (bNeedsNormals && Normal.SquaredLength() > 0)
+		if (RegularizeMethod == ERegularizeSimplificationMethod::NormalLine)
 		{
-			Normal.Normalize();
 			// Note: Very large weights line-regularizer weights introduce instability in the solve
 			// so we clamp the map weight in this case
 			constexpr double MaxLineRegularizerWeight = 100.;
@@ -589,7 +588,7 @@ void TMeshSimplification<QuadricErrorType>::UpdateNeighborhood(const FDynamicMes
 				}
 				vertQuadric.Add(triAreas[tid], triQuadrics[tid]);
 			}
-			ApplyRegularizationToVertexQuadric(vertQuadric, kvid, AreaSum, bNeedsNormals, Normal);
+			ApplyRegularizationToVertexQuadric(vertQuadric, kvid, AreaSum, Normal);
 			vertQuadrics[kvid] = vertQuadric;
 		}
 	}
diff --git a/Engine/Plugins/Runtime/GeometryProcessing/Source/DynamicMesh/Public/MeshSimplification.h b/Engine/Plugins/Runtime/GeometryProcessing/Source/DynamicMesh/Public/MeshSimplification.h
--- a/Engine/Plugins/Runtime/GeometryProcessing/Source/DynamicMesh/Public/MeshSimplification.h
+++ b/Engine/Plugins/Runtime/GeometryProcessing/Source/DynamicMesh/Public/MeshSimplification.h
@@ -461,7 +461,7 @@ protected:
 
 private:
 
-	void ApplyRegularizationToVertexQuadric(FQuadricErrorType& Q, int32 VID, double AreaSum, bool bNeedsNormals, FVector3d Normal);
+	void ApplyRegularizationToVertexQuadric(FQuadricErrorType& Q, int32 VID, double AreaSum, FVector3d Normal);
 
 };
 
diff --git a/Engine/Source/Runtime/GeometryCore/Public/QuadricError.h b/Engine/Source/Runtime/GeometryCore/Public/QuadricError.h
--- a/Engine/Source/Runtime/GeometryCore/Public/QuadricError.h
+++ b/Engine/Source/Runtime/GeometryCore/Public/QuadricError.h
@@ -257,6 +257,13 @@ struct TQuadricError
 	// Add regularizing term that penalizes movement away from the line through Pos along LineDir
 	void AddLineRegularizer(const TVector<RealType>& Pos, const TVector<RealType>& LineDir, RealType Weight)
 	{
+		TVector<RealType> NormalizedLineDir = LineDir;
+		if (!NormalizedLineDir.Normalize())
+		{
+			// if we fail to normalize, fall back to the position regularizer
+			AddPositionRegularizer(Pos, Weight);
+			return;
+		}
 		TVector<RealType> N1, N2;
 		VectorUtil::MakePerpVectors(LineDir, N1, N2);
 		Add(Weight * .5, TQuadricError<RealType>(N1, Pos));
