apply 'regularization' scale to the entire vertex quadric after adding the regularization term, so that the regularization can't outweight the actual QEM error term at large scales

diff --git a/Engine/Plugins/Runtime/GeometryProcessing/Source/DynamicMesh/Private/MeshSimplification.cpp b/Engine/Plugins/Runtime/GeometryProcessing/Source/DynamicMesh/Private/MeshSimplification.cpp
--- a/Engine/Plugins/Runtime/GeometryProcessing/Source/DynamicMesh/Private/MeshSimplification.cpp
+++ b/Engine/Plugins/Runtime/GeometryProcessing/Source/DynamicMesh/Private/MeshSimplification.cpp
@@ -189,10 +189,7 @@ void TMeshSimplification<QuadricErrorType>::ApplyRegularizationToVertexQuadric(F
 	if (RegularizeWeight > 0.)
 	{
 		double UseRegularizeWeight = RegularizeWeight;
-		if (CustomRegularizationScaleF)
-		{
-			UseRegularizeWeight = RegularizeWeight * CustomRegularizationScaleF(*Mesh, VID);
-		}
+
 		FVector3d VertPos = Mesh->GetVertex(VID);
 		if (RegularizeMethod == ERegularizeSimplificationMethod::NormalLine)
 		{
@@ -206,6 +203,14 @@ void TMeshSimplification<QuadricErrorType>::ApplyRegularizationToVertexQuadric(F
 		{
 			Q.AddPositionRegularizer(VertPos, UseRegularizeWeight * AreaSum);
 		}
+
+		// Note we scale the *entire* quadric after adding the regularization term
+		// to ensure that the 'actual' QEM error remains the dominant term even at
+		// large scales
+		if (CustomRegularizationScaleF)
+		{
+			Q.Scale(CustomRegularizationScaleF(*Mesh, VID));
+		}
 	}
 }
 
diff --git a/Engine/Source/Runtime/GeometryCore/Public/QuadricError.h b/Engine/Source/Runtime/GeometryCore/Public/QuadricError.h
--- a/Engine/Source/Runtime/GeometryCore/Public/QuadricError.h
+++ b/Engine/Source/Runtime/GeometryCore/Public/QuadricError.h
@@ -385,6 +385,12 @@ public:
 		PlaneData.Add(w, b.PlaneData);
 	}
 
+	void Scale(RealType Wt)
+	{
+		BaseStruct::Scale(Wt);
+		PlaneData.N *= Wt;
+		PlaneData.Dist *= Wt;
+	}
 
 	/**
 	* The optimal point minimizing the quadric error with respect to a volume conserving constraint
@@ -596,6 +602,18 @@ public:
 		d[2] += w * other.d[2];
 	}
 
+	void Scale(RealType Wt)
+	{
+		BaseClass::Scale(Wt);
+		a *= Wt;
+		grad[0] *= Wt;
+		grad[1] *= Wt;
+		grad[2] *= Wt;
+		d[0] *= Wt;
+		d[1] *= Wt;
+		d[2] *= Wt;
+		// note: we intentionally do not scale attrweight here; it characterizes the relative importance of attribute error and doesn't scale w/ data
+	}
 
 	/**
 	* The optimal point minimizing the quadric error with respect to a volume conserving constraint
