use Reset/SetNumZeroed to init boundary cache in mesh simplifier, so all the values are explicitly init'd to false

diff --git a/Engine/Plugins/Runtime/GeometryProcessing/Source/DynamicMesh/Private/MeshSimplification.cpp b/Engine/Plugins/Runtime/GeometryProcessing/Source/DynamicMesh/Private/MeshSimplification.cpp
--- a/Engine/Plugins/Runtime/GeometryProcessing/Source/DynamicMesh/Private/MeshSimplification.cpp
+++ b/Engine/Plugins/Runtime/GeometryProcessing/Source/DynamicMesh/Private/MeshSimplification.cpp
@@ -689,7 +689,8 @@ template <typename QuadricErrorType>
 void TMeshSimplification<QuadricErrorType>::Precompute(bool bMeshIsClosed)
 {
 	bHaveBoundary = false;
-	IsBoundaryVtxCache.SetNum(Mesh->MaxVertexID());
+	IsBoundaryVtxCache.Reset();
+	IsBoundaryVtxCache.SetNumZeroed(Mesh->MaxVertexID());
 	if (bMeshIsClosed == false)
 	{
 		for (int eid : Mesh->BoundaryEdgeIndicesItr())
