fix mesh simplifier's bad placement of boundary vertices (e.g. rounding off corners) -- consider error when picking boundary vertex position instead of always just picking the midpoint

diff --git a/Engine/Plugins/Runtime/GeometryProcessing/Source/DynamicMesh/Private/MeshSimplification.cpp b/Engine/Plugins/Runtime/GeometryProcessing/Source/DynamicMesh/Private/MeshSimplification.cpp
--- a/Engine/Plugins/Runtime/GeometryProcessing/Source/DynamicMesh/Private/MeshSimplification.cpp
+++ b/Engine/Plugins/Runtime/GeometryProcessing/Source/DynamicMesh/Private/MeshSimplification.cpp
@@ -3,6 +3,7 @@
 #include "MeshSimplification.h"
 #include "MeshConstraintsUtil.h"
 #include "DynamicMesh/DynamicMeshAttributeSet.h"
+#include "SegmentTypes.h"
 #include "Util/IndexUtil.h"
 #include "Async/ParallelFor.h"
 #include "Templates/UnrealTypeTraits.h"
@@ -339,6 +340,24 @@ void TMeshSimplification<QuadricErrorType>::InitializeQueue()
 template <typename QuadricErrorType>
 FVector3d TMeshSimplification<QuadricErrorType>::OptimalPoint(int eid, const FQuadricErrorType& q, int ea, int eb)
 {
+	// Common tracking used when we pick the best of multiple points below
+	FVector3d BestPt(0);
+	double BestErr = FMathd::MaxReal;
+	auto InitBestPt = [&BestPt, &BestErr, &q](const FVector3d& Pt)
+	{
+		BestErr = q.Evaluate(Pt);
+		BestPt = Pt;
+	};
+	auto UpdateBestPt = [&BestPt, &BestErr, &q](const FVector3d& Pt)
+	{
+		double Err = q.Evaluate(Pt);
+		if (Err < BestErr)
+		{
+			BestErr = Err;
+			BestPt = Pt;
+		}
+	};
+
 	// if we would like to preserve boundary, we need to know that here
 	// so that we properly score these edges
 	if (bHaveBoundary && bPreserveBoundaryShape)
@@ -347,8 +366,33 @@ FVector3d TMeshSimplification<QuadricErrorType>::OptimalPoint(int eid, const FQu
 		{
 			const bool bModeAllowsVertMovement = (CollapseMode != ESimplificationCollapseModes::MinimalExistingVertexError);
 			if (bModeAllowsVertMovement)
-			{ 
-				return (Mesh->GetVertex(ea) + Mesh->GetVertex(eb)) * 0.5;
+			{
+				if (CollapseMode == ESimplificationCollapseModes::MinimalQuadricPositionError)
+				{
+					// We project the optimal point to the segment, but since this may not be minimal after projection,
+					// we use this point only if it has lower error than the segment mid and endpoints.
+					// Note: Could consider doing a line search here instead
+
+					FVector3d A = Mesh->GetVertex(ea);
+					FVector3d B = Mesh->GetVertex(eb);
+					FSegment3d EdgeSeg(A, B);
+
+					InitBestPt(A);
+					UpdateBestPt(B);
+					UpdateBestPt(EdgeSeg.Center);
+
+					FVector3d UnconstrainedResult = FVector3d::Zero();
+					if (q.OptimalPoint(UnconstrainedResult))
+					{
+						FVector3d ProjectedOpt = EdgeSeg.NearestPoint(UnconstrainedResult);
+						UpdateBestPt(ProjectedOpt);
+					}
+					return BestPt;
+				}
+				else // ESimplificationCollapseModes::AverageVertexPosition
+				{
+					return (Mesh->GetVertex(ea) + Mesh->GetVertex(eb)) * 0.5;
+				}
 			} // else MinimalExistingVertexError case below will choose one of the vertex locations
 		}
 		else
@@ -376,19 +420,9 @@ FVector3d TMeshSimplification<QuadricErrorType>::OptimalPoint(int eid, const FQu
 
 		case ESimplificationCollapseModes::MinimalExistingVertexError:
 		{
-			FVector3d va = Mesh->GetVertex(ea);
-			FVector3d vb = Mesh->GetVertex(eb);
-			double fa = q.Evaluate(va);
-			double fb = q.Evaluate(vb);
-			if (fa < fb)
-			{
-				return va;
-			}
-			else
-			{
-				return vb;
-			}
-		
+			InitBestPt(Mesh->GetVertex(ea));
+			UpdateBestPt(Mesh->GetVertex(eb));
+			return BestPt;
 		}
 		break;
 
@@ -402,22 +436,12 @@ FVector3d TMeshSimplification<QuadricErrorType>::OptimalPoint(int eid, const FQu
 
 			// degenerate matrix, evaluate quadric at edge end and midpoints
 			// (could do line search here...)
-			FVector3d va = Mesh->GetVertex(ea);
-			FVector3d vb = Mesh->GetVertex(eb);
-			FVector3d c = GetProjectedPoint((va + vb) * 0.5);
-			double fa = q.Evaluate(va);
-			double fb = q.Evaluate(vb);
-			double fc = q.Evaluate(c);
-			double m = FMath::Min3(fa, fb, fc);
-			if (m == fa)
-			{
-				return va;
-			}
-			else if (m == fb)
-			{
-				return vb;
-			}
-			return c;
+			FVector3d VA = Mesh->GetVertex(ea);
+			FVector3d VB = Mesh->GetVertex(eb);
+			InitBestPt(VA);
+			UpdateBestPt(VB);
+			UpdateBestPt(GetProjectedPoint((VA + VB) * 0.5));
+			return BestPt;
 		}
 		break;
 	default:
