add geometric tolerance to mesh->convex hull simplification path, to allow simplification below the target triangle count when it will not introduce more error above the threshold

diff --git a/Engine/Plugins/Experimental/GeometryFlow/Source/GeometryFlowMeshProcessing/Private/PhysicsNodes/GenerateSimpleCollisionNode.cpp b/Engine/Plugins/Experimental/GeometryFlow/Source/GeometryFlowMeshProcessing/Private/PhysicsNodes/GenerateSimpleCollisionNode.cpp
--- a/Engine/Plugins/Experimental/GeometryFlow/Source/GeometryFlowMeshProcessing/Private/PhysicsNodes/GenerateSimpleCollisionNode.cpp
+++ b/Engine/Plugins/Experimental/GeometryFlow/Source/GeometryFlowMeshProcessing/Private/PhysicsNodes/GenerateSimpleCollisionNode.cpp
@@ -46,7 +46,8 @@ namespace GenerateSimpleCollisionNodeHelpers
 			}
 
 			Hull.bPostSimplify = Settings.ConvexHullSettings.SimplifyToTriangleCount > 0;
-			Hull.MaxTargetFaceCount = Settings.ConvexHullSettings.SimplifyToTriangleCount;
+			Hull.SimplifyOptions.MaxTargetFaceCount = Settings.ConvexHullSettings.SimplifyToTriangleCount;
+			Hull.SimplifyOptions.GeometricToleranceAfterTarget = Settings.ConvexHullSettings.GeometricToleranceAfterTarget;
 			if (Hull.Compute(Progress))
 			{
 				FConvexShape3d NewConvex;
@@ -79,7 +80,8 @@ namespace GenerateSimpleCollisionNodeHelpers
 				}
 
 				Hull.bPostSimplify = Settings.ConvexHullSettings.SimplifyToTriangleCount > 0;
-				Hull.MaxTargetFaceCount = Settings.ConvexHullSettings.SimplifyToTriangleCount;
+				Hull.SimplifyOptions.MaxTargetFaceCount = Settings.ConvexHullSettings.SimplifyToTriangleCount;
+				Hull.SimplifyOptions.GeometricToleranceAfterTarget = Settings.ConvexHullSettings.GeometricToleranceAfterTarget;
 				if (Hull.Compute(Progress))
 				{
 					FConvexShape3d NewConvex;
diff --git a/Engine/Plugins/Experimental/GeometryFlow/Source/GeometryFlowMeshProcessing/Public/PhysicsNodes/GenerateSimpleCollisionNode.h b/Engine/Plugins/Experimental/GeometryFlow/Source/GeometryFlowMeshProcessing/Public/PhysicsNodes/GenerateSimpleCollisionNode.h
--- a/Engine/Plugins/Experimental/GeometryFlow/Source/GeometryFlowMeshProcessing/Public/PhysicsNodes/GenerateSimpleCollisionNode.h
+++ b/Engine/Plugins/Experimental/GeometryFlow/Source/GeometryFlowMeshProcessing/Public/PhysicsNodes/GenerateSimpleCollisionNode.h
@@ -53,6 +53,9 @@ struct FGenerateConvexHullSettings
 	UPROPERTY(EditAnywhere, Category = "Geometry Flow")
 	int32 SimplifyToTriangleCount = 50;
 
+	UPROPERTY(EditAnywhere, Category = "Geometry Flow")
+	double GeometricToleranceAfterTarget = 0.;
+
 	UPROPERTY(EditAnywhere, Category = "Geometry Flow")
 	bool bPrefilterVertices = true;
 
diff --git a/Engine/Plugins/Experimental/MeshModelingToolsetExp/Source/MeshModelingToolsExp/Private/Physics/SetCollisionGeometryTool.cpp b/Engine/Plugins/Experimental/MeshModelingToolsetExp/Source/MeshModelingToolsExp/Private/Physics/SetCollisionGeometryTool.cpp
--- a/Engine/Plugins/Experimental/MeshModelingToolsetExp/Source/MeshModelingToolsExp/Private/Physics/SetCollisionGeometryTool.cpp
+++ b/Engine/Plugins/Experimental/MeshModelingToolsetExp/Source/MeshModelingToolsExp/Private/Physics/SetCollisionGeometryTool.cpp
@@ -318,6 +318,7 @@ void USetCollisionGeometryTool::Setup()
 	Settings->WatchProperty(Settings->bDetectCapsules, [this](bool) { InvalidateCompute(); });
 	Settings->WatchProperty(Settings->bSimplifyHulls, [this](bool) { InvalidateCompute(); });
 	Settings->WatchProperty(Settings->HullTargetFaceCount, [this](int32) { InvalidateCompute(); });
+	Settings->WatchProperty(Settings->HullSimplifyGeometryTolerance, [this](double) { InvalidateCompute(); });
 	Settings->WatchProperty(Settings->MaxHullsPerShape, [this](int32) { InvalidateCompute(); });
 	Settings->WatchProperty(Settings->bPreSimplifyToEdgeLength, [this](bool) { InvalidateCompute(); });
 	Settings->WatchProperty(Settings->DecompositionTargetEdgeLength, [this](double) { InvalidateCompute(); });
@@ -460,6 +461,7 @@ TUniquePtr<UE::Geometry::TGenericDataOperator<FPhysicsDataCollection>> USetColli
 	// SimplifyHulls on the shape generator controls simplification on both swept and convex hull paths, but for Swept Hulls UI we leave simplification always enabled
 	Op->UseShapeGenerator->bSimplifyHulls = Settings->GeometryType == ECollisionGeometryType::SweptHulls || Settings->bSimplifyHulls;
 	Op->UseShapeGenerator->HullTargetFaceCount = Settings->HullTargetFaceCount;
+	Op->UseShapeGenerator->HullGeometricToleranceAfterTarget = Settings->HullSimplifyGeometryTolerance;
 	Op->UseShapeGenerator->bDecompositionPreSimplifyWithEdgeLength = Settings->bPreSimplifyToEdgeLength;
 	Op->UseShapeGenerator->DecompositionPreSimplifyEdgeLength = Settings->DecompositionTargetEdgeLength;
 	Op->UseShapeGenerator->bUseConvexDecompositionMaxPieces = Settings->bLimitHullsPerShape || Settings->DecompositionMethod != EConvexDecompositionMethod::NavigationDriven;
diff --git a/Engine/Plugins/Experimental/MeshModelingToolsetExp/Source/MeshModelingToolsExp/Public/Physics/SetCollisionGeometryTool.h b/Engine/Plugins/Experimental/MeshModelingToolsetExp/Source/MeshModelingToolsExp/Public/Physics/SetCollisionGeometryTool.h
--- a/Engine/Plugins/Experimental/MeshModelingToolsetExp/Source/MeshModelingToolsExp/Public/Physics/SetCollisionGeometryTool.h
+++ b/Engine/Plugins/Experimental/MeshModelingToolsetExp/Source/MeshModelingToolsExp/Public/Physics/SetCollisionGeometryTool.h
@@ -195,6 +195,11 @@ public:
 		EditConditionHides, EditCondition = "(GeometryType == ECollisionGeometryType::ConvexHulls  || GeometryType == ECollisionGeometryType::ConvexDecompositions) && bSimplifyHulls"))
 	int32 HullTargetFaceCount = 20;
 
+	/** Geometric tolerance to apply after simplifying to the target face count */
+	UPROPERTY(EditAnywhere, Category = ConvexHulls, meta = (UIMin = "0", UIMax = "100", ClampMin = "0",
+		EditConditionHides, EditCondition = "(GeometryType == ECollisionGeometryType::ConvexHulls  || GeometryType == ECollisionGeometryType::ConvexDecompositions) && bSimplifyHulls"))
+	double HullSimplifyGeometryTolerance = 1.;
+
 	/** Whether to simplify the input to this edge length before computing convex decomposition. Can give a decomposition result faster for large meshes. */
 	UPROPERTY(EditAnywhere, DisplayName = "Simplify to Edge Length", Category = ConvexHulls,
 		meta = (EditConditionHides, EditCondition = "GeometryType == ECollisionGeometryType::ConvexDecompositions"))
diff --git a/Engine/Plugins/Runtime/GeometryProcessing/Source/DynamicMesh/Private/Operations/MeshConvexHull.cpp b/Engine/Plugins/Runtime/GeometryProcessing/Source/DynamicMesh/Private/Operations/MeshConvexHull.cpp
--- a/Engine/Plugins/Runtime/GeometryProcessing/Source/DynamicMesh/Private/Operations/MeshConvexHull.cpp
+++ b/Engine/Plugins/Runtime/GeometryProcessing/Source/DynamicMesh/Private/Operations/MeshConvexHull.cpp
@@ -3,6 +3,8 @@
 #include "Operations/MeshConvexHull.h"
 #include "Solvers/MeshLinearization.h"
 #include "MeshSimplification.h"
+#include "ProjectionTargets.h"
+#include "DynamicMesh/DynamicMeshAABBTree3.h"
 #include "DynamicMesh/MeshNormals.h"
 #include "CompGeom/ConvexHull3.h"
 #include "Util/GridIndexing3.h"
@@ -28,27 +30,49 @@ bool FMeshConvexHull::Compute(FProgressCancel* Progress)
 
 	if (bPostSimplify)
 	{
-		bOK = SimplifyHull(ConvexHull, MaxTargetFaceCount, Progress);
+		bOK = SimplifyHull(ConvexHull, SimplifyOptions, Progress);
 	}
 
 	return bOK;
 }
 
 
-bool FMeshConvexHull::SimplifyHull(FDynamicMesh3& HullMesh, int32 MaxTargetFaceCount, FProgressCancel* Progress)
+bool FMeshConvexHull::SimplifyHull(FDynamicMesh3& HullMesh, FMeshConvexHull::FSimplifyOptions SimplifyOptions, FProgressCancel* Progress)
 {
 	if (Progress && Progress->Cancelled())
 	{
 		return false;
 	}
 
-	check(MaxTargetFaceCount > 0);
+	if (!ensure(SimplifyOptions.MaxTargetFaceCount > 0))
+	{
+		SimplifyOptions.MaxTargetFaceCount = 1;
+	}
 	bool bSimplified = false;
-	if (HullMesh.TriangleCount() > MaxTargetFaceCount)
+
+	
+	bool bUseGeometricError = SimplifyOptions.GeometricToleranceAfterTarget > 0;
+	FDynamicMeshAABBTree3 Spatial;
+	TUniquePtr<FDynamicMesh3> RefMesh;
+	if (bUseGeometricError)
+	{
+		RefMesh = MakeUnique<FDynamicMesh3>(HullMesh);
+		Spatial.SetMesh(RefMesh.Get(), true);
+	}
+
+	if (HullMesh.TriangleCount() > SimplifyOptions.MaxTargetFaceCount)
 	{
 		FVolPresMeshSimplification Simplifier(&HullMesh);
 		Simplifier.CollapseMode = FVolPresMeshSimplification::ESimplificationCollapseModes::MinimalExistingVertexError;
-		Simplifier.SimplifyToTriangleCount(MaxTargetFaceCount);
+		Simplifier.SimplifyToTriangleCount(SimplifyOptions.MaxTargetFaceCount);
+		if (bUseGeometricError)
+		{
+			Simplifier.GeometricErrorConstraint = FVolPresMeshSimplification::EGeometricErrorCriteria::PredictedPointToProjectionTarget;
+			Simplifier.GeometricErrorTolerance = SimplifyOptions.GeometricToleranceAfterTarget;
+			UE::Geometry::FMeshProjectionTarget ProjTarget(RefMesh.Get(), &Spatial);
+			Simplifier.SetProjectionTarget(&ProjTarget);
+			Simplifier.SimplifyToTriangleCount(1);
+		}
 		bSimplified = true;
 	}
 
diff --git a/Engine/Plugins/Runtime/GeometryProcessing/Source/DynamicMesh/Private/ShapeApproximation/MeshSimpleShapeApproximation.cpp b/Engine/Plugins/Runtime/GeometryProcessing/Source/DynamicMesh/Private/ShapeApproximation/MeshSimpleShapeApproximation.cpp
--- a/Engine/Plugins/Runtime/GeometryProcessing/Source/DynamicMesh/Private/ShapeApproximation/MeshSimpleShapeApproximation.cpp
+++ b/Engine/Plugins/Runtime/GeometryProcessing/Source/DynamicMesh/Private/ShapeApproximation/MeshSimpleShapeApproximation.cpp
@@ -317,7 +317,8 @@ void FMeshSimpleShapeApproximation::Generate_ConvexHulls(FSimpleShapeSet3d& Shap
 		FMeshConvexHull Hull(&SourceMesh);
 
 		Hull.bPostSimplify = bSimplifyHulls;
-		Hull.MaxTargetFaceCount = HullTargetFaceCount;
+		Hull.SimplifyOptions.MaxTargetFaceCount = HullTargetFaceCount;
+		Hull.SimplifyOptions.GeometricToleranceAfterTarget = HullGeometricToleranceAfterTarget;
 		Hull.MinDimension = MinDimension;
 
 		if (Hull.Compute(Progress))
@@ -445,7 +446,10 @@ void FMeshSimpleShapeApproximation::Generate_ConvexHullDecompositions(FSimpleSha
 		for (int32 HullIdx = 0; HullIdx < Decomposition.NumHulls(); HullIdx++)
 		{
 			FDynamicMesh3 HullMesh = Decomposition.GetHullMesh(HullIdx);
-			if (bSimplifyHulls && FMeshConvexHull::SimplifyHull(HullMesh, HullTargetFaceCount, Progress) == false)
+			FMeshConvexHull::FSimplifyOptions SimplifyOptions;
+			SimplifyOptions.MaxTargetFaceCount = HullTargetFaceCount;
+			SimplifyOptions.GeometricToleranceAfterTarget = HullGeometricToleranceAfterTarget;
+			if (bSimplifyHulls && FMeshConvexHull::SimplifyHull(HullMesh, SimplifyOptions, Progress) == false)
 			{
 				return;
 			}
diff --git a/Engine/Plugins/Runtime/GeometryProcessing/Source/DynamicMesh/Public/Operations/MeshConvexHull.h b/Engine/Plugins/Runtime/GeometryProcessing/Source/DynamicMesh/Public/Operations/MeshConvexHull.h
--- a/Engine/Plugins/Runtime/GeometryProcessing/Source/DynamicMesh/Public/Operations/MeshConvexHull.h
+++ b/Engine/Plugins/Runtime/GeometryProcessing/Source/DynamicMesh/Public/Operations/MeshConvexHull.h
@@ -31,7 +31,19 @@ public:
 
 	/** If true, output convex hull is simplified down to MaxTargetFaceCount */
 	bool bPostSimplify = false;
-	/** Target triangle count of the output Convex Hull */
+
+	struct FSimplifyOptions
+	{
+		/** Target triangle count of the simplified output Convex Hull; hull will be simplified to at most this triangle count if possible */
+		int32 MaxTargetFaceCount = 0;
+		/** If greater than zero, hull will be futher simplified below target face count where doing so remains within this error tolerance */
+		double GeometricToleranceAfterTarget = 0.;
+	};
+
+	/** IF bPostSimplify, options controlling the simplification */
+	FSimplifyOptions SimplifyOptions;
+
+	UE_DEPRECATED(5.8, "Use SimplifyOptions.MaxTargetFaceCount instead")
 	int32 MaxTargetFaceCount = 0;
 
 	/** Minimum extent along the shortest dimension; if greater than zero, the hull may be expanded if it would be too thin */
@@ -75,7 +87,18 @@ public:
 	 * Simplify the output convex hull. Assumes ConvexHull is already computed.
 	 * @return true on success
 	 */
-	static UE_API bool SimplifyHull(FDynamicMesh3& HullMesh, int32 MaxTargetFaceCount, FProgressCancel* Progress = nullptr);
+	static bool SimplifyHull(FDynamicMesh3& HullMesh, int32 MaxTargetFaceCount, FProgressCancel* Progress = nullptr)
+	{
+		FSimplifyOptions Options;
+		Options.MaxTargetFaceCount = MaxTargetFaceCount;
+		return SimplifyHull(HullMesh, Options, Progress);
+	}
+
+	/**
+	 * Simplify the output convex hull. Assumes ConvexHull is already computed.
+	 * @return true on success
+	 */
+	static UE_API bool SimplifyHull(FDynamicMesh3& HullMesh, FSimplifyOptions SimplifyOptions, FProgressCancel* Progress = nullptr);
 
 protected:
 	UE_API bool Compute_FullMesh(FProgressCancel* Progress);
diff --git a/Engine/Plugins/Runtime/GeometryProcessing/Source/DynamicMesh/Public/ShapeApproximation/MeshSimpleShapeApproximation.h b/Engine/Plugins/Runtime/GeometryProcessing/Source/DynamicMesh/Public/ShapeApproximation/MeshSimpleShapeApproximation.h
--- a/Engine/Plugins/Runtime/GeometryProcessing/Source/DynamicMesh/Public/ShapeApproximation/MeshSimpleShapeApproximation.h
+++ b/Engine/Plugins/Runtime/GeometryProcessing/Source/DynamicMesh/Public/ShapeApproximation/MeshSimpleShapeApproximation.h
@@ -67,7 +67,10 @@ public:
 	bool bSimplifyHulls = true;
 	/** target number of triangles when simplifying 3D convex hulls */
 	int32 HullTargetFaceCount = 50;
-	/** simplification tolerance when simplifying 2D convex hulls, eg for swept/projected hulls */
+	/** simplification tolerance for 3D convex hulls, if > 0, applies *after* hull target face count is reached to futher simplify */
+	double HullGeometricToleranceAfterTarget = 0.;
+
+	/** simplification tolerance when simplifying 2D (swept/projected) convex hulls; note HullGeometricToleranceAfterTarget is used instead for 3D hulls */
 	double HullSimplifyTolerance = 1.0;
 
 	/** Whether to use apply an edge-length based simplification to the input before running convex decompositions. Useful for very dense input meshes where the decomposition will be slow to compute. */
diff --git a/Engine/Plugins/Runtime/GeometryScripting/Source/GeometryScriptingCore/Private/CollisionFunctions.cpp b/Engine/Plugins/Runtime/GeometryScripting/Source/GeometryScriptingCore/Private/CollisionFunctions.cpp
--- a/Engine/Plugins/Runtime/GeometryScripting/Source/GeometryScriptingCore/Private/CollisionFunctions.cpp
+++ b/Engine/Plugins/Runtime/GeometryScripting/Source/GeometryScriptingCore/Private/CollisionFunctions.cpp
@@ -116,6 +116,7 @@ void ComputeCollisionFromMesh(
 	case EGeometryScriptCollisionGenerationMethod::ConvexHulls:
 		ShapeGenerator.bSimplifyHulls = Options.bSimplifyHulls;
 		ShapeGenerator.HullTargetFaceCount = Options.ConvexHullTargetFaceCount;
+		ShapeGenerator.HullGeometricToleranceAfterTarget = (double)Options.ConvexHullGeometricTolerance;
 		if (Options.MaxConvexHullsPerMesh > 1)
 		{
 			ShapeGenerator.ConvexDecompositionMaxPieces = Options.MaxConvexHullsPerMesh;
diff --git a/Engine/Plugins/Runtime/GeometryScripting/Source/GeometryScriptingCore/Private/ContainmentFunctions.cpp b/Engine/Plugins/Runtime/GeometryScripting/Source/GeometryScriptingCore/Private/ContainmentFunctions.cpp
--- a/Engine/Plugins/Runtime/GeometryScripting/Source/GeometryScriptingCore/Private/ContainmentFunctions.cpp
+++ b/Engine/Plugins/Runtime/GeometryScripting/Source/GeometryScriptingCore/Private/ContainmentFunctions.cpp
@@ -54,8 +54,9 @@ UDynamicMesh* UGeometryScriptLibrary_ContainmentFunctions::ComputeMeshConvexHull
 			FMeshConvexHull::GridSample(EditMesh, FMath::Max(32, Options.PrefilterGridResolution), Hull.VertexSet);
 		}
 
-		Hull.bPostSimplify = (Options.SimplifyToFaceCount > 4);
-		Hull.MaxTargetFaceCount = Options.SimplifyToFaceCount;
+		Hull.bPostSimplify = (Options.SimplifyToFaceCount > 0);
+		Hull.SimplifyOptions.MaxTargetFaceCount = Options.SimplifyToFaceCount;
+		Hull.SimplifyOptions.GeometricToleranceAfterTarget = Options.GeometricToleranceAfterTarget;
 		if (Hull.Compute(nullptr))
 		{
 			HullMesh = MoveTemp(Hull.ConvexHull);
diff --git a/Engine/Plugins/Runtime/GeometryScripting/Source/GeometryScriptingCore/Public/GeometryScript/CollisionFunctions.h b/Engine/Plugins/Runtime/GeometryScripting/Source/GeometryScriptingCore/Public/GeometryScript/CollisionFunctions.h
--- a/Engine/Plugins/Runtime/GeometryScripting/Source/GeometryScriptingCore/Public/GeometryScript/CollisionFunctions.h
+++ b/Engine/Plugins/Runtime/GeometryScripting/Source/GeometryScriptingCore/Public/GeometryScript/CollisionFunctions.h
@@ -69,6 +69,9 @@ public:
 	UPROPERTY(EditAnywhere, BlueprintReadWrite, Category = Options)
 	int ConvexHullTargetFaceCount = 25;
 
+	UPROPERTY(EditAnywhere, BlueprintReadWrite, Category = Options)
+	float ConvexHullGeometricTolerance = 0.f;
+
 	UPROPERTY(EditAnywhere, BlueprintReadWrite, Category = Options)
 	int MaxConvexHullsPerMesh = 1;
 
diff --git a/Engine/Plugins/Runtime/GeometryScripting/Source/GeometryScriptingCore/Public/GeometryScript/ContainmentFunctions.h b/Engine/Plugins/Runtime/GeometryScripting/Source/GeometryScriptingCore/Public/GeometryScript/ContainmentFunctions.h
--- a/Engine/Plugins/Runtime/GeometryScripting/Source/GeometryScriptingCore/Public/GeometryScript/ContainmentFunctions.h
+++ b/Engine/Plugins/Runtime/GeometryScripting/Source/GeometryScriptingCore/Public/GeometryScript/ContainmentFunctions.h
@@ -27,6 +27,10 @@ public:
 	/** Try to simplify each convex hull to this triangle count. If 0, no simplification */
 	UPROPERTY(EditAnywhere, BlueprintReadWrite, Category = Options)
 	int SimplifyToFaceCount = 0;
+
+	/** After SimplifyToFaceCount is applied, further simplify within this tolerance. If SimplifyToFaceCount <= 0, or this tolerance is <= 0, does not apply. */
+	UPROPERTY(EditAnywhere, BlueprintReadWrite, Category = Options)
+	float GeometricToleranceAfterTarget = 0.f;
 };
 
 
