fix comment to reference new simplify settings instead of deprecated member

diff --git a/Engine/Plugins/Runtime/GeometryProcessing/Source/DynamicMesh/Public/Operations/MeshConvexHull.h b/Engine/Plugins/Runtime/GeometryProcessing/Source/DynamicMesh/Public/Operations/MeshConvexHull.h
--- a/Engine/Plugins/Runtime/GeometryProcessing/Source/DynamicMesh/Public/Operations/MeshConvexHull.h
+++ b/Engine/Plugins/Runtime/GeometryProcessing/Source/DynamicMesh/Public/Operations/MeshConvexHull.h
@@ -29,7 +29,7 @@ public:
 	/** If set, hull will be computed on subset of vertices */
 	TArray<int32> VertexSet;
 
-	/** If true, output convex hull is simplified down to MaxTargetFaceCount */
+	/** If true, output convex hull is simplified down based on the SimplifyOptions */
 	bool bPostSimplify = false;
 
 	struct FSimplifyOptions
