Add mult-label version of mesh iso curves cutting, via adapters to support arbitrary mesh types

#Synced-CL 48266510

diff --git a/Engine/Plugins/Runtime/GeometryProcessing/Source/DynamicMesh/Private/Operations/MeshIsoCurves.cpp b/Engine/Plugins/Runtime/GeometryProcessing/Source/DynamicMesh/Private/Operations/MeshIsoCurves.cpp
--- a/Engine/Plugins/Runtime/GeometryProcessing/Source/DynamicMesh/Private/Operations/MeshIsoCurves.cpp
+++ b/Engine/Plugins/Runtime/GeometryProcessing/Source/DynamicMesh/Private/Operations/MeshIsoCurves.cpp
@@ -8,6 +8,405 @@
 
 using namespace UE::Geometry;
 
+namespace UE::Geometry::MeshIsoCurve
+{
+	// Given a triangle's vertices and its highest-weight labels, find the triangle's full set of weights for each of those labels
+	static TStaticArray<FVector3f, 3> GetFullTriLabelWeights(FIndex3i TriV, FIndex3i TriLabels, FVector3f BestWeights, const TFunction<float(int32 VID, int32 Label)>& InGetVertexLabelWeight)
+	{
+		TStaticArray<FVector3f, 3> LabelWeights;
+		for (int32 LabelIdx = 0; LabelIdx < 3; ++LabelIdx)
+		{
+			FVector3f& Wts = LabelWeights[LabelIdx];
+			int32 Label = TriLabels[LabelIdx];
+			for (int32 VertIdx = 0; VertIdx < 3; ++VertIdx)
+			{
+				if (VertIdx == LabelIdx)
+				{
+					Wts[VertIdx] = BestWeights[VertIdx];
+				}
+				else
+				{
+					Wts[VertIdx] = InGetVertexLabelWeight(TriV[VertIdx], Label);
+				}
+			}
+		}
+		return LabelWeights;
+	}
+
+	void FMultiLabelIsoCurveAdapter::SetDefaultFindCutFunctions(const TFunction<float(int32 VID, int32 Label)>& InGetVertexLabelWeight)
+	{
+		FindEdgeCutParam = [InGetVertexLabelWeight](FIndex2i EdgeV, FIndex2i EdgeLabels, FVector2f EdgeWeights) -> double
+			{
+				if (!ensure(EdgeLabels.A != EdgeLabels.B))
+				{
+					return .5;
+				}
+				// Get the edge weights that for the other (non-assigned, so lower) label on each vertex
+				FVector2f AltEdgeWeights;
+				for (int32 SubIdx = 0; SubIdx < 2; ++SubIdx)
+				{
+					AltEdgeWeights[SubIdx] = InGetVertexLabelWeight(EdgeV[SubIdx], EdgeLabels[1 - SubIdx]);
+					checkSlow(AltEdgeWeights[SubIdx] <= EdgeWeights[SubIdx]);
+				}
+				// Solve for the point on the edge where the label weights cross (are equal)
+				double Diff0 = EdgeWeights[0] - AltEdgeWeights[0];
+				double Denom = Diff0 + (EdgeWeights[1] - AltEdgeWeights[1]);
+				if (FMath::Abs(Denom) < UE_DOUBLE_SMALL_NUMBER)
+				{
+					return .5;
+				}
+				return FMath::Clamp(Diff0 / Denom, 0., 1.);
+			};
+		FindTriCutBaryCoords = [InGetVertexLabelWeight](FIndex3i TriV, FIndex3i TriLabels, FVector3f TriWeights) ->FVector3d
+			{
+				TStaticArray<FVector3f, 3> LabelWeights = GetFullTriLabelWeights(TriV, TriLabels, TriWeights, InGetVertexLabelWeight);
+				// Given separate weights per vertex for each label, LabelWeights[NumLabels] 
+				// we solve for the barycentric coordinates where the linearly-interpolated weights intersect at the same weight value
+				// 
+				// i.e., for TriWeights A,B,C and Barycentric coords W, we solve:
+				//  W.Dot(A) == W.Dot(B) -> W.Dot(A-B) == 0
+				//  W.Dot(B) == W.Dot(C) -> W.Dot(B-C) == 0
+				//  W.X + W.Y + W.Z == 1
+				// Expanding out w/ AmB = A-B and BmC = B-C:
+				//  W.X * AmB.X + W.Y * AmB.Y + W.Z * AmB.Z == 0, (and same eqn for BmC)
+				// Sub W.Z == 1 - W.X - W.Y, and re-distribute to get:
+				//  W.X * (AmB.X - AmB.Z) + W.Y * (AmB.Y - AmB.Z) = -AmB.Z
+				//  W.X * (BmC.X - BmC.Z) + W.Y * (BmC.Y - BmC.Z) = -BmC.Z
+				// We solve this 2x2 system w/ Cramer's rule. For small determinant, we fall back to just taking the centroid.
+
+				FVector3d AmB = FVector3d(LabelWeights[0] - LabelWeights[1]);
+				FVector3d BmC = FVector3d(LabelWeights[1] - LabelWeights[2]);
+				FMatrix2d M(
+					AmB.X - AmB.Z, AmB.Y - AmB.Z,
+					BmC.X - BmC.Z, BmC.Y - BmC.Z
+				);
+				double Det = M.Determinant();
+				
+				if (FMath::Abs(Det) < UE_DOUBLE_SMALL_NUMBER)
+				{
+					return FVector3d(1. / 3.);
+				}
+				FVector3d Res;
+				Res.X = (M.Row0.Y * BmC.Z - AmB.Z * M.Row1.Y) / Det;
+				Res.Y = (AmB.Z * M.Row1.X - M.Row0.X * BmC.Z) / Det;
+				Res.Z = 1 - Res.X - Res.Y;
+				return Res;
+			};
+	}
+
+	void MultiLabelCut(const FTriangleMeshAdapterd& Mesh, const FMultiLabelIsoCurveAdapter& IsoCurveAdapter, const FMeshUpdateAdapter& UpdateAdapter, const FMultiLabelCutSettings& Settings)
+	{
+		const double SnapThreshSq = Settings.SnapToExistingTolerance * Settings.SnapToExistingTolerance;
+
+		const int MaxVID = Mesh.MaxVertexID();
+		TArray<int32> VertexLabels;
+		TArray<float> VertexValues;
+		VertexLabels.SetNumUninitialized(MaxVID);
+		VertexValues.SetNumUninitialized(MaxVID);
+		// Pass 1: the best label and corresponding weight for each vertex
+		ParallelFor(MaxVID, [&Mesh, &IsoCurveAdapter, &Settings, &VertexLabels, &VertexValues](int32 VID)
+		{
+			if (Mesh.IsVertex(VID))
+			{
+				VertexLabels[VID] = IsoCurveAdapter.LabelVertex(VID, VertexValues[VID]);
+			}
+		});
+		
+		TMap<FIndex2i, TPair<int32,double>> EdgeCrossings;
+		struct FTriCutInfo
+		{
+			int32 TID;
+			FIndex3i EdgeCutVIDs;
+			FVector3d EdgeCutParams;
+			int32 PokeVID = INDEX_NONE;
+			FVector3d BaryCoords;
+			bool bSnapped = false;
+		};
+		TArray<FTriCutInfo> Cuts;
+
+		// Pass 2: Decide which tris to cut, and add new cut vertices as needed (but no triangles yet)
+		for (int32 TID = 0; TID < Mesh.MaxTriangleID(); ++TID)
+		{
+			if (!Mesh.IsTriangle(TID))
+			{
+				continue;
+			}
+
+			FIndex3i Tri = Mesh.GetTriangle(TID);
+			FIndex3i TriLabels{ VertexLabels[Tri.A], VertexLabels[Tri.B], VertexLabels[Tri.C] };
+			int32 NumCut = 0;
+			FIndex3i EdgeCutVIDs = FIndex3i::Invalid();
+			FVector3d EdgeCutParams(.5);
+			bool bHasSnapped = false;
+			if (Settings.bNeverCut)
+			{
+				bHasSnapped = TriLabels.A != TriLabels.B || TriLabels.B != TriLabels.C;
+			}
+			else
+			{
+				for (int32 PrevIdx = 2, SubIdx = 0; SubIdx < 3; PrevIdx = SubIdx++)
+				{
+
+					if (TriLabels[PrevIdx] != TriLabels[SubIdx] && (!IsoCurveAdapter.ShouldCutEdge || IsoCurveAdapter.ShouldCutEdge(TriLabels[PrevIdx], TriLabels[SubIdx])))
+					{
+						FIndex2i EdgeV(Tri[PrevIdx], Tri[SubIdx]);
+						FIndex2i EdgeL(TriLabels[PrevIdx], TriLabels[SubIdx]);
+						FVector2f EdgeW(VertexValues[EdgeV.A], VertexValues[EdgeV.B]);
+						bool bSwapped = false;
+						if (EdgeV.A > EdgeV.B)
+						{
+							bSwapped = true;
+							EdgeV.Swap();
+							EdgeL.Swap();
+							Swap(EdgeW.X, EdgeW.Y);
+						}
+						TPair<int32, double>* FoundCut = EdgeCrossings.Find(EdgeV);
+						int32 UseCutVID = INDEX_NONE;
+						double EdgeParam = .5f;
+						if (!FoundCut)
+						{
+							EdgeParam = IsoCurveAdapter.FindEdgeCutParam(EdgeV, EdgeL, EdgeW);
+							bool bSnap = false;
+							if (EdgeParam <= 0 || EdgeParam >= 1.)
+							{
+								bSnap = true;
+							}
+							else if (Settings.SnapToExistingTolerance > 0.)
+							{
+								FVector3d V[2];
+								V[0] = Mesh.GetVertex(EdgeV.A);
+								V[1] = Mesh.GetVertex(EdgeV.B);
+								FVector3d InterpV = FMath::LerpStable(V[0], V[1], EdgeParam);
+								bSnap = (FVector3d::DistSquared(V[0], InterpV) < SnapThreshSq)
+									|| (FVector3d::DistSquared(V[1], InterpV) < SnapThreshSq);
+							}
+							if (!bSnap)
+							{
+								UseCutVID = UpdateAdapter.AddInterpolatedEdgeVertex(EdgeV.A, EdgeV.B, EdgeParam);
+							}
+							else
+							{
+								bHasSnapped = true;
+							}
+							// We add the crossing even if we didn't add a vertex due to snapping, 
+							// so the same snapping will be picked up by any triangles sharing the edge
+							EdgeCrossings.Add(EdgeV, TPair<int32, double>(UseCutVID, EdgeParam));
+						}
+						else
+						{
+							if (FoundCut->Key == INDEX_NONE)
+							{
+								bHasSnapped = true;
+							}
+							else
+							{
+								UseCutVID = FoundCut->Key;
+								EdgeParam = FoundCut->Value;
+							}
+						}
+						if (UseCutVID != INDEX_NONE)
+						{
+							++NumCut;
+							EdgeCutVIDs[SubIdx] = UseCutVID;
+							EdgeCutParams[SubIdx] = bSwapped ? (1 - EdgeParam) : EdgeParam;
+						}
+					}
+				}
+			}
+			if (NumCut > 0)
+			{
+				FTriCutInfo& CutInfo = Cuts.AddDefaulted_GetRef();
+				if (NumCut == 3)
+				{
+					FVector3f TriWeights(VertexValues[Tri.A], VertexValues[Tri.B], VertexValues[Tri.C]);
+					CutInfo.BaryCoords = IsoCurveAdapter.FindTriCutBaryCoords(Tri, TriLabels, TriWeights);
+					// Do not poke if the barycentric coordinates are outside or very close to the boundary of the triangle
+					bool bSnapPoke = false;
+					if (CutInfo.BaryCoords.GetMin() < UE_DOUBLE_SMALL_NUMBER)
+					{
+						bSnapPoke = true;
+					}
+					else if (Settings.SnapToExistingTolerance > 0.)
+					{
+						FTriangle3d TriVerts(Mesh.GetVertex(Tri.A), Mesh.GetVertex(Tri.B), Mesh.GetVertex(Tri.C));
+						FVector3d InterpV = TriVerts.BarycentricPoint(CutInfo.BaryCoords);
+						for (int32 Idx = 0, Prev = 2; Idx < 3; Prev = Idx++)
+						{
+							if ((double)FMath::PointDistToSegmentSquared(InterpV, TriVerts.V[Prev], TriVerts.V[Idx]) < SnapThreshSq)
+							{
+								bSnapPoke = true;
+								break;
+							}
+						} 
+					}
+					if (!bSnapPoke)
+					{
+						CutInfo.PokeVID = UpdateAdapter.AddInterpolatedTriangleVertex(Tri, CutInfo.BaryCoords);
+					}
+					
+				}
+				CutInfo.TID = TID;
+				CutInfo.EdgeCutVIDs = EdgeCutVIDs;
+				CutInfo.EdgeCutParams = EdgeCutParams;
+				CutInfo.bSnapped = bHasSnapped;
+			}
+			else
+			{
+				if (bHasSnapped) // tri verts were multi-label, but all cuts were snapped away
+				{
+					FVector3f BestTriWeights(VertexValues[Tri.A], VertexValues[Tri.B], VertexValues[Tri.C]);
+					TStaticArray<FVector3f, 3> FullTriWeights = GetFullTriLabelWeights(Tri, TriLabels, BestTriWeights, IsoCurveAdapter.GetVertexLabelWeight);
+					const FVector3f MidBary(1.f / 3.f);
+					int32 LabelIdx = FMath::Max3Index(FullTriWeights[0].Dot(MidBary), FullTriWeights[1].Dot(MidBary), FullTriWeights[2].Dot(MidBary));
+					UpdateAdapter.LabelTriangle(TID, TriLabels[LabelIdx]);
+				}
+				else // tri is cleanly a single label
+				{
+					UpdateAdapter.LabelTriangle(TID, TriLabels[0]);
+				}
+			}
+		}
+
+
+		// Pass 3: Add triangles as needed for cut triangles
+		TArray<FIndex3i, TInlineAllocator<6>> NewTris;
+		TArray<int32, TInlineAllocator<6>> NewTriLabels;
+		for (FTriCutInfo& CutInfo : Cuts)
+		{
+			int32 NumCut = 0;
+			int32 FirstCut = INDEX_NONE, FirstNoCut = INDEX_NONE;
+			NewTris.Reset();
+			NewTriLabels.Reset();
+			FIndex3i Tri = Mesh.GetTriangle(CutInfo.TID);
+			
+			for (int32 SubIdx = 0; SubIdx < 3; ++SubIdx)
+			{
+				if (CutInfo.EdgeCutVIDs[SubIdx] != INDEX_NONE)
+				{
+					if (FirstCut == INDEX_NONE)
+					{
+						FirstCut = SubIdx;
+					}
+					++NumCut;
+				}
+				else if (FirstNoCut == INDEX_NONE)
+				{
+					FirstNoCut = SubIdx;
+				}
+			}
+			if (!ensure(NumCut > 0)) // should not have added to CutInfo if nothing to cut
+			{
+				continue;
+			}
+
+			TStaticArray<FVector3f, 3> FullTriWeights;
+			if (CutInfo.bSnapped || (CutInfo.PokeVID == INDEX_NONE && NumCut == 3))
+			{
+				FVector3f BestTriWeights(VertexValues[Tri.A], VertexValues[Tri.B], VertexValues[Tri.C]);
+				FIndex3i TriLabels(VertexLabels[Tri.A], VertexLabels[Tri.B], VertexLabels[Tri.C]);
+				FullTriWeights = GetFullTriLabelWeights(Tri, TriLabels, BestTriWeights, IsoCurveAdapter.GetVertexLabelWeight);
+			}
+
+			// mapping from local indexing to actual mesh vertex IDs
+			// Original triangle vertices: 0,1,2; Triangle edge cut vertices: 3,4,5; Poke vertex: 6
+			const int32 LocalToGlobal[7]{ Tri[0], Tri[1], Tri[2], CutInfo.EdgeCutVIDs[0], CutInfo.EdgeCutVIDs[1], CutInfo.EdgeCutVIDs[2], CutInfo.PokeVID };
+			// Offsets for local indexing -- EOff for edge cuts, TOff for the triangle 'poke' vertex
+			constexpr int32 EOff = 3;
+			constexpr int32 TOff = 6;
+
+			auto AddCutTri = [&VertexLabels, &NewTris, &NewTriLabels, &CutInfo, &FullTriWeights, &LocalToGlobal, &Tri, EOff, TOff]
+				(FIndex3i LocalTriIndices, int32 DefaultLabel, bool bComputeLabel) -> void
+			{
+				FIndex3i NewTri(
+					LocalToGlobal[LocalTriIndices[0]],
+					LocalToGlobal[LocalTriIndices[1]],
+					LocalToGlobal[LocalTriIndices[2]]);
+				NewTris.Add(NewTri);
+				// for cases where some snapping was applied, there are effectively multiple labels on some vertices
+				// so rather than copy vertex labels directly, we compute a best label from the triangle centroid
+				if (bComputeLabel)
+				{
+					// Get the centroid of the new triangle, expressed as barycentric coordinates of the original triangle
+					FVector3d BarySum(0.);
+					for (int32 SubIdx = 0; SubIdx < 3; ++SubIdx)
+					{
+						int32 LocalIdx = LocalTriIndices[SubIdx];
+						FVector3d Bary(0.);
+						if (LocalIdx < EOff)
+						{
+							Bary[LocalIdx] = 1.;
+						}
+						else if (LocalIdx < TOff)
+						{
+							double Param = CutInfo.EdgeCutParams[LocalIdx - 3];
+							Bary[LocalIdx - 3] = Param;
+							Bary[(LocalIdx + 2) % 3] = 1. - Param;
+						}
+						else
+						{
+							Bary = CutInfo.BaryCoords;
+						}
+						BarySum += Bary;
+					}
+					// Choose the label w/ the highest interpolated weight on the original triangle, at the new triangle's centroid
+					FVector3f MidBary = FVector3f(BarySum / 3.);
+					int32 LabelIdx = FMath::Max3Index(FullTriWeights[0].Dot(MidBary), FullTriWeights[1].Dot(MidBary), FullTriWeights[2].Dot(MidBary));
+					NewTriLabels.Add(VertexLabels[Tri[LabelIdx]]);
+				}
+				// for cases w/out snapping, we just copy an appropriate vertex label
+				else
+				{
+					NewTriLabels.Add(DefaultLabel);
+				}
+			};
+
+			if (NumCut == 1)
+			{
+				int32 Shift[3]{ FirstCut, (FirstCut + 1) % 3, (FirstCut + 2) % 3 };
+				int32 Label = VertexLabels[Tri[Shift[1]]];
+				int32 CutV = CutInfo.EdgeCutVIDs[FirstCut];
+				AddCutTri(FIndex3i(EOff + FirstCut, Shift[0], Shift[1]), Label, CutInfo.bSnapped);
+				AddCutTri(FIndex3i(EOff + FirstCut, Shift[1], Shift[2]), Label, CutInfo.bSnapped);
+			}
+			else if (NumCut == 2)
+			{
+				int32 ShiftN[3]{ FirstNoCut, (FirstNoCut + 1) % 3, (FirstNoCut + 2) % 3 };
+				int32 Label0 = VertexLabels[Tri[ShiftN[0]]];
+				int32 Label1 = VertexLabels[Tri[ShiftN[1]]];
+				AddCutTri(FIndex3i(ShiftN[0], EOff + ShiftN[1], EOff + ShiftN[2]), Label0, CutInfo.bSnapped);
+				AddCutTri(FIndex3i(EOff + ShiftN[2], ShiftN[2], ShiftN[0]), Label0, CutInfo.bSnapped);
+				AddCutTri(FIndex3i(EOff + ShiftN[1], ShiftN[1], EOff + ShiftN[2]), Label1, CutInfo.bSnapped);
+			}
+			else // NumCut == 3
+			{
+				// We didn't add a new 'poke' vertex inside the triangle, so use a 1->4 tri split
+				if (CutInfo.PokeVID == INDEX_NONE)
+				{
+					AddCutTri(FIndex3i(EOff + 0, EOff + 1, EOff + 2), 0, true);
+					for (int32 Idx = 0, Next = 1, Prev = 2; Idx < 3; Next = Prev, Prev = Idx++)
+					{
+						int32 Label = VertexLabels[Tri[Idx]];
+						AddCutTri(FIndex3i(EOff + Idx, Idx, EOff + Next), Label, CutInfo.bSnapped);
+					}
+				}
+				else
+				{
+					// 1->6 tri split, making a fan from the new 'poke' vertex
+					for (int32 Idx = 0, Next = 1, Prev = 2; Idx < 3; Next = Prev, Prev = Idx++)
+					{
+						int32 Label = VertexLabels[Tri[Idx]];
+						AddCutTri(FIndex3i(EOff + Idx, Idx, TOff), Label, CutInfo.bSnapped);
+						AddCutTri(FIndex3i(Idx, EOff + Next, TOff), Label, CutInfo.bSnapped);
+					}
+				}
+			}
+
+			check(NewTris.Num() == NewTriLabels.Num());
+			UpdateAdapter.ReplaceTriangle(CutInfo.TID, NewTris, NewTriLabels);
+		}
+	}
+}
+
 void FMeshIsoCurves::Cut(FDynamicMesh3& Mesh, TFunctionRef<float(int32)> VertexFn, TFunctionRef<float(int32 VIDMin, int32 VIDMax, float ValMin, float ValMax)> EdgeCutFn, float IsoValue)
 {
 	int MaxVID = Mesh.MaxVertexID();
diff --git a/Engine/Plugins/Runtime/GeometryProcessing/Source/DynamicMesh/Public/Operations/MeshIsoCurves.h b/Engine/Plugins/Runtime/GeometryProcessing/Source/DynamicMesh/Public/Operations/MeshIsoCurves.h
--- a/Engine/Plugins/Runtime/GeometryProcessing/Source/DynamicMesh/Public/Operations/MeshIsoCurves.h
+++ b/Engine/Plugins/Runtime/GeometryProcessing/Source/DynamicMesh/Public/Operations/MeshIsoCurves.h
@@ -6,6 +6,8 @@
 #include "VectorTypes.h"
 #include "GeometryTypes.h"
 
+#include "MeshAdapter.h"
+
 #define UE_API DYNAMICMESH_API
 
 namespace UE
@@ -15,6 +17,70 @@ namespace Geometry
 
 class FDynamicMesh3;
 
+// Methods for cutting meshes along isolines
+namespace MeshIsoCurve
+{
+	// Adapter methods to support updating a mesh with a cut.
+	// All triangles will either be labelled via LabelTriangle, or replaced w/ labelled triangles via ReplaceTriangle.
+	struct FMeshUpdateAdapter
+	{
+		// Add a vertex interpolating the two input vertices (but not yet connected to any triangles)
+		TFunction<int32(int32 VID0, int32 VID1, double T)> AddInterpolatedEdgeVertex;
+		// Add a vertex interpolating the three input vertices (but not yet connected to any triangles)
+		TFunction<int32(FIndex3i Tri, const FVector3d& BaryCoords)> AddInterpolatedTriangleVertex;
+		// Replace a triangle with the given new triangles, with the specified new per-triangle labels
+		TFunction<void(int32 ReplaceTriID, TConstArrayView<FIndex3i> NewTris, TConstArrayView<int32> NewTriLabels)> ReplaceTriangle;
+		// Set label on a triangle
+		TFunction<void(int32 TID, int32 Label)> LabelTriangle;
+	};
+
+	// Adapter methods to help define the iso-curves where a triangle mesh changes between multiple discrete labels
+	// where the labels are expressed as a weight per vertex, and the highest weight should be chosen
+	struct FMultiLabelIsoCurveAdapter
+	{
+		// This function should return the 'class' index with the highest weight for the given vertex, and also return that weight by reference
+		TFunction<int32(int32 VID, float& OutLabelWeight)> LabelVertex;
+		// This function should return the weight for a given vertex and class index
+		TFunction<float(int32 VID, int32 Label)> GetVertexLabelWeight;
+
+		// Find the cut parameter along a given edge where the label should change
+		TFunction<double(FIndex2i EdgeV, FIndex2i EdgeLabels, FVector2f EdgeWeights)> FindEdgeCutParam;
+		// For a triangle with different labels on each vertex, find the barycentric coordinates where those labels should meet
+		// Note: This method is allowed to return a coordinate outside of the triangle (should be handled gracefully by the cutting method)
+		TFunction<FVector3d(FIndex3i TriV, FIndex3i TriLabels, FVector3f TriWeights)> FindTriCutBaryCoords;
+
+		// Set default implementations for FindEdgeCutParam and FindTriCutBaryCoords
+		UE_API void SetDefaultFindCutFunctions(const TFunction<float(int32 VID, int32 Label)>& InGetVertexLabelWeight);
+
+		// Return true if an edge with these two vertex labels should be cut.
+		// Note: Will not be called if the labels are the same.
+		// If not set, edges with different labels will always be cut.
+		TFunction<bool(int32 Vert0Label, int32 Vert1Label)> ShouldCutEdge;
+	};
+
+	struct FMultiLabelCutSettings
+	{
+		// Distance tolerance at which a cut vertex is too close to an existing vertex (or, for an on-face insertion, an existing edge) and should not be inserted
+		double SnapToExistingTolerance = UE_DOUBLE_KINDA_SMALL_NUMBER;
+
+		// Option to never cut triangles and only label triangles; equivalent to setting the snap tolerance high enough to always snap
+		bool bNeverCut = false;
+	};
+
+	/**
+	 * Algorithm to cut a mesh's triangles based on per-vertex labels, and transfer those labels to the mesh triangles.
+	 * 
+	 * Note this will not introduce new boundaries along the cuts; the resulting mesh will just have new edges along the boundaries.
+	 * 
+	 * @param InMesh The mesh to cut. Note only original mesh vertices will be queried via this adapter -- never added vertices and triangles.
+	 * @param IsoCurveAdapter An adapter defining how vertices should be labelled, which edges should be cut, and the desired cut positions
+	 * @param UpdateAdapter An adapter that update the mesh (or a new target mesh) with newly cut/labelled triangles
+	 * @param Settings Additional settings controlling how the cuts are made
+	 */
+	UE_API void MultiLabelCut(const FTriangleMeshAdapterd& InMesh, const FMultiLabelIsoCurveAdapter& IsoCurveAdapter, const FMeshUpdateAdapter& UpdateAdapter, const FMultiLabelCutSettings& Settings);
+}
+
+// Dynamic Mesh IsoCurve Settings
 struct FMeshIsoCurveSettings
 {
 	// Whether to collapse any degenerate edges created by the curve insertion
@@ -24,14 +90,15 @@ struct FMeshIsoCurveSettings
 	double DegenerateEdgeTol = FMathd::ZeroTolerance;
 
 	// Distance at which to snap curve vertices to nearby existing vertices
-	double SnapToExistingVertexTol = UE_KINDA_SMALL_NUMBER;
+	double SnapToExistingVertexTol = UE_DOUBLE_KINDA_SMALL_NUMBER;
 
 	// Tolerance distance (in function domain) to an existing vertex to be 'on curve'
 	float CurveIsoValueSnapTolerance = 0.f;
 };
 
+
 /**
- * Insert edges on a mesh along the isocurve where some scalar value function over the mesh surface crosses a specified value
+ * Insert edges on a dynamic mesh along the isocurve where some scalar value function over the mesh surface crosses a specified value
  */
 class FMeshIsoCurves
 {
