GeometryProcessing: Optimize FMeshRegionGraph::OptimizeBorders

#Synced-CL 48742553

diff --git a/Engine/Plugins/Runtime/GeometryProcessing/Source/DynamicMesh/Private/Parameterization/MeshRegionGraph.cpp b/Engine/Plugins/Runtime/GeometryProcessing/Source/DynamicMesh/Private/Parameterization/MeshRegionGraph.cpp
--- a/Engine/Plugins/Runtime/GeometryProcessing/Source/DynamicMesh/Private/Parameterization/MeshRegionGraph.cpp
+++ b/Engine/Plugins/Runtime/GeometryProcessing/Source/DynamicMesh/Private/Parameterization/MeshRegionGraph.cpp
@@ -15,8 +15,8 @@ void FMeshRegionGraph::BuildFromComponents(
 {
 	int32 N = Components.Num();
 	Regions.SetNum(N);
-	TriangleToRegionMap.Init(-1, Mesh.MaxTriangleID());
-	TriangleNbrTris.Init(FIndex3i(-1, -1, -1), Mesh.MaxTriangleID());
+	TriangleToRegionMap.Init(INDEX_NONE, Mesh.MaxTriangleID());
+	TriangleNbrTris.Init(FIndex3i(INDEX_NONE, INDEX_NONE, INDEX_NONE), Mesh.MaxTriangleID());
 
 	for (int32 k = 0; k < N; ++k)
 	{
@@ -31,7 +31,7 @@ void FMeshRegionGraph::BuildFromComponents(
 			{
 				if (NbrTris[j] >= 0 && TrisConnectedPredicate(tid, NbrTris[j]) == false)
 				{
-					NbrTris[j] = -1;
+					NbrTris[j] = INDEX_NONE;
 				}
 			}
 			TriangleNbrTris[tid] = NbrTris;
@@ -52,8 +52,8 @@ void FMeshRegionGraph::BuildFromTriangleSets(const FDynamicMesh3& Mesh,
 {
 	int32 N = TriangleSets.Num();
 	Regions.SetNum(N);
-	TriangleToRegionMap.Init(-1, Mesh.MaxTriangleID());
-	TriangleNbrTris.Init(FIndex3i(-1, -1, -1), Mesh.MaxTriangleID());
+	TriangleToRegionMap.Init(INDEX_NONE, Mesh.MaxTriangleID());
+	TriangleNbrTris.Init(FIndex3i(INDEX_NONE, INDEX_NONE, INDEX_NONE), Mesh.MaxTriangleID());
 
 	for (int32 k = 0; k < N; ++k)
 	{
@@ -68,7 +68,7 @@ void FMeshRegionGraph::BuildFromTriangleSets(const FDynamicMesh3& Mesh,
 			{
 				if (NbrTris[j] >= 0 && TrisConnectedPredicate(tid, NbrTris[j]) == false)
 				{
-					NbrTris[j] = -1;
+					NbrTris[j] = INDEX_NONE;
 				}
 			}
 			TriangleNbrTris[tid] = NbrTris;
@@ -151,7 +151,7 @@ bool FMeshRegionGraph::MergeSmallRegionsPass(int32 SmallThreshold,
 	struct FMatch 
 	{
 		float Score = 0.0f;
-		int Index = -1;
+		int Index = INDEX_NONE;
 	};
 
 	for (int32 j = 0; j < SmallRegions.Num(); ++j)
@@ -326,7 +326,7 @@ bool FMeshRegionGraph::OptimizeBorders(int32 MaxRounds)
 		int32 Index = TriangleToRegionMap[tid];
 		FIndex3i NbrTris = TriangleNbrTris[tid];
 		int32 NbrCount = 0;
-		TArray<int32, TFixedAllocator<3>> UniqueNbrs;
+		int32 Nbrs[3];
 		for (int32 j = 0; j < 3; ++j)
 		{
 			if (NbrTris[j] >= 0)
@@ -334,55 +334,68 @@ bool FMeshRegionGraph::OptimizeBorders(int32 MaxRounds)
 				int32 NbrRegionIndex = TriangleToRegionMap[NbrTris[j]];
 				if (IsRegion(NbrRegionIndex) && NbrRegionIndex != Index)
 				{
-					NbrCount++;
-					UniqueNbrs.AddUnique(NbrRegionIndex);
+					Nbrs[NbrCount++] = NbrRegionIndex;
 				}
 			}
 		}
-		if (NbrCount >= 2 && UniqueNbrs.Num() == 1)
+		if ((NbrCount == 2 && Nbrs[0] == Nbrs[1]) || (NbrCount == 3 && Nbrs[0] == Nbrs[1] && Nbrs[0] == Nbrs[2]))
 		{
-			return UniqueNbrs[0];
+			return Nbrs[0];
 		}
-		return -1;
+
+		return INDEX_NONE;
 	};
 
 	bool bModified = false;
 
 	bool bDone = false;
 	int32 RoundCounter = 0;
+	TArray<FIndex2i> PotentialSwaps;
 	while (bDone == false && RoundCounter++ < MaxRounds)
 	{
-		TArray<FIndex2i> PotentialSwaps;
-		for (const FRegion& Region : Regions)
+		PotentialSwaps.Reset();
+
+		for (int32 tid = 0, Num = TriangleToRegionMap.Num(); tid < Num; ++tid)
 		{
-			for (int32 tid : Region.Triangles)
+			if (TriangleToRegionMap[tid] != INDEX_NONE)
 			{
-				int32 SwapToNbrIndex = GetSwapNbrRegionIndex(tid);
+				const int32 SwapToNbrIndex = GetSwapNbrRegionIndex(tid);
 				if (IsRegion(SwapToNbrIndex))
 				{
-					PotentialSwaps.Add( FIndex2i(tid, SwapToNbrIndex) );
+					PotentialSwaps.Add(FIndex2i(tid, SwapToNbrIndex));
 				}
 			}
 		}
 		bDone = (PotentialSwaps.Num() == 0);
 
-		for (const FIndex2i& Swap : PotentialSwaps)
+		for (const FIndex2i Swap : PotentialSwaps)
 		{
-			int32 tid = Swap.A;
-			int32 Index = TriangleToRegionMap[tid];
-			int32 SwapToNbrIndex = Swap.B;
-
-			int32 CheckNbrIndex = GetSwapNbrRegionIndex(tid);
-			if (CheckNbrIndex == SwapToNbrIndex)
+			if (GetSwapNbrRegionIndex(Swap.A) == Swap.B)
 			{
-				Regions[Index].Triangles.RemoveSwap(tid, EAllowShrinking::No);
-				Regions[SwapToNbrIndex].Triangles.Add(tid);
-				TriangleToRegionMap[tid] = SwapToNbrIndex;
+				TriangleToRegionMap[Swap.A] = Swap.B;
 				bModified = true;
 			}
 		}
 	}
 
+	if (bModified)
+	{
+		// Rebuild region triangles
+
+		for (FRegion& Region : Regions)
+		{
+			Region.Triangles.Reset();
+		}
+
+		for (int32 tid = 0, Num = TriangleToRegionMap.Num(); tid < Num; ++tid)
+		{
+			if (TriangleToRegionMap[tid] != INDEX_NONE)
+			{
+				Regions[TriangleToRegionMap[tid]].Triangles.Add(tid);
+			}
+		}
+	}
+
 	return bModified;
 }
 
@@ -412,7 +425,7 @@ void FMeshRegionGraph::BuildNeigbours(int32 RegionIdx)
 			else
 			{
 				int32 NbrRegionIndex = TriangleToRegionMap[Nbrs[j]];
-				if (NbrRegionIndex == -1)
+				if (NbrRegionIndex == INDEX_NONE)
 				{
 					Region.bIsOnROIBoundary = true;
 				}
