improve edge collapse parameter calc for dynamic mesh's mesh simplification

#Synced-CL 48949184

diff --git a/Engine/Plugins/Runtime/GeometryProcessing/Source/DynamicMesh/Private/MeshSimplification.cpp b/Engine/Plugins/Runtime/GeometryProcessing/Source/DynamicMesh/Private/MeshSimplification.cpp
--- a/Engine/Plugins/Runtime/GeometryProcessing/Source/DynamicMesh/Private/MeshSimplification.cpp
+++ b/Engine/Plugins/Runtime/GeometryProcessing/Source/DynamicMesh/Private/MeshSimplification.cpp
@@ -1419,6 +1419,7 @@ ESimplificationResult TMeshSimplification<QuadricErrorType>::CollapseEdge(int ed
 	}
 
 	ProfileBeginCollapse();
+	ON_SCOPE_EXIT { ProfileEndCollapse(); };
 
 	// check if we should collapse, and also find which vertex we should retain
 	// in cases where we have constraints/etc
@@ -1487,18 +1488,17 @@ ESimplificationResult TMeshSimplification<QuadricErrorType>::CollapseEdge(int ed
 	if (!bConstraintsSpecifyPosition)
 	{
 		checkSlow(!bMinimalVertexMode || (vNewPos == vA || vNewPos == vB));
-		
-		// [TODO] maybe skip Projection call when !bModeAllowsVertMovment.
-		vNewPos = GetProjectedCollapsePosition(iKeep, vNewPos);
-		double EdgeLength = Distance(vA, vB);
-		collapse_t = (EdgeLength < FMathd::ZeroTolerance) ? 0.5 : (Distance(vNewPos, Mesh->GetVertex(iKeep))) / EdgeLength;
-		collapse_t = VectorUtil::Clamp(collapse_t, 0.0, 1.0);
 
-		if (bMinimalVertexMode)
+		FVector3d VCollapse = Mesh->GetVertex(iCollapse), VKeep = Mesh->GetVertex(iKeep);
+		if (!bMinimalVertexMode)
+		{
+			vNewPos = GetProjectedCollapsePosition(iKeep, vNewPos);
+			double SegLenSq_Unused;
+			collapse_t = CurveUtil::ProjectToSegment(VKeep, VCollapse, SegLenSq_Unused, vNewPos);
+		}
+		else
 		{
-			// this _should_ already be 0 or 1 with perfect precision. 
-			// round here to make sure later attribute lerps don't change values
-			collapse_t = FMath::RoundToDouble(collapse_t);
+			collapse_t = FVector3d::DistSquared(vNewPos, VKeep) < FVector3d::DistSquared(vNewPos, VCollapse) ? 0. : 1.;
 		}
 
 		// If vertex is not explicitly constrained, and geometric error constraint is requested, we check it here.
@@ -1509,11 +1509,14 @@ ESimplificationResult TMeshSimplification<QuadricErrorType>::CollapseEdge(int ed
 		{
 			if (CheckIfCollapseWithinGeometricTolerance(iKeep, iCollapse, vNewPos, t0, t1) == false)
 			{
+				if (bMinimalVertexMode)
+				{
+					return ESimplificationResult::Failed_GeometricDeviation;
+				}
 				// project new position back onto the edge
 				vNewPos = (1.0-collapse_t)*Mesh->GetVertex(iKeep) + (collapse_t)*Mesh->GetVertex(iCollapse);
 				if (CheckIfCollapseWithinGeometricTolerance(iKeep, iCollapse, vNewPos, t0, t1) == false)
 				{
-					ProfileEndCollapse();
 					return ESimplificationResult::Failed_GeometricDeviation;
 				}
 			}
@@ -1530,7 +1533,6 @@ ESimplificationResult TMeshSimplification<QuadricErrorType>::CollapseEdge(int ed
 		{
 			if (CheckIfCollapseWithinGeometricTolerance(iKeep, iCollapse, vNewPos, t0, t1) == false)
 			{
-				ProfileEndCollapse();
 				return ESimplificationResult::Failed_GeometricDeviation;
 			}
 		}
@@ -1541,13 +1543,11 @@ ESimplificationResult TMeshSimplification<QuadricErrorType>::CollapseEdge(int ed
 	// [TODO] could we skip this one-ring check in CollapseEdge? pass in hints?
 	if (CheckIfCollapseCreatesFlipOrInvalid(a, b, vNewPos, t0, t1) || CheckIfCollapseCreatesFlipOrInvalid(b, a, vNewPos, t0, t1))
 	{
-		ProfileEndCollapse();
 		return ESimplificationResult::Ignored_CreatesFlip;
 	}
 
 	if (bPreventTinyTriangles && (CheckIfCollapseCreatesTinyTriangle(a, b, vNewPos, t0, t1) || CheckIfCollapseCreatesTinyTriangle(b, a, vNewPos, t0, t1)))
 	{
-		ProfileEndCollapse();
 		return ESimplificationResult::Ignored_CreatesTinyTriangle;
 	}
 
@@ -1635,7 +1635,6 @@ ESimplificationResult TMeshSimplification<QuadricErrorType>::CollapseEdge(int ed
 		retVal = ESimplificationResult::Failed_IsolatedTriangle;
 	}
 
-	ProfileEndCollapse();
 	return retVal;
 }
 
