Fix MeshSimplification CanCollapseEdge sometimes collapsing too much and breaking seam topology + also enable additional collapses in a common special case

+ Add edge-constraint tolerance distance to limit how far constrained edges can move during simplification

#Synced-CL 48949184

diff --git a/Engine/Plugins/Runtime/GeometryProcessing/Source/DynamicMesh/Private/MeshConstraintsUtil.cpp b/Engine/Plugins/Runtime/GeometryProcessing/Source/DynamicMesh/Private/MeshConstraintsUtil.cpp
--- a/Engine/Plugins/Runtime/GeometryProcessing/Source/DynamicMesh/Private/MeshConstraintsUtil.cpp
+++ b/Engine/Plugins/Runtime/GeometryProcessing/Source/DynamicMesh/Private/MeshConstraintsUtil.cpp
@@ -306,7 +306,8 @@ void FMeshConstraintsUtil::SetBoundaryConstraintsWithProjection(
 	FMeshConstraints& Constraints,
 	const EBoundaryType BoundaryToConstrain,
 	const FDynamicMesh3& Mesh,
-	double CornerAngleThreshold)
+	double CornerAngleThreshold,
+	float CurveConstraintToleranceDistance)
 {
 	TSet<int32> GroupCorners;
 
@@ -337,6 +338,25 @@ void FMeshConstraintsUtil::SetBoundaryConstraintsWithProjection(
 
 	Constraints.ProjectionData.ProjectionCurves.SetNum(NewNumCurves);
 
+
+	float EdgeToleranceDistSq = CurveConstraintToleranceDistance > 0 ? CurveConstraintToleranceDistance * CurveConstraintToleranceDistance : -1.f;
+	// helper to update edge tolerance distances
+	auto MergeEdgeToleranceDistSq = [](float A, float B) -> float
+	{
+		if (A <= 0) // if A is invalid, return B
+		{
+			return B;
+		}
+		else if (B <= 0) // if B is invalid, return A
+		{
+			return A;
+		}
+		else // if both valid, return the smallest tolerance
+		{
+			return FMath::Min(A, B);
+		}
+	};
+
 	for (int32 TopologyEdgeIndex = 0; TopologyEdgeIndex < NumEdges; ++TopologyEdgeIndex)
 	{
 		const FGroupTopology::FGroupEdge& GroupEdge = Topology->Edges[TopologyEdgeIndex];
@@ -408,11 +428,13 @@ void FMeshConstraintsUtil::SetBoundaryConstraintsWithProjection(
 			if (Constraints.GetEdgeConstraint(EdgeID, FoundEdgeConstraint))
 			{
 				FoundEdgeConstraint.Target = Curve.Get();
+				FoundEdgeConstraint.DistanceToleranceSq = MergeEdgeToleranceDistSq(EdgeToleranceDistSq, FoundEdgeConstraint.DistanceToleranceSq);
 				Constraints.SetOrUpdateEdgeConstraint(EdgeID, FoundEdgeConstraint);
 			}
 			else
 			{
 				FEdgeConstraint GroupEdgeConstraint(EEdgeRefineFlags::NoConstraint, Curve.Get());
+				GroupEdgeConstraint.DistanceToleranceSq = EdgeToleranceDistSq;
 				Constraints.SetOrUpdateEdgeConstraint(EdgeID, GroupEdgeConstraint);
 			}
 		}
diff --git a/Engine/Plugins/Runtime/GeometryProcessing/Source/DynamicMesh/Private/MeshSimplification.cpp b/Engine/Plugins/Runtime/GeometryProcessing/Source/DynamicMesh/Private/MeshSimplification.cpp
--- a/Engine/Plugins/Runtime/GeometryProcessing/Source/DynamicMesh/Private/MeshSimplification.cpp
+++ b/Engine/Plugins/Runtime/GeometryProcessing/Source/DynamicMesh/Private/MeshSimplification.cpp
@@ -1280,19 +1280,142 @@ void TMeshSimplification<QuadricErrorType>::FastCollapsePass(double fMinEdgeLeng
 
 
 template <typename QuadricErrorType>
-bool TMeshSimplification<QuadricErrorType>::CanCollapseEdge(int edgeID, int a, int b, int c, int d, int t0, int t1, int& collapse_to) const
+bool TMeshSimplification<QuadricErrorType>::CanCollapseEdge(int edgeID, int a, int b, int c, int d, int tc, int td, int& collapse_to, FCollapseInfo* OutCollapseInfo) const
 {
-	const bool bPreserveSeamTopology = (CollapseMode == ESimplificationCollapseModes::MinimalExistingVertexError);
+	collapse_to = -1;
+	if (!Constraints)
+	{
+		return true;
+	}
+
+	FVertexConstraint ca = Constraints->GetVertexConstraint(a);
+	FVertexConstraint cb = Constraints->GetVertexConstraint(b);
+
+	// Use the MeshRefinerBase::CanCollapseVertex method to test if the vertices are too constrained to collapse
+	// and if there is a restriction on which vertices must be kept/removed
+	bool bCanCollapse = CanCollapseVertex(edgeID, a, b, collapse_to);
+
+	// Special handling for the case where we have two same-target constrained vertices connecting an unconstrained edge
+	// Then we *can* safely collapse if the AC CB xor AD DB edges match that same target, and the shared C/D vertex is otherwise-unconstrained ...
+	// as long as we track that the shared C/D vertex and matched edges should have their constraints removed post collapse (tracked via OutCollapseInfo)
+	bool bIndirectConstraintConnection = false;
+	FEdgeConstraint EdgeConstraint = Constraints->GetEdgeConstraint(edgeID);
+	if (OutCollapseInfo && !bCanCollapse && ca.Target != nullptr && ca.Target == cb.Target && (ca.bCanMove || cb.bCanMove) && EdgeConstraint.IsUnconstrained())
+	{
+		int32 eac = Mesh->FindEdgeFromTri(a, c, tc);
+		int32 ebc = Mesh->FindEdgeFromTri(b, c, tc);
+		bool bMatchTargetsThroughACB =
+			Constraints->GetEdgeConstraint(eac).Target == ca.Target &&
+			Constraints->GetEdgeConstraint(ebc).Target == ca.Target;
+		bool bMatchTargetsThroughADB = false;
+		int32 ead = IndexConstants::InvalidID, ebd = IndexConstants::InvalidID;
+		if (d != IndexConstants::InvalidID)
+		{
+			ead = Mesh->FindEdgeFromTri(a, d, td);
+			ebd = Mesh->FindEdgeFromTri(b, d, td);
+			bMatchTargetsThroughADB = 
+				Constraints->GetEdgeConstraint(ead).Target == ca.Target &&
+				Constraints->GetEdgeConstraint(ebd).Target == ca.Target;
+		}
+		if (bMatchTargetsThroughACB != bMatchTargetsThroughADB)
+		{
+			int32 PathVID, PathEA, PathEB;
+			if (bMatchTargetsThroughACB)
+			{
+				PathVID = c;
+				PathEA = eac;
+				PathEB = ebc;
+			}
+			else
+			{
+				PathVID = d;
+				PathEA = ead;
+				PathEB = ebd;
+			}
+			FVertexConstraint CPath = Constraints->GetVertexConstraint(PathVID);
+			if (!CPath.bCanMove || CPath.bCannotDelete || CPath.Target != ca.Target)
+			{
+				return false;
+			}
+			bool bOtherEdgesUnconstrained = true;
+			Mesh->EnumerateVertexEdges(PathVID, [&](int32 EID)
+			{
+				if (EID == PathEA || EID == PathEB)
+				{
+					return;
+				}
+				bOtherEdgesUnconstrained = bOtherEdgesUnconstrained &&
+					!Constraints->HasEdgeConstraint(EID);
+			});
+			if (bOtherEdgesUnconstrained)
+			{
+				bCanCollapse = true;
+				bIndirectConstraintConnection = true;
+				OutCollapseInfo->UnconstrainVIDs.Add(PathVID);
+				OutCollapseInfo->UnconstrainEIDs.Add(PathEA);
+				OutCollapseInfo->UnconstrainEIDs.Add(PathEB);
+			}
+		}
+	}
+
+	if (!bCanCollapse)
+	{
+		return false;
+	}
 
-	if (bAllowSeamCollapse && !bPreserveSeamTopology)
+	// This non-matching targets case is not caught by MeshRefinerBase::CanCollapseVertex;
+	// for simplicity, we handle it after the indirect connections case is ruled out
+	if (!bIndirectConstraintConnection && ca.Target != EdgeConstraint.Target && ca.Target != nullptr && cb.Target != nullptr)
 	{
-		return CanCollapseVertex(edgeID, a, b, collapse_to);
+		return false;
 	}
 
-	// make sure that the retained edges from the collapsed triangles don't merge seams
-	bool bCanCollapse = FMeshRefinerBase::CanCollapseEdge(edgeID, a, b, c, d, t0, t1, collapse_to);
+	// Helper to update collapse_to + bCanCollapse when considering additional constraints
+	auto AddNewABConstraint = [&bCanCollapse, a, b, &collapse_to](bool bConstrainA, bool bConstrainB)
+	{
+		if (bConstrainA == bConstrainB)
+		{
+			bCanCollapse = !bConstrainA;
+		}
+		else // only constraining A or B
+		{
+			int32 ConstrainVID = bConstrainA ? a : b;
+			if (collapse_to == IndexConstants::InvalidID)
+			{
+				collapse_to = ConstrainVID;
+			}
+			bCanCollapse = collapse_to == ConstrainVID;
+		}
+	};
+
+	// if the vertex constraints did not prevent collapse, and we didn't follow the special case above,
+	// then make sure the AC CB and AD DB edges don't prevent the collapse via their constraints
+	// Note this follows the MeshRefinerBase::CanCollapseEdge logic, after the MeshRefinerBase::CanCollapseVertex call
+	if (bCanCollapse && !bIndirectConstraintConnection)
+	{
+		// determine if edge constraints require keeping either vert.
+		// if edge ac is constrained, we must keep it and thus vertex a, likewise for edge bc and vertex b.
+		int32 eac = Mesh->FindEdgeFromTri(a, c, tc);
+		int32 ebc = Mesh->FindEdgeFromTri(b, c, tc);
+		bool bMustRetainA = Constraints->HasEdgeConstraint(eac);
+		bool bMustRetainB = Constraints->HasEdgeConstraint(ebc);
+
+		// if edge ad is constrained, we must keep it and thus vertex a, likewise for edge bd and vertex b.
+		if (d != IndexConstants::InvalidID)
+		{
+			int32 ead = Mesh->FindEdgeFromTri(a, d, td);
+			int32 ebd = Mesh->FindEdgeFromTri(b, d, td);
+
+			bMustRetainA = bMustRetainA || Constraints->HasEdgeConstraint(ead);
+			bMustRetainB = bMustRetainB || Constraints->HasEdgeConstraint(ebd);
+		}
+
+		AddNewABConstraint(bMustRetainA, bMustRetainB);
+	}
 
 	// make sure more general seam topology is preserved
+
+	const bool bPreserveSeamTopology = (CollapseMode == ESimplificationCollapseModes::MinimalExistingVertexError);
 	if (bCanCollapse && bAllowSeamCollapse && bPreserveSeamTopology)
 	{
 		if (!Constraints)
@@ -1306,6 +1429,10 @@ bool TMeshSimplification<QuadricErrorType>::CanCollapseEdge(int edgeID, int a, i
 		// check if this edge is a seam
 		if (Constraints->HasEdgeConstraint(edgeID))
 		{
+			// TODO [jimmy]: two issues here:
+			// (1) conceptually it's in the wrong place; we should instead set no-move constraints on the seam intersections in the original constraint system
+			// (2) it's incorrectly constraining both vertices of the seam edge, in cases where only one end is 'intersecting'
+			
 			// examine local topology
 			bool bCanCollapseSeam = true;
 			if (const FDynamicMeshAttributeSet* Attributes = Mesh->Attributes())
@@ -1338,32 +1465,7 @@ bool TMeshSimplification<QuadricErrorType>::CanCollapseEdge(int edgeID, int a, i
 			// - this is done by checking for vertex constraint
 			bool bVertexAOnSeam = Constraints->HasVertexConstraint(a);
 			bool bVertexBOnSeam = Constraints->HasVertexConstraint(b);
-			if (bVertexAOnSeam && bVertexBOnSeam)
-			{
-				bCanCollapse = false;
-			}
-			else if (bVertexAOnSeam)
-			{
-				if (collapse_to == -1)
-				{
-					collapse_to = a;
-				}
-				else if (collapse_to != a)
-				{
-					bCanCollapse = false;
-				}
-			}
-			else if (bVertexBOnSeam)
-			{
-				if (collapse_to == -1)
-				{
-					collapse_to = b;
-				}
-				else if (collapse_to != b)
-				{
-					bCanCollapse = false;
-				}
-			}
+			AddNewABConstraint(bVertexAOnSeam, bVertexBOnSeam);
 		}
 	}
 
@@ -1424,7 +1526,8 @@ ESimplificationResult TMeshSimplification<QuadricErrorType>::CollapseEdge(int ed
 	// check if we should collapse, and also find which vertex we should retain
 	// in cases where we have constraints/etc
 	int collapse_to = -1;
-	bool bCanCollapse = CanCollapseEdge(edgeID, a, b, c, d, t0, t1, collapse_to);
+	FCollapseInfo CollapseInfo;
+	bool bCanCollapse = CanCollapseEdge(edgeID, a, b, c, d, t0, t1, collapse_to, &CollapseInfo);
 
 	if (bCanCollapse == false)
 	{
@@ -1538,6 +1641,14 @@ ESimplificationResult TMeshSimplification<QuadricErrorType>::CollapseEdge(int ed
 		}
 	}
 
+	// Check if edge constraint tolerance would prevent collapse
+	// Note this must happen after the above block has already called GetProjectedCollapsePosition
+	// to apply constraint projection to the vNewPos if needed
+	if (CheckIfCollapseWithinEdgeConstraintsTolerance(iKeep, iCollapse, vNewPos, t0) == false)
+	{
+		return ESimplificationResult::Failed_ConstraintDeviation;
+	}
+
 	// check if this collapse will create a normal flip. Also checks
 	// for invalid collapse nbrhood, since we are doing one-ring iter anyway.
 	// [TODO] could we skip this one-ring check in CollapseEdge? pass in hints?
@@ -1562,6 +1673,14 @@ ESimplificationResult TMeshSimplification<QuadricErrorType>::CollapseEdge(int ed
 		if (Constraints)
 		{
 			Constraints->ClearEdgeConstraint(edgeID);
+			for (int32 ClearVID : CollapseInfo.UnconstrainVIDs)
+			{
+				Constraints->ClearVertexConstraint(ClearVID);
+			}
+			for (int32 ClearEID : CollapseInfo.UnconstrainEIDs)
+			{
+				Constraints->ClearEdgeConstraint(ClearEID);
+			}
 
 			auto ConstraintUpdator = [this](int cur_eid)->void
 			{
@@ -1638,6 +1757,43 @@ ESimplificationResult TMeshSimplification<QuadricErrorType>::CollapseEdge(int ed
 	return retVal;
 }
 
+template <typename QuadricErrorType>
+bool TMeshSimplification<QuadricErrorType>::CheckIfCollapseWithinEdgeConstraintsTolerance(int VKeep, int VRemove, const FVector3d& NewPosition, int TC)
+{
+	if (!Constraints || !bLimitConstrainedSeamMovement)
+	{
+		return true;
+	}
+
+	// test edge midpoints, except the edge being collapsed
+	int32 CollapseEdgeID = Mesh->FindEdgeFromTri(VKeep, VRemove, TC);
+	auto EdgeMidpointsWithinTolerance = [this, CollapseEdgeID, NewPosition](int32 VID) -> bool
+	{
+		for (int32 EID : Mesh->VtxEdgesItr(VID))
+		{
+			if (EID != CollapseEdgeID)
+			{
+				FEdgeConstraint EdgeConstraint = Constraints->GetEdgeConstraint(EID);
+				if (EdgeConstraint.Target && EdgeConstraint.DistanceToleranceSq > 0)
+				{
+					FIndex2i EdgeV = Mesh->GetEdgeV(EID);
+					FVector3d OtherVertexPos = (EdgeV.A == VID) ? Mesh->GetVertex(EdgeV.B) : Mesh->GetVertex(EdgeV.A);
+					FVector3d NewMidpoint = (OtherVertexPos + NewPosition) * 0.5;
+
+					FVector3d ConstraintProj = EdgeConstraint.Target->Project(NewMidpoint);
+					if (DistanceSquared(NewMidpoint, ConstraintProj) > EdgeConstraint.DistanceToleranceSq)
+					{
+						return false;
+					}
+				}
+			}
+		}
+		return true;
+	};
+
+	return EdgeMidpointsWithinTolerance(VKeep) && EdgeMidpointsWithinTolerance(VRemove);
+}
+
 template <typename QuadricErrorType>
 bool TMeshSimplification<QuadricErrorType>::CheckIfCollapseWithinGeometricTolerance(int vKeep, int vRemove, const FVector3d& NewPosition, int tc, int td)
 {
@@ -1669,7 +1825,7 @@ bool TMeshSimplification<QuadricErrorType>::CheckIfCollapseWithinGeometricTolera
 			}
 
 			// test edge midpoints, except the edge being collapsed
-			int32 CollapseEdgeID = Mesh->FindEdge(vKeep, vRemove);
+			int32 CollapseEdgeID = Mesh->FindEdgeFromTri(vKeep, vRemove, tc);
 			auto EdgeMidpointsWithinTolerance = [this, CollapseEdgeID, NewPosition, ToleranceSqr](int32 vid) {
 				for ( int32 eid : Mesh->VtxEdgesItr(vid) )
 				{
diff --git a/Engine/Plugins/Runtime/GeometryProcessing/Source/DynamicMesh/Public/MeshConstraints.h b/Engine/Plugins/Runtime/GeometryProcessing/Source/DynamicMesh/Public/MeshConstraints.h
--- a/Engine/Plugins/Runtime/GeometryProcessing/Source/DynamicMesh/Public/MeshConstraints.h
+++ b/Engine/Plugins/Runtime/GeometryProcessing/Source/DynamicMesh/Public/MeshConstraints.h
@@ -50,11 +50,15 @@ public:
 	/** This ID is not a constraint, but can be used to find descendants of a constrained input edge after splits */
 	int TrackingSetID;
 
+	/** Squared tolerance distance: if positive, edge should not move farther than this distance from Target (if set, and feature is enabled; see TMeshSimplification::bLimitConstrainedSeamMovement) */
+	float DistanceToleranceSq = -1.f;
+
 	FEdgeConstraint()
 	{
 		RefineFlags = EEdgeRefineFlags::NoConstraint;
 		Target = nullptr;
 		TrackingSetID = -1;
+		DistanceToleranceSq = -1.f;
 	}
 
 	explicit FEdgeConstraint(EEdgeRefineFlags ConstraintFlags)
diff --git a/Engine/Plugins/Runtime/GeometryProcessing/Source/DynamicMesh/Public/MeshConstraintsUtil.h b/Engine/Plugins/Runtime/GeometryProcessing/Source/DynamicMesh/Public/MeshConstraintsUtil.h
--- a/Engine/Plugins/Runtime/GeometryProcessing/Source/DynamicMesh/Public/MeshConstraintsUtil.h
+++ b/Engine/Plugins/Runtime/GeometryProcessing/Source/DynamicMesh/Public/MeshConstraintsUtil.h
@@ -189,12 +189,14 @@ s	 * EdgeROI that are seams, as well as vertices incident to edges in EdgeROI th
 	 * @param Constraints (out) constraint data to add new constraints to
 	 * @param Mesh the mesh to constrain
 	 * @param CornerAngleThreshold angle threshold (in degrees) above which adjacent edges are considered a corner
+	 * @param CurveConstraintToleranceDistance threshold distance used to limit how far edges should be allowed to deviate from their constraint curves
 	 */
 	static UE_API void SetBoundaryConstraintsWithProjection(
 		FMeshConstraints& Constraints,
 		const EBoundaryType BoundaryToConstrain,
 		const FDynamicMesh3& Mesh,
-		double CornerAngleThreshold = 45.0);
+		double CornerAngleThreshold = 45.0,
+		float CurveConstraintToleranceDistance = -1.f);
 
 private:
 	FMeshConstraintsUtil() = delete;		// this class is not constructible
diff --git a/Engine/Plugins/Runtime/GeometryProcessing/Source/DynamicMesh/Public/MeshSimplification.h b/Engine/Plugins/Runtime/GeometryProcessing/Source/DynamicMesh/Public/MeshSimplification.h
--- a/Engine/Plugins/Runtime/GeometryProcessing/Source/DynamicMesh/Public/MeshSimplification.h
+++ b/Engine/Plugins/Runtime/GeometryProcessing/Source/DynamicMesh/Public/MeshSimplification.h
@@ -28,7 +28,8 @@ enum class ESimplificationResult
 	Failed_NotAnEdge = 7,
 	Failed_IsolatedTriangle = 8,
 	Failed_GeometricDeviation = 9,
-	Ignored_CreatesTinyTriangle = 10
+	Ignored_CreatesTinyTriangle = 10,
+	Failed_ConstraintDeviation = 11
 };
 
 enum class ERegularizeSimplificationMethod
@@ -92,6 +93,11 @@ public:
 	 */
 	bool bPreventTinyTriangles = false;
 
+	/**
+	 * If true, we prevent collapses that would move edge midpoints away from their constraint target by more than the edge constraint's threshold distance (if any)
+	 */
+	bool bLimitConstrainedSeamMovement = false;
+
 	/** 
 	 * If greater than 0, we will add a regularization term with this weight.
 	 * This can help give a better triangulation for example in flat regions where the error metric is under-determined
@@ -335,13 +341,25 @@ protected:
 		return IsBoundaryVtxCache[vid];
 	}
 
+	// Optional additional info to support edge collapses in more configurations
+	// i.e. for tri ABC, if edge AB is not constrained, but AC and BC are, 
+	// it is *sometimes* safe to collapse AB *if* we also remove constraints from 
+	// vertex B and edges B(AC) after the collapse. (As long as B is not otherwise constrained)
+	struct FCollapseInfo
+	{
+		TArray<int32, TInlineAllocator<2>> UnconstrainVIDs;
+		TArray<int32, TInlineAllocator<2>> UnconstrainEIDs;
+	};
+
 	/**
 	* Figure out if we can collapse edge eid=[a,b] under current constraint set
-	* This depends on the base class function FMeshRefinerBase::CanCollapseEdge, 
-	* but includes additional restrictions specified by the CollapseMode of this simplifier.
-	* See the base class for description of parameters.
+	* Similar to the base class FMeshRefinerBase::CanCollapseEdge, but with
+	* simplifier-specific modifications (some more flexibility and restrictions depending on CollapseMode)
+	* 
+	* @param OutCollapseInfo If non-null, we can allow more edge collapses that require some additional post-collapse updates
+	* See the similar base class for description of the other parameters.
 	*/
-	bool CanCollapseEdge(int edgeID, int a, int b, int c, int d, int t0, int t1, int& collapse_to) const;
+	bool CanCollapseEdge(int edgeID, int a, int b, int c, int d, int t0, int t1, int& collapse_to, FCollapseInfo* OutCollapseInfo = nullptr) const;
 
 
 	/**
@@ -384,6 +402,17 @@ protected:
 	 */
 	bool CheckIfCollapseWithinGeometricTolerance(int vid, int vother, const FVector3d& newv, int tc, int td);
 
+	/**
+	 * Check if edge collapse would violate edge constraint tolerance
+	 * @param VKeep First edge vertex
+	 * @param VRemove Other edge vertex
+	 * @param NewPosition New vertex position after collapse
+	 * @param TC triangle on one side of edge
+	 * @return true if collapsing the edge won't move constrained edge midpoints too far from their projection targets
+	 */
+	bool CheckIfCollapseWithinEdgeConstraintsTolerance(int VKeep, int VRemove, const FVector3d& NewPosition, int TC);
+	
+
 
 	/*
 	 * testing/debug/profiling stuff
