Add medial-axis sampling/skeletonization to GeometryAlgorithms

#Synced-CL 50036218

#Synced-CL 50036218

diff --git a/Engine/Plugins/Runtime/GeometryProcessing/Source/GeometryAlgorithms/Private/Skeletonization/MeshMedialAxisSampling.cpp b/Engine/Plugins/Runtime/GeometryProcessing/Source/GeometryAlgorithms/Private/Skeletonization/MeshMedialAxisSampling.cpp
new file mode 100644
--- /dev/null
+++ b/Engine/Plugins/Runtime/GeometryProcessing/Source/GeometryAlgorithms/Private/Skeletonization/MeshMedialAxisSampling.cpp
@@ -0,0 +1,841 @@
+// Copyright Epic Games, Inc. All Rights Reserved.
+
+#include "Skeletonization/MeshMedialAxisSampling.h"
+
+#include <atomic>
+
+#include "VertexConnectedComponents.h"
+#include "MeshAdapter.h"
+#include "MeshAdapter/MeshVertexNormals.h"
+
+// Include Eigen/Dense for a 4x4 matrix solve
+#if defined(_MSC_VER) && USING_CODE_ANALYSIS
+#pragma warning(push)
+#pragma warning(disable : 6011)
+#pragma warning(disable : 6387)
+#pragma warning(disable : 6313)
+#pragma warning(disable : 6294)
+#endif
+PRAGMA_DEFAULT_VISIBILITY_START
+THIRD_PARTY_INCLUDES_START
+#include <Eigen/Dense>
+THIRD_PARTY_INCLUDES_END
+PRAGMA_DEFAULT_VISIBILITY_END
+#if defined(_MSC_VER) && USING_CODE_ANALYSIS
+#pragma warning(pop)
+#endif
+
+namespace UE::Geometry::MedialAxis
+{
+
+	template<typename MeshType>
+	FMedialSkeleton FSkeletonViaSampling::ComputeSkeleton(
+		const TMeshAABBTree3<MeshType>& MeshBVH, 
+		TFunctionRef<FVector3d(int32)> GetVertexNormal, 
+		TFunctionRef<int32(int32)> VIDtoWeldVIDFn, 
+		const TFastWindingTree<MeshType>* MeshFWTree)
+	{
+		using FEigMat4 = Eigen::Matrix4d;
+		using FEigVec4 = Eigen::Vector4d;
+		
+		FMedialSkeleton Result;
+
+		const MeshType& Mesh = *MeshBVH.GetMesh();
+		if (Mesh.TriangleCount() == 0)
+		{
+			return Result;
+		}
+
+		// Get mapping from vertex IDs to coincident triangles, accounting for any welding
+		TArray<TArray<int32, TInlineAllocator<8>>> VIDtoTIDs, VIDtoNbrVIDs;
+		VIDtoTIDs.SetNum(Mesh.MaxVertexID());
+		VIDtoNbrVIDs.SetNum(Mesh.MaxVertexID());
+		for (int32 TID = 0; TID < Mesh.MaxTriangleID(); ++TID)
+		{
+			if (Mesh.IsTriangle(TID))
+			{
+				FIndex3i Tri = Mesh.GetTriangle(TID);
+				Tri.A = VIDtoWeldVIDFn(Tri.A);
+				Tri.B = VIDtoWeldVIDFn(Tri.B);
+				Tri.C = VIDtoWeldVIDFn(Tri.C);
+				for (int32 PrevIdx = 2, Idx = 0; Idx < 3; PrevIdx = Idx++)
+				{
+
+					VIDtoTIDs[Tri[Idx]].Add(TID);
+					VIDtoNbrVIDs[Tri[Idx]].AddUnique(Tri[PrevIdx]);
+					VIDtoNbrVIDs[Tri[PrevIdx]].AddUnique(Tri[Idx]);
+				}
+			}
+		}
+
+		// Simple enum to track which vertices are valid (exist in the mesh and are referenced by triangles) and have non-zero-radius medial spheres
+		enum EVertexState : uint8
+		{
+			// unreferenced / not valid
+			Skip = 0,
+			// valid but failed to find a plausible medial sphere, don't use as source
+			NoMedialSphere,
+			// fully valid, incl. with a medial sphere
+			Valid
+		};
+		TArray<EVertexState> VertexState; // track whether a vertex ID is valid and has an associated sphere
+		VertexState.SetNumZeroed(Mesh.MaxVertexID());
+
+		TArray<bool> AddedVertexSphere; // track whether we've already used a given seed vertex
+		AddedVertexSphere.SetNumZeroed(Mesh.MaxVertexID());
+
+		TArray<FSphere> VtxMedialSpheres; // computed initial/seed spheres per vertex
+		double InitialRadius = MeshBVH.GetBoundingBox().MaxDim();
+
+		// Pre-compute a medial sphere per-vertex (or a tag indicating where we couldn't find a sphere)
+		VtxMedialSpheres.SetNumUninitialized(Mesh.MaxVertexID());
+		const double FailedMedialSphereThreshold = FMath::Max(FMathd::ZeroTolerance, ShrinkMedial.RadiusThreshold);
+		ParallelFor(Mesh.MaxVertexID(), 
+			[&MeshBVH, &Mesh, &VIDtoWeldVIDFn, &VIDtoNbrVIDs, this, &VtxMedialSpheres, &VertexState, InitialRadius, FailedMedialSphereThreshold, &GetVertexNormal]
+			(int32 VID)
+		{
+			if (Mesh.IsVertex(VID) && VIDtoWeldVIDFn(VID) == VID && !VIDtoNbrVIDs[VID].IsEmpty())
+			{
+				double LocalFindRadius = InitialRadius;
+				FVector3d P = Mesh.GetVertex(VID), N = GetVertexNormal(VID);
+				bool bHotStart = ShrinkMedial.FindLocalMinCurvatureRadius(Mesh, P, N, VIDtoNbrVIDs[VID], LocalFindRadius);
+				VtxMedialSpheres[VID] = ShrinkMedial.FindWithImplicitWeld(MeshBVH, VID, VIDtoWeldVIDFn, P, N, bHotStart, LocalFindRadius);
+				VertexState[VID] = VtxMedialSpheres[VID].W > FailedMedialSphereThreshold ? EVertexState::Valid : EVertexState::NoMedialSphere;
+			}
+		});
+
+		VIDtoNbrVIDs.Empty(); // Not used after this point
+
+		// Start with the biggest medial sphere
+		int32 BiggestSphereIdx = INDEX_NONE;
+		double BigSphereRad = -1;
+		for (int32 VID = 0; VID < Mesh.MaxVertexID(); ++VID)
+		{
+			if (VertexState[VID] == EVertexState::Valid && VtxMedialSpheres[VID].W > BigSphereRad)
+			{
+				BigSphereRad = VtxMedialSpheres[VID].W;
+				BiggestSphereIdx = VID;
+			}
+		}
+		// Early out if there were no valid medial spheres found
+		if (BiggestSphereIdx == INDEX_NONE)
+		{
+			return Result;
+		}
+
+		// Initialize clustering with the found largest sphere, mapping all vertices to that sphere
+		Result.Spheres.SetNum(1);
+		Result.Spheres[0] = VtxMedialSpheres[BiggestSphereIdx];
+		Result.VIDtoClusterIndex.SetNumUninitialized(Mesh.MaxVertexID());
+		for (int32 VID = 0; VID < Result.VIDtoClusterIndex.Num(); ++VID)
+		{
+			Result.VIDtoClusterIndex[VID] = VertexState[VID] == EVertexState::Skip ? INDEX_NONE : 0;
+		}
+
+
+		/// Arrays for tracking clusters and associated stats, will be re-used across multiple iterations below
+
+		TArray<double> MaxErrInCluster; // indexed by cluster ID (i.e., sphere index)
+		TArray<int32> MaxErrVertInCluster; // 1:1 with MaxErrInCluster
+		TArray<int32> ClusterSize; // number of vertices in a cluster
+		TArray<int32> RemapClusters; // used to remap cluster indices during a cluster deletion pass
+		TArray<int32> ClusterVertStarts, VertsByCluster, ClustersToUpdate;
+		// used to track which clusters should 'sleep' (not be optimized) b/c they haven't added or removed any vertices recently
+		TArray<std::atomic<int8>> FullPassesToUpdateCluster;
+		FullPassesToUpdateCluster.SetNum(1);
+		FullPassesToUpdateCluster[0] = OptimizeSpheresForIterationsPostClusterChange;
+		VertsByCluster.SetNumUninitialized(Mesh.MaxVertexID());
+		Result.ClusterNeighbors.SetNum(1);
+		TArray<int32> ClusterIndicesToSplitHeap; // tracks candidate cluster indices to split (ordered via heap)
+		TArray<bool> HasSplitNbr;
+		TArray<TArray<int32, TInlineAllocator<3>>> ShouldSplitClusterEdge;
+		TArray<FMedialSkeleton::FNbrArray> ClusterNbrNbrs; // Stores neighbors-of-neighbors per cluster (excluding the cluster itself)
+
+
+		// Helper to compute cluster sizes based on current VIDtoClusterIndex assignments
+		auto UpdateClusterSize = [](const MeshType& Mesh, FMedialSkeleton& Result, TArray<int32>& OutClusterSize)
+		{
+			OutClusterSize.Reset(); OutClusterSize.SetNumZeroed(Result.Spheres.Num());
+			for (int32 VID = 0; VID < Mesh.MaxVertexID(); ++VID)
+			{
+				// Note: Don't need to check Mesh.IsVertex(VID) here, because invalid vertices will already be mapped to the INDEX_NONE cluster index
+				int32 ClusterIdx = Result.VIDtoClusterIndex[VID];
+				if (ClusterIdx != INDEX_NONE)
+				{
+					OutClusterSize[ClusterIdx]++;
+				}
+			}
+		};
+
+		// Helper to transform cluster neighbors to neighbors-of-neighbors
+		auto BuildClusterNbrNbrs = [&Result, &ClusterNbrNbrs]()
+		{
+			ClusterNbrNbrs.Reset(); ClusterNbrNbrs.SetNum(Result.Spheres.Num());
+			// Cluster Neighbors must be already computed
+			if (!ensure(Result.Spheres.Num() == Result.ClusterNeighbors.Num()))
+			{
+				return;
+			}
+			ParallelFor(Result.Spheres.Num(), [&Result, &ClusterNbrNbrs](int32 ClusterIdx)
+			{
+				ClusterNbrNbrs[ClusterIdx].Append(Result.ClusterNeighbors[ClusterIdx]);
+				for (int32 Nbr : Result.ClusterNeighbors[ClusterIdx])
+				{
+					for (int32 NbrNbr : Result.ClusterNeighbors[Nbr])
+					{
+						if (NbrNbr != ClusterIdx)
+						{
+							ClusterNbrNbrs[ClusterIdx].AddUnique(NbrNbr);
+						}
+					}
+				}
+			});
+		};
+		
+		// Helper compute a spherical quadric error metric for a given triangle, and to sum and evaluate them for faster error calculation
+		struct FSQEM
+		{
+			FEigMat4 A;
+			FEigVec4 B;
+			double C;
+			FSQEM() : A(FEigMat4::Zero()), B(FEigVec4::Zero()), C(0) {}
+			FSQEM(const FVector3d& Centroid, const FVector3d& Normal, double Wt)
+			{
+				FEigVec4 N_1(Normal.X, Normal.Y, Normal.Z, 1.);
+				A = (Wt * N_1) * N_1.transpose();
+				double NdP = Normal.Dot(Centroid);
+				B = N_1 * (2. * NdP * Wt);
+				C = NdP * NdP * Wt;
+			}
+			double Eval(const FSphere& Sphere)
+			{
+				// Compute the spherical quadric error, s^T A s - s^T B + C
+				FEigVec4 S(Sphere.Center.X, Sphere.Center.Y, Sphere.Center.Z, Sphere.W);
+				double Aterm = S.transpose() * A * S;
+				double Bterm = S.transpose() * B;
+				return Aterm - Bterm + C;
+			}
+			void Add(const FSQEM& Other)
+			{
+				A += Other.A;
+				B += Other.B;
+				C += Other.C;
+			}
+		};
+
+		// cache the triangle normals and areas, and vertex areas and SQEMs, for use in optimization below
+		TArray<FVector3f> TriNormals;
+		TArray<FSQEM> VtxSQEMs;
+		TArray<float> TriAreas, VtxAreas;
+		double TotalArea = 0;
+		TriNormals.SetNumUninitialized(Mesh.MaxTriangleID());
+		TriAreas.SetNumUninitialized(Mesh.MaxTriangleID());
+		VtxSQEMs.SetNumZeroed(Mesh.MaxVertexID());
+		VtxAreas.SetNumZeroed(Mesh.MaxVertexID());
+		for (int32 TID = 0; TID < Mesh.MaxTriangleID(); ++TID)
+		{
+			if (!Mesh.IsTriangle(TID))
+			{
+				continue;
+			}
+
+			FVector3d Centroid, Normal;
+			double Area;
+			TMeshQueries<MeshType>::GetTriNormalAreaCentroid(Mesh, TID, Normal, Area, Centroid);
+			TotalArea += Area;
+			FIndex3i Tri = Mesh.GetTriangle(TID);
+			TriNormals[TID] = (FVector3f)Normal;
+			TriAreas[TID] = (float)Area;
+			for (int32 SubIdx = 0; SubIdx < 3; ++SubIdx)
+			{
+				// Note we could divide by 3 for the more logical vertex area, but would then need to divide by 3
+				// when scaling triangle area contributions in a number of places below as well ...
+				// In practice these 1/3 factors will cancel out in our usages
+				VtxAreas[VIDtoWeldVIDFn(Tri[SubIdx])] += TriAreas[TID];
+				VtxSQEMs[VIDtoWeldVIDFn(Tri[SubIdx])].Add(FSQEM(Centroid, Normal, Area));
+			}
+		}
+
+		// Stop early on a fully degenerate / zero-area mesh
+		if (TotalArea == 0)
+		{
+			return Result;
+		}
+
+
+
+		// Area-Weighted Average of face-Normal errors, using FSQEMs
+		auto GetVertexAvgSquaredError = [&Mesh, &VtxSQEMs, &VtxAreas, this](int32 VID, const FSphere& Sphere) -> double
+		{
+			FVector3d V = Mesh.GetVertex(VID);
+			double PosErr = (V - Sphere.Center).Length() - Sphere.W;
+			double PosErrContrib = PosErr * PosErr * PosErrorWt;
+			double NormalErrContrib = VtxAreas[VID] > 0 ? VtxSQEMs[VID].Eval(Sphere) / VtxAreas[VID] : 0.;
+			return PosErrContrib + NormalErrContrib;
+		};
+
+		// Area-Weighted Average of face-Normal errors, without using FSQEM (should be equivalent to the above GetVertexAvgSquaredError, kept here for comparison -- currently this appears to be slower)
+		auto GetVertexAvgSquaredError_NoSQEM = 
+			[&Mesh, this, &VtxAreas, &TriNormals, &TriAreas, &VIDtoTIDs]
+			(int32 VID, const FSphere& Sphere) -> double
+		{
+			FVector3d V = Mesh.GetVertex(VID);
+			double PosErr = (V - Sphere.Center).Length() - Sphere.W;
+			double TotalErr = PosErr * PosErr * PosErrorWt;
+			double AreaSum = VtxAreas[VID];
+			if (AreaSum > 0)
+			{
+				double NormalErrContrib = 0;
+				for (int32 TID : VIDtoTIDs[VID])
+				{
+					FVector3d N = (FVector3d)TriNormals[TID];
+					double SignedDist = N.Dot(V - Sphere.Center) - Sphere.W;
+					NormalErrContrib += TriAreas[TID] * SignedDist * SignedDist;
+				}
+				TotalErr += NormalErrContrib / AreaSum;
+			}
+			return TotalErr;
+		};
+
+		double LastAvgError = FMathd::MaxReal;
+		const double MinClusterErrorToSplitSq = MinClusterErrorToSplit * MinClusterErrorToSplit;
+
+
+		// Helper to assign a vertex to the lowest-error cluster, via exhaustive search
+		auto AssignVertexCluster = [&Result, &Mesh, &VertexState, &GetVertexAvgSquaredError, this, &FullPassesToUpdateCluster](int32 VID)
+		{
+			if (VertexState[VID] == EVertexState::Skip)
+			{
+				return;
+			}
+
+			int32 OrigCluster = Result.VIDtoClusterIndex[VID];
+
+			double BestErr = GetVertexAvgSquaredError(VID, Result.Spheres[0]);
+			int32 BestSphere = 0;
+			for (int32 ClusterIdx = 1; ClusterIdx < Result.Spheres.Num(); ++ClusterIdx)
+			{
+				double Err = GetVertexAvgSquaredError(VID, Result.Spheres[ClusterIdx]);
+				if (Err < BestErr)
+				{
+					BestErr = Err;
+					BestSphere = ClusterIdx;
+				}
+			}
+			if (BestSphere != OrigCluster)
+			{
+				FullPassesToUpdateCluster[BestSphere] = OptimizeSpheresForIterationsPostClusterChange;
+				if (OrigCluster != INDEX_NONE)
+				{
+					FullPassesToUpdateCluster[OrigCluster] = OptimizeSpheresForIterationsPostClusterChange;
+				}
+			}
+			Result.VIDtoClusterIndex[VID] = BestSphere;
+		};
+
+		// Helper to assign a vertex to the lowest-error cluster, via local search of the ClusterNbrNbrs
+		auto AssignVertexCluster_LocalSearch = [&Result, &Mesh, &VertexState, &GetVertexAvgSquaredError, this, &FullPassesToUpdateCluster, &ClusterNbrNbrs, &AssignVertexCluster](int32 VID)
+		{
+			checkSlow(ClusterNbrNbrs.Num() == Result.Spheres.Num());
+			if (VertexState[VID] == EVertexState::Skip)
+			{
+				return;
+			}
+
+			int32 OrigCluster = Result.VIDtoClusterIndex[VID];
+			// fall back to full search for unassigned or isolated clusters
+			if (OrigCluster == INDEX_NONE || ClusterNbrNbrs[OrigCluster].IsEmpty())
+			{
+				AssignVertexCluster(VID);
+				return;
+			}
+
+			double BestErr = GetVertexAvgSquaredError(VID, Result.Spheres[OrigCluster]);
+			int32 BestSphere = OrigCluster;
+			for (int32 ClusterIdx : ClusterNbrNbrs[OrigCluster])
+			{
+				double Err = GetVertexAvgSquaredError(VID, Result.Spheres[ClusterIdx]);
+				if (Err < BestErr)
+				{
+					BestErr = Err;
+					BestSphere = ClusterIdx;
+				}
+			}
+			if (BestSphere != OrigCluster)
+			{
+				FullPassesToUpdateCluster[BestSphere] = OptimizeSpheresForIterationsPostClusterChange;
+				if (OrigCluster != INDEX_NONE)
+				{
+					FullPassesToUpdateCluster[OrigCluster] = OptimizeSpheresForIterationsPostClusterChange;
+				}
+			}
+			Result.VIDtoClusterIndex[VID] = BestSphere;
+		};
+
+		// Helper to update cluster assignments, either via exhuastive or local search
+		auto UpdateClusterAssignments = 
+			[&FullPassesToUpdateCluster, &Result, &Mesh, &AssignVertexCluster_LocalSearch, &AssignVertexCluster, &VIDtoWeldVIDFn, &BuildClusterNbrNbrs, this]
+			(bool bStaleClusterNeighbors)
+		{
+			ensure(FullPassesToUpdateCluster.Num() == Result.Spheres.Num());
+			if (UseLocalClusterVertexReassignmentAtClusterCount >= 0 && Result.Spheres.Num() > UseLocalClusterVertexReassignmentAtClusterCount)
+			{
+				if (bStaleClusterNeighbors)
+				{
+					Result.ComputeClusterNeighbors(Mesh, VIDtoWeldVIDFn);
+				}
+				BuildClusterNbrNbrs();
+				ParallelFor(Mesh.MaxVertexID(), [&AssignVertexCluster_LocalSearch](int32 VID) { AssignVertexCluster_LocalSearch(VID); });
+			}
+			else
+			{
+				ParallelFor(Mesh.MaxVertexID(), [&AssignVertexCluster](int32 VID) { AssignVertexCluster(VID); });
+			}
+		};
+
+		// Helper to evaluate the cluster error at a given vertex
+		// Note: ErrorScaleFactor and ForceMinError are used to force the error to be higher where an edge between clusters intersects the mesh
+		//  (when using the bSplitClustersIfEdgesIntersectSurface setting)
+		auto UpdateMaxError =
+			[&VertexState, &AddedVertexSphere, &Result, &MaxErrInCluster, &MaxErrVertInCluster, &GetVertexAvgSquaredError]
+			(int32 VID, double ErrorScaleFactor = 1., double ForcedMinError = 0.)
+		{
+			// Note this is skipping the vertices that have no associated sphere, 
+			// since we are looking for a high-error vertex that we can add a corresponding sphere for
+			if (VertexState[VID] != EVertexState::Valid)
+			{
+				return;
+			}
+			// Heuristically, we only add the medial sphere of a given vertex once, to help the algorithm avoid 'fixating' on 
+			// single bad vertices and to avoid e.g. cases where the same vertex's medial sphere could be repeatedly added and removed.
+			// This may not always be ideal, as the optimization could shift spheres s.t. we would want to re-add a vertex's medial sphere.
+			// TODO: consider exposing this as a parameter, and/or revisiting the effect of this parameter across more inputs.
+			constexpr bool bNeverReAddVertexMedialSpheres = true;
+			if (bNeverReAddVertexMedialSpheres && AddedVertexSphere[VID])
+			{
+				return;
+			}
+			int32 ClusterIdx = Result.VIDtoClusterIndex[VID];
+			double VtxErr = FMath::Max(ForcedMinError, GetVertexAvgSquaredError(VID, Result.Spheres[ClusterIdx]) * ErrorScaleFactor);
+			if (VtxErr > MaxErrInCluster[ClusterIdx])
+			{
+				MaxErrInCluster[ClusterIdx] = VtxErr;
+				MaxErrVertInCluster[ClusterIdx] = VID;
+			}
+		};
+
+		const FAxisAlignedBox3d MaxCenterBounds = MeshBVH.GetBoundingBox();
+		const double MaxRadius = MaxCenterBounds.MaxDim();
+		
+		int32 ItersWithoutIncrease = 0;
+		bool bHasSurfaceIntersections = false;
+		while (Result.Spheres.Num() < MaxSpheres)
+		{
+			int32 LastNumSpheres = Result.Spheres.Num();
+
+			MaxErrInCluster.Reset(); MaxErrInCluster.SetNumZeroed(Result.Spheres.Num());
+			HasSplitNbr.Reset(); HasSplitNbr.SetNumZeroed(Result.Spheres.Num());
+
+			// Find the highest error vertex per cluster
+			MaxErrVertInCluster.Init(INDEX_NONE, Result.Spheres.Num());
+			
+			for (int32 VID = 0; VID < Mesh.MaxVertexID(); ++VID)
+			{
+				UpdateMaxError(VID);
+			}
+
+			if (bSplitClustersIfEdgesIntersectSurface)
+			{
+				ShouldSplitClusterEdge.SetNum(Result.Spheres.Num());
+				for (int32 Idx = 0; Idx < ShouldSplitClusterEdge.Num(); ++Idx)
+				{
+					ShouldSplitClusterEdge[Idx].Reset();
+				}
+				ParallelFor(Result.Spheres.Num(), [&Result, &MeshBVH, &ShouldSplitClusterEdge](int32 ClusterIdx)
+				{
+					FVector3d ClusterCenter = Result.Spheres[ClusterIdx].Center;
+					for (int32 NbrIdx : Result.ClusterNeighbors[ClusterIdx])
+					{
+						if (NbrIdx > ClusterIdx)
+						{
+							FVector3d NbrCenter = Result.Spheres[NbrIdx].Center;
+							FVector3d Dir = NbrCenter - ClusterCenter;
+							double Len = Normalize(Dir);
+							if (Len > 0)
+							{ 
+								FRay3d Ray(ClusterCenter, Dir);
+								IMeshSpatial::FQueryOptions Options;
+								Options.MaxDistance = Len;
+								if (MeshBVH.FindNearestHitTriangle(Ray, Options) != INDEX_NONE)
+								{
+									ShouldSplitClusterEdge[ClusterIdx].Add(NbrIdx);
+								}
+							}
+						}
+					}
+				});
+				bHasSurfaceIntersections = false;
+				for (int32 Idx = 0; Idx < ShouldSplitClusterEdge.Num(); ++Idx)
+				{
+					if (!ShouldSplitClusterEdge[Idx].IsEmpty())
+					{
+						bHasSurfaceIntersections = true;
+						break;
+					}
+				}
+				if (bHasSurfaceIntersections)
+				{
+					for (int32 TID = 0; TID < Mesh.MaxTriangleID(); ++TID)
+					{
+						if (!Mesh.IsTriangle(TID))
+						{
+							continue;
+						}
+						FIndex3i Tri = Mesh.GetTriangle(TID);
+						Tri.A = VIDtoWeldVIDFn(Tri.A);
+						Tri.B = VIDtoWeldVIDFn(Tri.B);
+						Tri.C = VIDtoWeldVIDFn(Tri.C);
+						for (int32 Prev = 2, Idx = 0; Idx < 3; Prev = Idx++)
+						{
+							FIndex2i EdgeV(Tri[Prev], Tri[Idx]);
+							FIndex2i EdgeCluster(
+								Result.VIDtoClusterIndex[EdgeV.A],
+								Result.VIDtoClusterIndex[EdgeV.B]);
+							if (EdgeCluster.A == EdgeCluster.B 
+								|| !ensure(EdgeCluster.A != INDEX_NONE && EdgeCluster.B != INDEX_NONE))
+							{
+								continue;
+							}
+							EdgeCluster.Sort();
+							if (ShouldSplitClusterEdge[EdgeCluster.A].Contains(EdgeCluster.B))
+							{
+								UpdateMaxError(EdgeV.A, ErrorScaleNearEdgeSurfaceIntersection, MinClusterErrorToSplitSq + FMathd::ZeroTolerance);
+								UpdateMaxError(EdgeV.B, ErrorScaleNearEdgeSurfaceIntersection, MinClusterErrorToSplitSq + FMathd::ZeroTolerance);
+							}
+						}
+					}
+				}
+			}
+
+			// For clusters w/ max error above threshold, add the max-error vertex's medial sphere
+			// (if no neighboring cluster already added a sphere, and up to a fixed number of splits)
+			ClusterIndicesToSplitHeap.Reset();
+			auto MaxHeapPred = [&MaxErrInCluster](int32 A, int32 B) -> bool { return MaxErrInCluster[A] > MaxErrInCluster[B]; };
+			for (int32 Idx = 0; Idx < MaxErrInCluster.Num(); ++Idx)
+			{
+				if (MaxErrInCluster[Idx] > MinClusterErrorToSplitSq)
+				{
+					ClusterIndicesToSplitHeap.Add(Idx);
+				}
+			}
+			if (ClusterIndicesToSplitHeap.IsEmpty())
+			{
+				break;
+			}
+			ClusterIndicesToSplitHeap.Heapify(MaxHeapPred);
+			int32 NumSplits = 0;
+			double MinErrorToSplit = MaxErrInCluster[ClusterIndicesToSplitHeap.HeapTop()] * MinFractionOfMaxErrorToSplit;
+			int32 MaxNumSplits = FMath::Max(MaxSplitPerIteration, FMath::CeilToInt32(Result.Spheres.Num() * MaxSplitFractionPerIteration));
+			if (MaxNumSplits <= 0) // handle case where both Max parameters were disabled (set to 0)
+			{
+				MaxNumSplits = ClusterIndicesToSplitHeap.Num();
+			}
+			while (!ClusterIndicesToSplitHeap.IsEmpty() && NumSplits < MaxNumSplits)
+			{
+				int32 ClusterToSplit;
+				ClusterIndicesToSplitHeap.HeapPop(ClusterToSplit, MaxHeapPred, EAllowShrinking::No);
+				if (MaxErrInCluster[ClusterToSplit] < MinErrorToSplit)
+				{
+					break;
+				}
+				if (!HasSplitNbr[ClusterToSplit])
+				{
+					int32 ClusterSourceVID = MaxErrVertInCluster[ClusterToSplit];
+					AddedVertexSphere[ClusterSourceVID] = true;
+					FSphere ToAdd = VtxMedialSpheres[ClusterSourceVID];
+					int32 NewClusterIdx = Result.Spheres.Add(ToAdd);
+					FullPassesToUpdateCluster.Emplace(OptimizeSpheresForIterationsPostClusterChange);
+					checkSlow(FullPassesToUpdateCluster.Num() == Result.Spheres.Num());
+					NumSplits++;
+					for (int32 Nbr : Result.ClusterNeighbors[ClusterToSplit])
+					{
+						HasSplitNbr[Nbr] = true;
+					}
+
+					// Add new cluster as neighbor of the split cluster, so local re-clustering can find it
+					Result.ClusterNeighbors[ClusterToSplit].Add(NewClusterIdx);
+					Result.ClusterNeighbors.Emplace_GetRef().Add(ClusterToSplit);
+					check(Result.ClusterNeighbors.Num() == Result.Spheres.Num());
+				}
+			}
+
+			// If no spheres were added, stop iterating
+			if (LastNumSpheres == Result.Spheres.Num())
+			{
+				break;
+			}
+
+			UpdateClusterAssignments(false /*bStaleClusterNeighbors, neighbors are already updated above so aren't stale here*/);
+
+			// Delete clusters that are below a minimum size, and re-assign their vertices
+			{
+				UpdateClusterSize(Mesh, Result, ClusterSize);
+				RemoveSmallClusters(Result, ClusterSize, RemapClusters, FullPassesToUpdateCluster, AssignVertexCluster);
+
+				// Track the number of iterations where we delete as many spheres as we add --
+				// if this happens repeatedly over multiple sequential iterations, the algorithm is likely stuck
+				if (Result.Spheres.Num() > LastNumSpheres)
+				{
+					ItersWithoutIncrease = 0;
+				}
+				else
+				{
+					ItersWithoutIncrease++;
+				}
+			}
+
+			// Build a flat array of per-cluster vertex IDs for the updated clusters
+			// Transform cluster sizes into 'cluster starts' (indicating the first VID for each cluster, in the combined array)
+			// Reset cluster size counts for updated clusters to use for per-cluster index tracking
+			ClusterVertStarts.Reset(); ClusterVertStarts.SetNumZeroed(Result.Spheres.Num());
+			ClustersToUpdate.Reset();
+			for (int32 ClusterIdx = 0, ClusterStart = 0; ClusterIdx < Result.Spheres.Num(); ++ClusterIdx)
+			{
+				if (FullPassesToUpdateCluster[ClusterIdx] > 0)
+				{
+					ClustersToUpdate.Add(ClusterIdx);
+					ClusterVertStarts[ClusterIdx] = ClusterStart;
+					ClusterStart += ClusterSize[ClusterIdx];
+					ClusterSize[ClusterIdx] = 0;
+				}
+			}
+			// Distribute the vertex IDs to their corresponding cluster in the VertsByCluster array
+			for (int32 VID = 0; VID < Mesh.MaxVertexID(); ++VID)
+			{
+				int32 ClusterIdx = Result.VIDtoClusterIndex[VID];
+				if (ClusterIdx != INDEX_NONE && FullPassesToUpdateCluster[ClusterIdx])
+				{
+					VertsByCluster[ClusterVertStarts[ClusterIdx] + ClusterSize[ClusterIdx]++] = VID;
+				}
+			}
+
+			// Optimize the sphere for each cluster
+			ParallelFor(ClustersToUpdate.Num(),
+				[this, &MeshBVH, &MeshFWTree, &Mesh, &Result, &VIDtoTIDs, &TriAreas, &TriNormals, InitialRadius, &MaxCenterBounds, MaxRadius, &FullPassesToUpdateCluster, &ClustersToUpdate, &VtxSQEMs,
+				&ClusterVertStarts, &ClusterSize, &VertsByCluster](int32 ClustersToUpdateIdx)
+			{
+				int32 ClusterIdx = ClustersToUpdate[ClustersToUpdateIdx];
+				if (!ensure(FullPassesToUpdateCluster[ClusterIdx] > 0))
+				{
+					return;
+				}
+
+				// optimization is not reliable if the cluster has very few vertices; better to leave the sphere as-is
+				if (ClusterSize[ClusterIdx] < FMath::Max(2, MinClusterSizeToOptimize))
+				{
+					return;
+				}
+				// the plane-error part of A is just an area-weighted outer product of [N 1] with itself, same as for the SQEM metric, 
+				// and doesn't depend on the sphere parameters, so we can compute it once outside of the newton iterations
+				FEigMat4 A_SQEM = FEigMat4::Zero();
+				for (int32 Idx = ClusterVertStarts[ClusterIdx], EndIdx = ClusterVertStarts[ClusterIdx] + ClusterSize[ClusterIdx]; Idx < EndIdx; ++Idx)
+				{
+					int32 VID = VertsByCluster[Idx];
+					A_SQEM += VtxSQEMs[VID].A;
+				}
+				for (int32 Iter = 0; Iter < MaxGaussNewtonIterations; ++Iter)
+				{
+					FEigMat4 A = FEigMat4::Zero();
+					FEigVec4 B = FEigVec4::Zero();
+					for (int32 Idx = ClusterVertStarts[ClusterIdx], EndIdx = ClusterVertStarts[ClusterIdx] + ClusterSize[ClusterIdx]; Idx < EndIdx; ++Idx)
+					{
+						int32 VID = VertsByCluster[Idx];
+						checkSlow(Result.VIDtoClusterIndex[VID] == ClusterIdx);
+						FVector3d VertPos = Mesh.GetVertex(VID);
+						double AreaSum = 0;
+						for (int32 TID : VIDtoTIDs[VID])
+						{
+							double AreaWt = (double)TriAreas[TID];
+							AreaSum += AreaWt;
+							FVector3d TriNormal = (FVector3d)TriNormals[TID];
+							FEigVec4 N_1(TriNormal.X, TriNormal.Y, TriNormal.Z, 1.);
+							// Plane-error contribution to A added below via A_SQEM
+							B += N_1 * (AreaWt * ((VertPos - Result.Spheres[ClusterIdx].Center).Dot((FVector3d)TriNormal) - Result.Spheres[ClusterIdx].W));
+						}
+						FVector3d FromCenterDir = (VertPos - Result.Spheres[ClusterIdx].Center);
+						double FromCenterDist = Normalize(FromCenterDir);
+						FEigVec4 D_1(FromCenterDir.X, FromCenterDir.Y, FromCenterDir.Z, 1.);
+						A += D_1 * (D_1.transpose() * (AreaSum * PosErrorWt));
+						B += D_1 * ((AreaSum * PosErrorWt) * (FromCenterDist - Result.Spheres[ClusterIdx].W));
+					}
+					A += A_SQEM;
+					FEigVec4 Delta = A.ldlt().solve(B);
+					if (!Delta.hasNaN())
+					{
+						FVector3d NewCenter = Result.Spheres[ClusterIdx].Center + FVector3d(Delta[0], Delta[1], Delta[2]);
+						double NewRadius = Result.Spheres[ClusterIdx].W + Delta[3];
+						
+						// if the optimization starts sending the sphere off to a tiny/negative radius, a huge size or outside the input shape bounds, bail out of the optimization
+						if (NewRadius <= UE_DOUBLE_SMALL_NUMBER || NewRadius > MaxRadius || !MaxCenterBounds.Contains(NewCenter))
+						{
+							break;
+						}
+
+						Result.Spheres[ClusterIdx].Center = NewCenter;
+						Result.Spheres[ClusterIdx].W = NewRadius;
+
+						if (Delta.squaredNorm() < GaussNewtonConvergenceThreshold * GaussNewtonConvergenceThreshold)
+						{
+							break;
+						}
+					}
+					else // solve failed, stop running the optimization
+					{
+						break;
+					}
+				}
+				// Project the optimized sphere to a medial sphere, via shrinking
+				ShrinkMedial.Project(MeshBVH, Result.Spheres[ClusterIdx], InitialRadius, MeshFWTree);
+
+				// Decrement the UpdatedCluster counter, so we can stop updating clusters that haven't changed recently
+				if (--FullPassesToUpdateCluster[ClusterIdx] < 0)
+				{
+					FullPassesToUpdateCluster[ClusterIdx] = 0;
+				}
+			}, EParallelForFlags::Unbalanced);
+
+			UpdateClusterAssignments(true /*bStaleClusterNeighbors, they can be stale here b/c e.g. clusters could be deleted above*/);
+
+			// Re-compute cluster neighbors
+			Result.ComputeClusterNeighbors(Mesh, VIDtoWeldVIDFn);
+			if (ItersWithoutIncrease > 3)
+			{
+				break;
+			}
+
+			// Stop adding spheres if error stops decreasing
+			if (MinErrorDelta > 0)
+			{
+				double AvgError = 0;
+				for (int32 VID = 0; VID < Mesh.MaxVertexID(); ++VID)
+				{
+					if (Result.VIDtoClusterIndex[VID] != INDEX_NONE)
+					{
+						double AreaWt = (double)VtxAreas[VID];
+						AvgError += AreaWt * GetVertexAvgSquaredError(VID, Result.Spheres[Result.VIDtoClusterIndex[VID]]);
+					}
+				}
+				AvgError /= TotalArea;
+				if (!bHasSurfaceIntersections && LastAvgError - AvgError < MinErrorDelta * MinErrorDelta)
+				{
+					break;
+				}
+				LastAvgError = AvgError;
+			}
+
+		}
+
+		// Do a final pass of cluster assignment, filtering, and connectivity
+		UpdateClusterAssignments(true /*bStaleClusterNeighbors*/);
+		UpdateClusterSize(Mesh, Result, ClusterSize);
+		RemoveSmallClusters(Result, ClusterSize, RemapClusters, FullPassesToUpdateCluster, AssignVertexCluster);
+		Result.ComputeClusterNeighbors(Mesh, VIDtoWeldVIDFn);
+		Result.ComputeClusterTriangles(Mesh, VIDtoWeldVIDFn);
+
+		return Result;
+	}
+
+	template<typename MeshType>
+	FMedialSkeleton FSkeletonViaSampling::ComputeSkeleton(const TMeshAABBTree3<MeshType>& MeshBVH, TFunctionRef<bool(int32)> OptionalBoundaryVertexFilter, double ImplicitVertexNormalWeldThreshold, const TFastWindingTree<MeshType>* MeshFWTree)
+	{
+		const MeshType& Mesh = *MeshBVH.GetMesh();
+
+		TArray<FVector3f> Normals;
+		Normals.SetNumUninitialized(Mesh.MaxVertexID());
+		auto NormalFn = [&Normals](int32 VID) -> FVector3d { return (FVector3d)Normals[VID]; };
+		if (ImplicitVertexNormalWeldThreshold > 0)
+		{
+			FVertexConnectedComponents WeldVerts(Mesh.MaxVertexID());
+			WeldVerts.ConnectCloseFilteredVertices(Mesh, ImplicitVertexNormalWeldThreshold, OptionalBoundaryVertexFilter);
+			// Extract the fully-compacted weld parent array, since we will access it a lot
+			TArray<int32> WeldIdx;
+			WeldIdx.SetNumUninitialized(Mesh.MaxVertexID());
+			ParallelFor(Mesh.MaxVertexID(), [&WeldIdx, &WeldVerts](int32 VID) { WeldIdx[VID] = WeldVerts.GetComponentThreadsafe(VID); });
+			MeshNormals::ComputeWeldedVertexNormals(Mesh, [&WeldIdx](int32 VID)->int32 {return WeldIdx[VID];}, Normals);
+			return ComputeSkeleton(MeshBVH, NormalFn, [&WeldIdx](int32 VID)->int32 {return WeldIdx[VID];}, MeshFWTree);
+		}
+		else
+		{
+			MeshNormals::ComputeVertexNormals(Mesh, Normals);
+			return ComputeSkeleton(MeshBVH, NormalFn, [](int32 VID)->int32 {return VID;}, MeshFWTree);
+		}
+	}
+
+	void FSkeletonViaSampling::RemoveSmallClusters(FMedialSkeleton& Result, TArray<int32>& InOutClusterSize, 
+		TArray<int32>& OutRemapClusters, 
+		TArray<std::atomic<int8>>& OutUpdatedClusters,
+		TFunctionRef<void(int32)> AssignVertexCluster)
+	{
+		const int32 MinimumClusterSize = FMath::Max(MinClusterSizeToKeep, 0);
+		checkSlow(InOutClusterSize.Num() == Result.Spheres.Num());
+		OutRemapClusters.Init(INDEX_NONE, Result.Spheres.Num());
+		int32 ValidClusters = 0, LargestClusterIdx = 0, LargestClusterSize = -1;
+		for (int32 ClusterIdx = 0; ClusterIdx < Result.Spheres.Num(); ++ClusterIdx)
+		{
+			if (InOutClusterSize[ClusterIdx] > LargestClusterSize)
+			{
+				LargestClusterSize = InOutClusterSize[ClusterIdx];
+				LargestClusterIdx = ClusterIdx;
+			}
+			if (InOutClusterSize[ClusterIdx] < MinimumClusterSize)
+			{
+				OutRemapClusters[ClusterIdx] = INDEX_NONE;
+			}
+			else
+			{
+				OutRemapClusters[ClusterIdx] = ValidClusters++;
+			}
+		}
+		// recover from edge case where all clusters were deleted, by keeping the largest cluster
+		if (ValidClusters == 0)
+		{
+			OutRemapClusters[LargestClusterIdx] = ValidClusters++;
+		}
+		for (int32 ClusterIdx = 0; ClusterIdx < Result.Spheres.Num(); ++ClusterIdx)
+		{
+			int32 NewIdx = OutRemapClusters[ClusterIdx];
+			if (NewIdx != INDEX_NONE)
+			{
+				Result.Spheres[NewIdx] = Result.Spheres[ClusterIdx];
+				InOutClusterSize[NewIdx] = InOutClusterSize[ClusterIdx];
+				OutUpdatedClusters[NewIdx] = OutUpdatedClusters[ClusterIdx].load();
+			}
+		}
+		Result.Spheres.SetNum(ValidClusters, EAllowShrinking::No);
+		InOutClusterSize.SetNum(ValidClusters, EAllowShrinking::No);
+		OutUpdatedClusters.SetNum(ValidClusters, EAllowShrinking::No);
+		for (int32 VID = 0; VID < Result.VIDtoClusterIndex.Num(); ++VID)
+		{
+			int32& ClusterIdx = Result.VIDtoClusterIndex[VID];
+			if (ClusterIdx != INDEX_NONE)
+			{
+				ClusterIdx = OutRemapClusters[ClusterIdx];
+				if (ClusterIdx == INDEX_NONE)
+				{
+					AssignVertexCluster(VID);
+					InOutClusterSize[Result.VIDtoClusterIndex[VID]]++;
+				}
+			}
+		}
+	}
+
+
+	// Instantiate templated ComputeSkeleton methods for the supported mesh types
+	template FMedialSkeleton FSkeletonViaSampling::ComputeSkeleton<FDynamicMesh3>(const TMeshAABBTree3<FDynamicMesh3>& MeshBVH, TFunctionRef<FVector3d(int32)> GetVertexNormal, TFunctionRef<int32(int32)> VIDtoWeldVIDFn, const TFastWindingTree<FDynamicMesh3>* MeshFWTree);
+	template FMedialSkeleton FSkeletonViaSampling::ComputeSkeleton<FDynamicMesh3>(const TMeshAABBTree3<FDynamicMesh3>& MeshBVH, TFunctionRef<bool(int32)> OptionalBoundaryVertexFilter, double ImplicitVertexNormalWeldThreshold, const TFastWindingTree<FDynamicMesh3>* MeshFWTree);
+	template FMedialSkeleton FSkeletonViaSampling::ComputeSkeleton<FTriangleMeshAdapterd>(const TMeshAABBTree3<FTriangleMeshAdapterd>& MeshBVH, TFunctionRef<FVector3d(int32)> GetVertexNormal, TFunctionRef<int32(int32)> VIDtoWeldVIDFn, const TFastWindingTree<FTriangleMeshAdapterd>* MeshFWTree);
+	template FMedialSkeleton FSkeletonViaSampling::ComputeSkeleton<FTriangleMeshAdapterd>(const TMeshAABBTree3<FTriangleMeshAdapterd>& MeshBVH, TFunctionRef<bool(int32)> OptionalBoundaryVertexFilter, double ImplicitVertexNormalWeldThreshold, const TFastWindingTree<FTriangleMeshAdapterd>* MeshFWTree);
+	
+
+
+} // namespace UE::Geometry::MedialAxis
\ No newline at end of file
diff --git a/Engine/Plugins/Runtime/GeometryProcessing/Source/GeometryAlgorithms/Public/Skeletonization/MeshMedialAxisSampling.h b/Engine/Plugins/Runtime/GeometryProcessing/Source/GeometryAlgorithms/Public/Skeletonization/MeshMedialAxisSampling.h
new file mode 100644
--- /dev/null
+++ b/Engine/Plugins/Runtime/GeometryProcessing/Source/GeometryAlgorithms/Public/Skeletonization/MeshMedialAxisSampling.h
@@ -0,0 +1,433 @@
+// Copyright Epic Games, Inc. All Rights Reserved.
+
+#pragma once
+
+#include "MathUtil.h"
+#include "VectorTypes.h"
+#include "FrameTypes.h"
+#include "Util/ProgressCancel.h"
+
+#include "Async/ParallelFor.h"
+#include "Distance/DistPoint3Triangle3.h"
+#include "DynamicMesh/DynamicMesh3.h"
+#include "Spatial/FastWinding.h"
+#include "Spatial/MeshAABBTree3.h"
+
+
+class FProgressCancel;
+
+namespace UE::Geometry
+{
+
+namespace MedialAxis
+{
+	// Algorithm adapted from "3D medial axis point approximation using nearest neighbors and the normal field" by Ma et al.
+	// Starting from a surface point + normal, aims to find a "medial axis ball" touching that point and one other point
+	// w/ center along Point,-Normal ray, via iteratively shrinking an estimated radius to touch the surface point + a nearest point
+	// 
+	// Note that the returned sphere may not touch the input surface point, esp. in corner cases where the actual medial ball would have zero radius.
+	struct FShrinkMedialBall
+	{
+		// Max iterations of shrinking to apply
+		int32 MaxIters = 20;
+		// Threshold to stop iterating, if radius or center changes less than this
+		double MinDelta = UE_DOUBLE_KINDA_SMALL_NUMBER;
+		// Stop iterating below this radius (if > 0). Returned radius may still be smaller than this.
+		double RadiusThreshold = UE_DOUBLE_KINDA_SMALL_NUMBER;
+
+		// Method to 'hot start' the search via local vertex neighborhood, using a provided neighborhood
+		template<typename MeshType>
+		bool FindLocalMinCurvatureRadius(const MeshType& Mesh, FVector3d& SurfacePoint, FVector3d& SurfaceNormal, TConstArrayView<int32> NbrVIDs, double& OutFoundRadius)
+		{
+			bool bFound = false;
+			for (int32 NbrVID : NbrVIDs)
+			{
+				bFound |= ShrinkFromLocalVertex(SurfacePoint, SurfaceNormal, Mesh.GetVertex(NbrVID), OutFoundRadius);
+			}
+			return bFound;
+		}
+		// Method to 'hot start' the search via local vertex neighborhood, using the neighborhood information in the dynamic mesh
+		bool FindLocalMinCurvatureRadius(const FDynamicMesh3& Mesh, FVector3d& SurfacePoint, FVector3d& SurfaceNormal, int32 SurfaceVID, double& OutFoundRadius)
+		{
+			bool bFound = false;
+			for (int32 NbrVID : Mesh.VtxVerticesItr(SurfaceVID))
+			{
+				bFound |= ShrinkFromLocalVertex(SurfacePoint, SurfaceNormal, Mesh.GetVertex(NbrVID), OutFoundRadius);
+			}
+			return bFound;
+		}
+
+
+		/**
+		 * Shrink a sphere touching a given SurfacePoint, with center along the line through its SurfaceNormal
+		 * until the sphere approximately touches two closest surface points including the input point
+		 * Note: Sphere will often not actually touch the input surface point due to discretization error, since we set the radius s.t. the sphere remains inside the input surface
+		 * 
+		 * @param MeshBVH BVH of the surface
+		 * @param SurfacePoint Reference point that the medial ball should approximately touch
+		 * @param SurfaceNormal Normal at the SurfacePoint; note for vertices/edges this should be a smooth normal
+		 * @param bRadiusHotStart If true, the InitialRadius should be a more accurate initial guess e.g. provided via one of the FindLocalMinCurvatureRadius methods
+		 * @param InitialRadius If non-negative, initial radius to use for the sphere-shrinking search. Otherwise, will be automatically set from the surface bounding box's max dimension
+		 * @param EarlyStopAtTIDFn Use when starting from a vertex/edge to stop in the discrete neighborhood of that source vertex/edge. Passed the current closest triangle ID.
+		 */
+		template<typename MeshType>
+		FSphere Find(const TMeshAABBTree3<MeshType>& MeshBVH, const FVector3d& SurfacePoint, const FVector3d& SurfaceNormal, bool bRadiusHotStart,
+			double InitialRadius, TFunctionRef<bool(int32)> EarlyStopAtTIDFn)
+		{
+			const MeshType& Mesh = *MeshBVH.GetMesh();
+
+			if (SurfaceNormal.IsZero())
+			{
+				return FSphere(SurfacePoint, 0.);
+			}
+
+			if (InitialRadius < 0)
+			{
+				InitialRadius = MeshBVH.GetBoundingBox().MaxDim();
+			}
+
+			double CurRadius = InitialRadius;
+			FVector3d CurCenter = SurfacePoint - SurfaceNormal * CurRadius;
+			double SafeDist = CurRadius;
+			for (int32 Iter = (int32)bRadiusHotStart; /*note: we break at MaxIters only after updating SafeDist*/; ++Iter)
+			{
+				IMeshSpatial::FQueryOptions Options;
+				Options.MaxDistance = CurRadius;
+				double NearDistSq;
+				int32 NearTriID = MeshBVH.FindNearestTriangle(CurCenter, NearDistSq, Options);
+				// No nearest point found; either mesh was empty or last iteration converged & nothing else was within MaxDistance of search
+				if (NearTriID == INDEX_NONE)
+				{
+					SafeDist = CurRadius;
+					break;
+				}
+				SafeDist = FMath::Sqrt(NearDistSq);
+				
+				// Note: All the stopping conditions are tested after updating SafeDist so that if we stop, the returned sphere 
+				// (which uses SafeDist for radius) doesn't extend past the closest point
+				
+				// Radius below stopping threshold
+				// Note: it's important that we test vs CurRadius and not SafeDist, because SafeDist can be spuriously small 
+				// when the CurRadius is still large (as the sphere center moves past another surface)
+				if (CurRadius <= RadiusThreshold)
+				{
+					break;
+				}
+
+				if (Iter >= MaxIters)
+				{
+					break;
+				}
+
+				if (EarlyStopAtTIDFn(NearTriID))
+				{
+					break;
+				}
+
+				FDistPoint3Triangle3d Query = TMeshQueries<MeshType>::TriangleDistance(Mesh, NearTriID, CurCenter);
+				FVector3d ClosePoint = Query.ClosestTrianglePoint;
+
+				FVector3d SmC = (SurfacePoint - ClosePoint);
+				double SampleDistsSq = SmC.SquaredLength();
+				// If closest point is (effectively) the source point, we can't shrink further
+				if (SampleDistsSq < FMathd::ZeroTolerance)
+				{
+					break;
+				}
+				double SmCdN = SmC.Dot(SurfaceNormal);
+				// closest point effectively co-planar, likely converged + would have unstable radius
+				if (SmCdN < FMathd::ZeroTolerance)
+				{
+					break;
+				}
+				double OldRadius = CurRadius;
+				CurRadius = SampleDistsSq / (2 * SmCdN);
+				// radius converged (note OldRadius should always be larger than CurRadius by construction)
+				if (OldRadius - CurRadius < MinDelta)
+				{
+					break;
+				}
+				CurCenter = SurfacePoint - SurfaceNormal * CurRadius;
+			}
+
+			// Note we use SafeDist instead of CurRadius to push the sphere inside
+			// the target surface, especially in cases where the algorithm terminates early ...
+			// On converged cases, this should be approximately the same as using CurRadius.
+			return FSphere(CurCenter, SafeDist);
+		}
+
+		template<typename MeshType>
+		FSphere FindWithImplicitWeld(const TMeshAABBTree3<MeshType>& MeshBVH, int32 WeldedSurfacePointVID, TFunctionRef<int32(int32)> VIDtoWeldVIDFn, 
+			const FVector3d& SurfacePoint, const FVector3d& SurfaceNormal, bool bRadiusHotStart,
+			double InitialRadius)
+		{
+			return Find(MeshBVH, SurfacePoint, SurfaceNormal, bRadiusHotStart, InitialRadius, [&MeshBVH, &VIDtoWeldVIDFn, WeldedSurfacePointVID](int32 TID)
+			{
+				FIndex3i Tri = MeshBVH.GetMesh()->GetTriangle(TID);
+				Tri.A = VIDtoWeldVIDFn(Tri.A);
+				Tri.B = VIDtoWeldVIDFn(Tri.B);
+				Tri.C = VIDtoWeldVIDFn(Tri.C);
+				return Tri.Contains(WeldedSurfacePointVID);
+			});
+		}
+
+		/**
+		 * Project an existing sphere to a surface-contained sphere, searching along the line through the initial sphere center and corresponding closest surface point
+		 * @param MeshBVH BVH of the surface
+		 * @param InOutToProject Sphere to project
+		 * @param InitialSearchRadius If non-negative, initial radius to use for the sphere-shrinking search. Otherwise, will be automatically set from the surface bounding box's max dimension
+		 * @param MeshFWTree If non-null, test the fast winding of the center before projecting, and reverse the search direction if it's outside. Otherwise, a (less robust) normal test will be used instead.
+		 * @return true if the sphere was updated via projected, false otherwise (e.g. if the BVH was empty, or if we couldn't get a valid search normal)
+		 */
+		template<typename MeshType>
+		bool Project(const TMeshAABBTree3<MeshType>& MeshBVH,
+			FSphere& InOutToProject,
+			double InitialSearchRadius = -1,
+			const TFastWindingTree<MeshType>* MeshFWTree = nullptr
+		)
+		{
+			double InitialCenterDistSq;
+			int32 NearTriID = MeshBVH.FindNearestTriangle(InOutToProject.Center, InitialCenterDistSq);
+			if (NearTriID == INDEX_NONE)
+			{
+				return false;
+			}
+
+			const MeshType& Mesh = *MeshBVH.GetMesh();
+			FTriangle3d Tri;
+			Mesh.GetTriVertices(NearTriID, Tri.V[0], Tri.V[1], Tri.V[2]);
+			FDistPoint3Triangle3d Query(InOutToProject.Center, Tri);
+			Query.GetSquared();
+			FVector3d ClosePoint = Query.ClosestTrianglePoint;
+			FVector3d Normal = Tri.Normal();
+			FVector3d RadiusSearchDir = ClosePoint - InOutToProject.Center;
+			if (!RadiusSearchDir.Normalize())
+			{
+				// the center was almost exactly on the surface; try falling back to the surface normal
+				// TODO: if we're on a vertex or edge of the triangle, would ideally use the smooth normal
+				RadiusSearchDir = Normal;
+				if (Normal.IsZero())
+				{
+					return false;
+				}
+			}
+			// if the center is outside the surface, we want to flip the search normal; optionally use the fast winding for this. 
+			// If no fast winding available, test the closest normal (but this can fail near sharp features!)
+			else if (((MeshFWTree && !MeshFWTree->IsInside(InOutToProject.Center)) || (!MeshFWTree && Normal.Dot(RadiusSearchDir) < 0.)))
+			{
+				RadiusSearchDir *= -1.;
+			}
+			InOutToProject = Find(MeshBVH, ClosePoint, RadiusSearchDir, false, InitialSearchRadius, [](int32 TID) {return false;});
+			return true;
+		}
+
+	private:
+		// Update curvature estimate from local vertex in neighborhood of SurfacePoint
+		bool ShrinkFromLocalVertex(const FVector3d& SurfacePoint, const FVector3d& SurfaceNormal, const FVector3d& LocalVertex, double& InOutRadius)
+		{
+			FVector3d LocalToSurface = SurfacePoint - LocalVertex;
+			double LocalDistSq = LocalToSurface.SquaredLength();
+			// ignore near-duplicate neighbors
+			if (LocalDistSq < FMathd::ZeroTolerance)
+			{
+				return false;
+			}
+			double LocalDistCosTheta = LocalToSurface.Dot(SurfaceNormal);
+			if (LocalDistCosTheta <= FMathd::ZeroTolerance)
+			{
+				return false;
+			}
+
+			InOutRadius = FMath::Min(InOutRadius, LocalDistSq / (2 * LocalDistCosTheta));
+			return true;
+		}
+	};
+
+	/**
+	 * Simple representation of a surface's medial skeleton, as a set of spheres + connectivity on those spheres, and a mapping to the surface vertices
+	 */
+	struct FMedialSkeleton
+	{
+		// Medial spheres, 1 per cluster, indexed by cluster index
+		TArray<FSphere> Spheres;
+		// Mapping from source mesh vertex IDs to cluster Index
+		TArray<int32> VIDtoClusterIndex;
+
+		// Inline-allocated array used for representing the typically-small neighborhoods in the cluster connectivity graph
+		using FNbrArray = TArray<int32, TInlineAllocator<16>>;
+		// Indices of neighbors of clusters, for each cluster
+		TArray<FNbrArray> ClusterNeighbors;
+		// Triangles of the medial skeleton, where the skeleton is a surface rather than a chain. Note: Not oriented, and generally not manifold -- so cannot be directly appended to an FDynamicMesh3
+		TArray<FIndex3i> ClusterTriangles;
+
+
+		// Compute the ClusterTriangles array for the current FMedialSkeleton, given a Mesh and assuming VIDtoClusterIndex is already computed
+		template<typename TriangleMeshType>
+		void ComputeClusterTriangles(const TriangleMeshType& Mesh, TFunctionRef<int32(int32)> VIDtoWeldVIDFn)
+		{
+			TSet<FIndex3i> Tris;
+			for (int32 TID = 0; TID < Mesh.MaxTriangleID(); ++TID)
+			{
+				if (!Mesh.IsTriangle(TID))
+				{
+					continue;
+				}
+
+				FIndex3i Tri = Mesh.GetTriangle(TID);
+				Tri.A = VIDtoClusterIndex[VIDtoWeldVIDFn(Tri.A)];
+				Tri.B = VIDtoClusterIndex[VIDtoWeldVIDFn(Tri.B)];
+				Tri.C = VIDtoClusterIndex[VIDtoWeldVIDFn(Tri.C)];
+				// Only consider triangles that span three different, valid clusters
+				if (Tri.Contains(INDEX_NONE) || (Tri.A == Tri.B || Tri.B == Tri.C || Tri.A == Tri.C))
+				{
+					continue;
+				}
+				// Since the cluster tris are not oriented, sort the triangle indices to get a unique handle
+				Tri.Sort();
+				Tris.Add(Tri);
+			}
+			ClusterTriangles = Tris.Array();
+		}
+
+		// Compute the ClusterNeighbors array for the current FMedialSkeleton, given a Mesh and assuming Spheres and VIDtoClusterIndex are already computed
+		template<typename TriangleMeshType>
+		void ComputeClusterNeighbors(const TriangleMeshType& Mesh, TFunctionRef<int32(int32)> VIDtoWeldVIDFn)
+		{
+			ClusterNeighbors.SetNum(Spheres.Num());
+			for (FMedialSkeleton::FNbrArray& Nbrs : ClusterNeighbors)
+			{
+				Nbrs.Reset();
+			}
+			for (int32 TID = 0; TID < Mesh.MaxTriangleID(); ++TID)
+			{
+				if (!Mesh.IsTriangle(TID))
+				{
+					continue;
+				}
+
+				FIndex3i Tri = Mesh.GetTriangle(TID);
+				for (int32 Prev = 2, Idx = 0; Idx < 3; Prev = Idx++)
+				{
+					int32 ClusterA = VIDtoClusterIndex[VIDtoWeldVIDFn(Tri[Idx])];
+					int32 ClusterB = VIDtoClusterIndex[VIDtoWeldVIDFn(Tri[Prev])];
+					if (ClusterA != INDEX_NONE && ClusterB != INDEX_NONE && ClusterA != ClusterB)
+					{
+						// note: AddUnique to arrays b/c we expect neighbor arrays to be very small
+						ClusterNeighbors[ClusterA].AddUnique(ClusterB);
+						ClusterNeighbors[ClusterB].AddUnique(ClusterA);
+					}
+				}
+			}
+		}
+	};
+
+	// Implementation/adaptation of "Dynamic Skeletonization via Variational Medial Skeletal Sampling" by Huang et al
+	// with some surface-focused additions -- optionally uses raycasts to guide refinement, and fast winding to guide medial-sphere projection
+	struct FSkeletonViaSampling
+	{
+		/// Main user parameters
+
+		// Stop at this sphere count
+		int32 MaxSpheres = 10;
+		// Stop when the average squared error at vertices decreases by less than this w/ an added sphere addition
+		double MinErrorDelta = .1;
+		// Do not split a cluster if its max error is below this threshold
+		double MinClusterErrorToSplit = .1;
+		// Weight for the position error term. Relative to plane error term, which has weight 1.
+		double PosErrorWt = .2;
+		// Whether to test medial skeleton edges for intersection w/ the input surface, and try to refine where these intersections are found
+		bool bSplitClustersIfEdgesIntersectSurface = true;
+		// Factor to scale vertex error near skeleton-edge/input-surface intersections
+		double ErrorScaleNearEdgeSurfaceIntersection = 10.;
+		// Clusters with fewer vertices than this will be discarded
+		int32 MinClusterSizeToKeep = 4;
+
+		// Settings for the medial ball shrinking, used to initialize medial spheres and to project spheres into the surface
+		FShrinkMedialBall ShrinkMedial;
+
+
+		/// Additional parameters -- most users should not need to change these
+
+		// Maximum cluster splits to apply per iteration (if > 0) --
+		// Setting this to a low value can force the algorithm to target higher error more aggressively, at the cost of 
+		// requiring many more full iterations.
+		// Prefer using MaxSplitFractionPerIteration and MinFractionOfMaxErrorToSplit instead, which scale better with the number of clusters.
+		int32 MaxSplitPerIteration = -1;
+		// Maximum fraction of current cluster count to split per iteration (if > 0) --
+		// Functionally similar to MaxSplitPerIteration but scales better with larger cluster counts
+		double MaxSplitFractionPerIteration = .2;
+		// Minimum vertex error to allow a cluster split, as fraction of current maximum vertex error --
+		// This helps to avoid splitting relatively low-error regions when there is significantly-higher error elsewhere since
+		// the higher error region will often split into multiple clusters that should all be split before the low error region
+		double MinFractionOfMaxErrorToSplit = .75;
+
+		// Number of full passes (top-level skeleton building iterations) to keep optimizing spheres after the cluster vertex assignments have changed
+		// This allows us to 'sleep' the spheres for clusters where vertices have not changed.
+		// Note: Should generally be at least 1 so that at least one optimization pass is applied.
+		// Higher values will allow for the optimization to smooth out the result of the sphere projection,
+		// (since projection happens separately after the optimization) and potentially amortize gauss-newton iterations
+		// over multiple passes -- i.e., one could try raising this value and lowering MaxGaussNewtonIterations.
+		int8 OptimizeSpheresForIterationsPostClusterChange = 2;
+		// Maximum number of gauss-newton iterations to apply to improve sphere fit
+		int32 MaxGaussNewtonIterations = 10;
+		// Stop optimization if the norm of the change in sphere parameters is less than this
+		double GaussNewtonConvergenceThreshold = UE_DOUBLE_KINDA_SMALL_NUMBER;
+		// Minimum number of vertices in a cluster below which we will not optimize (because the optimization may not capture a good medial sphere if it has too few samples)
+		int32 MinClusterSizeToOptimize = 4;
+		// Number of clusters at which to switch to local cluster assignment updates, rather than global searches. If < 0, always use global search.
+		int32 UseLocalClusterVertexReassignmentAtClusterCount = 100;
+
+		
+		/**
+		 * Compute mesh skeleton
+		 * Current Supported Mesh Types: FDynamicMesh3, FTriangleMeshAdapterd
+		 *
+		 * @param MeshBVH BVH of the surface
+		 * @param GetVertexNormal Helper to provide per-vertex surface normals
+		 * @param VIDtoWeldVIDFn Mapping from vertex ID to implicit 'welded' vertex ID, to handle cases where multiple vertices should be treated as one
+		 * @param MeshFWTree Optional fast winding tree; if provided will be used for more-robust medial sphere projection
+		 */
+		template<typename MeshType>
+		GEOMETRYALGORITHMS_API FMedialSkeleton ComputeSkeleton(const TMeshAABBTree3<MeshType>& MeshBVH, TFunctionRef<FVector3d(int32)> GetVertexNormal, TFunctionRef<int32(int32)> VIDtoWeldVIDFn, const TFastWindingTree<MeshType>* MeshFWTree = nullptr);
+
+		/**
+		 * Compute mesh skeleton
+		 * This version will compute vertex normals for you, and will optionally consider near-coincident vertices as welded
+		 * Current Supported Mesh Types: FDynamicMesh3, FTriangleMeshAdapterd
+		 *
+		 * @param MeshBVH BVH of the surface
+		 * @param OptionalWeldVertexFilter Optionally filter what vertices are candidates to be welded. Return true if the vertex is a candidate to weld.
+		 * @param ImplicitVertexNormalWeldThreshold If > 0, when computing vertex normals will consider coincident boundary vertices as implicitly welded if they're closer than this threshold distance
+		 * @param MeshFWTree Optional fast winding tree; if provided will be used for more-robust medial sphere projection
+		 */
+		template<typename MeshType>
+		GEOMETRYALGORITHMS_API FMedialSkeleton ComputeSkeleton(const TMeshAABBTree3<MeshType>& MeshBVH, TFunctionRef<bool(int32)> OptionalWeldVertexFilter = [](int32)->bool {return true;}, double ImplicitVertexNormalWeldThreshold = UE_DOUBLE_KINDA_SMALL_NUMBER, const TFastWindingTree<MeshType>* MeshFWTree = nullptr);
+
+		/**
+		 * Compute mesh skeleton
+		 * This version will compute vertex normals for you, and will optionally consider near-coincident boundary vertices as welded
+		 *
+		 * @param MeshBVH BVH of the surface
+		 * @param ImplicitVertexNormalWeldThreshold If > 0, when computing vertex normals will consider coincident boundary vertices as implicitly welded if they're closer than this threshold distance
+		 * @param MeshFWTree Optional fast winding tree; if provided will be used for more-robust medial sphere projection
+		 */
+		FMedialSkeleton ComputeSkeleton(const TMeshAABBTree3<FDynamicMesh3>& MeshBVH, double ImplicitVertexNormalWeldThreshold = UE_DOUBLE_KINDA_SMALL_NUMBER, const TFastWindingTree<FDynamicMesh3>* MeshFWTree = nullptr)
+		{
+			const FDynamicMesh3& Mesh = *MeshBVH.GetMesh();
+			return ComputeSkeleton(MeshBVH, [&Mesh](int32 VID)->bool {return Mesh.IsBoundaryVertex(VID);}, ImplicitVertexNormalWeldThreshold, MeshFWTree);
+		}
+
+	private:
+
+		// Helper to remove clusters smaller than MinClusterSizeToKeep
+		void RemoveSmallClusters(FMedialSkeleton& Result, TArray<int32>& InOutClusterSize, 
+			TArray<int32>& OutRemapClusters, TArray<std::atomic<int8>>& OutUpdatedClusters,
+			TFunctionRef<void(int32)> AssignVertexCluster);
+
+
+	};
+}
+
+
+} // end namespace UE::Geometry
+
