add new topology preservation bitflag value to mesh constraints so that simplification can control which constraints need topology preserved / which can be merged freely

#Synced-CL 50036218

diff --git a/Engine/Plugins/Runtime/GeometryProcessing/Source/DynamicMesh/Private/MeshSimplification.cpp b/Engine/Plugins/Runtime/GeometryProcessing/Source/DynamicMesh/Private/MeshSimplification.cpp
--- a/Engine/Plugins/Runtime/GeometryProcessing/Source/DynamicMesh/Private/MeshSimplification.cpp
+++ b/Engine/Plugins/Runtime/GeometryProcessing/Source/DynamicMesh/Private/MeshSimplification.cpp
@@ -1394,20 +1394,24 @@ bool TMeshSimplification<QuadricErrorType>::CanCollapseEdge(int edgeID, int a, i
 	if (bCanCollapse && !bIndirectConstraintConnection)
 	{
 		// determine if edge constraints require keeping either vert.
-		// if edge ac is constrained, we must keep it and thus vertex a, likewise for edge bc and vertex b.
+		// if edge ac is constrained w.r.t. merging topology, we must keep it and thus vertex a, likewise for edge bc and vertex b.
 		int32 eac = Mesh->FindEdgeFromTri(a, c, tc);
 		int32 ebc = Mesh->FindEdgeFromTri(b, c, tc);
-		bool bMustRetainA = Constraints->HasEdgeConstraint(eac);
-		bool bMustRetainB = Constraints->HasEdgeConstraint(ebc);
+		FEdgeConstraint Ceac = Constraints->GetEdgeConstraint(eac);
+		FEdgeConstraint Cebc = Constraints->GetEdgeConstraint(ebc);
+		bool bMustRetainA = !Ceac.CanMergeTopology();
+		bool bMustRetainB = !Cebc.CanMergeTopology();
 
-		// if edge ad is constrained, we must keep it and thus vertex a, likewise for edge bd and vertex b.
+		// if edge ad is constrained w.r.t. merging topology, we must keep it and thus vertex a, likewise for edge bd and vertex b.
 		if (d != IndexConstants::InvalidID)
 		{
 			int32 ead = Mesh->FindEdgeFromTri(a, d, td);
 			int32 ebd = Mesh->FindEdgeFromTri(b, d, td);
 
-			bMustRetainA = bMustRetainA || Constraints->HasEdgeConstraint(ead);
-			bMustRetainB = bMustRetainB || Constraints->HasEdgeConstraint(ebd);
+			FEdgeConstraint Cead = Constraints->GetEdgeConstraint(ead);
+			FEdgeConstraint Cebd = Constraints->GetEdgeConstraint(ebd);
+			bMustRetainA = bMustRetainA || !Cead.CanMergeTopology();
+			bMustRetainB = bMustRetainB || !Cebd.CanMergeTopology();
 		}
 
 		AddNewABConstraint(bMustRetainA, bMustRetainB);
diff --git a/Engine/Plugins/Runtime/GeometryProcessing/Source/DynamicMesh/Public/MeshConstraints.h b/Engine/Plugins/Runtime/GeometryProcessing/Source/DynamicMesh/Public/MeshConstraints.h
--- a/Engine/Plugins/Runtime/GeometryProcessing/Source/DynamicMesh/Public/MeshConstraints.h
+++ b/Engine/Plugins/Runtime/GeometryProcessing/Source/DynamicMesh/Public/MeshConstraints.h
@@ -25,8 +25,10 @@ enum class EEdgeRefineFlags
 	NoSplit = 2,
 	/** Edge cannot be collapsed */
 	NoCollapse = 4,
-	/** Edge cannot be flipped, split, or collapsed */
-	FullyConstrained = NoFlip | NoSplit | NoCollapse,
+	/** Topology of constraints should not be changed, i.e. disconnected constraint edges should not be connected via an adjacent edge collapse */
+	NoTopologyMerge = 8,
+	/** Edge cannot be flipped, split, collapsed, or topologically-merged */
+	FullyConstrained = NoFlip | NoSplit | NoCollapse | NoTopologyMerge,
 	/** Edge can only be split */
 	SplitsOnly = NoFlip | NoCollapse,
 	/** Edge can only flip*/
@@ -93,6 +95,11 @@ public:
 		return CanCollapse(RefineFlags);
 	}
 
+	bool CanMergeTopology() const
+	{
+		return CanMergeTopology(RefineFlags);
+	}
+
 	/** @return true if edge cannot be modified at all */
 	bool NoModifications() const
 	{
@@ -135,6 +142,11 @@ public:
 		return ((int)Flags & (int)EEdgeRefineFlags::NoCollapse) == 0;
 	}
 
+	static inline bool CanMergeTopology(const EEdgeRefineFlags Flags)
+	{
+		return ((int)Flags & (int)EEdgeRefineFlags::NoTopologyMerge) == 0;
+	}
+
 	/** @return true if flags permit no edge modification */
 	static inline bool NoModifications(const EEdgeRefineFlags Flags)
 	{
