add medial skeleton -> hierarchical / animation skeleton conversion via minimum spanning tree

#Synced-CL 50781459

diff --git a/Engine/Plugins/Runtime/GeometryProcessing/Source/GeometryAlgorithms/Private/Skeletonization/MeshMedialAxisSampling.cpp b/Engine/Plugins/Runtime/GeometryProcessing/Source/GeometryAlgorithms/Private/Skeletonization/MeshMedialAxisSampling.cpp
--- a/Engine/Plugins/Runtime/GeometryProcessing/Source/GeometryAlgorithms/Private/Skeletonization/MeshMedialAxisSampling.cpp
+++ b/Engine/Plugins/Runtime/GeometryProcessing/Source/GeometryAlgorithms/Private/Skeletonization/MeshMedialAxisSampling.cpp
@@ -8,6 +8,12 @@
 #include "MeshAdapter.h"
 #include "MeshAdapter/MeshVertexNormals.h"
 
+#include "DynamicMesh/DynamicBoneAttribute.h"
+#include "DynamicMesh/DynamicMeshAttributeSet.h"
+#include "DynamicMesh/DynamicVertexSkinWeightsAttribute.h"
+
+#include "Util/SizedDisjointSet.h"
+
 // Include Eigen/Dense for a 4x4 matrix solve
 #if defined(_MSC_VER) && USING_CODE_ANALYSIS
 #pragma warning(push)
@@ -742,6 +748,19 @@ namespace UE::Geometry::MedialAxis
 		Result.ComputeClusterNeighbors(Mesh, VIDtoWeldVIDFn);
 		Result.ComputeClusterTriangles(Mesh, VIDtoWeldVIDFn);
 
+		// For unassigned vertices, try to transfer the cluster assignment from the weld vertex
+		for (int32 VID = 0; VID < Result.VIDtoClusterIndex.Num(); ++VID)
+		{
+			if (Result.VIDtoClusterIndex[VID] == INDEX_NONE)
+			{
+				int32 WeldVID = VIDtoWeldVIDFn(VID);
+				if (Result.VIDtoClusterIndex.IsValidIndex(WeldVID) && WeldVID != VID)
+				{
+					Result.VIDtoClusterIndex[VID] = Result.VIDtoClusterIndex[WeldVID];
+				}
+			}
+		}
+
 		return Result;
 	}
 
@@ -829,6 +848,293 @@ namespace UE::Geometry::MedialAxis
 		}
 	}
 
+	void FMedialSkeletonToTreeSkeletonOptions::ToHierarchy(const FMedialSkeleton& InMedialSkeleton, TArray<int32>& OutParent, TArray<FTransform>& OutTransform, FMedialSkeletonToTreeSkeletonOptions Options, int32 InRootIndex,
+		TArray<int32>* OutBoneIndexToMedialIndex, TArray<int32>* OutMedialIndexToBoneIndex)
+	{
+		using namespace UE::Geometry;
+
+		const int32 NumClusters = InMedialSkeleton.ClusterNeighbors.Num();
+		if (!ensureMsgf(NumClusters == InMedialSkeleton.Spheres.Num(), TEXT("Invalid medial skeleton data; %d spheres does not match %d cluster neighbors"), InMedialSkeleton.Spheres.Num(), NumClusters))
+		{
+			return;
+		}
+
+		// setup/clear output arrays
+		TArray<int32> LocalIndexToMedial;
+		TArray<int32>* UseIndexToMedial = OutBoneIndexToMedialIndex ? OutBoneIndexToMedialIndex : &LocalIndexToMedial;
+		UseIndexToMedial->Empty(NumClusters);
+		OutParent.Empty(NumClusters);
+		OutTransform.Empty(NumClusters);
+		if (OutMedialIndexToBoneIndex)
+		{
+			OutMedialIndexToBoneIndex->Empty(NumClusters);
+		}
+
+		if (NumClusters == 0)
+		{
+			return;
+		}
+
+		TArray<FIndex2i> Edges;
+		for (int32 Source = 0; Source < NumClusters; ++Source)
+		{
+			for (int32 Dest : InMedialSkeleton.ClusterNeighbors[Source])
+			{
+				if (!ensure(InMedialSkeleton.ClusterNeighbors.IsValidIndex(Dest)))
+				{
+					continue;
+				}
+				if (Source < Dest)
+				{
+					Edges.Emplace(Source, Dest);
+				}
+			}
+		}
+		
+		if (Options.EdgeWeightMethod == FMedialSkeletonToTreeSkeletonOptions::EEdgeWeightMethod::ArrayOrder)
+		{
+			Edges.Sort([](const FIndex2i& A, const FIndex2i& B) { return A.A + A.B < B.A + B.B; });
+		}
+		else if (Options.EdgeWeightMethod == FMedialSkeletonToTreeSkeletonOptions::EEdgeWeightMethod::AvgRadius)
+		{
+			Edges.Sort([&InMedialSkeleton](const FIndex2i& A, const FIndex2i& B) 
+				{
+					return	InMedialSkeleton.Spheres[A.A].W + InMedialSkeleton.Spheres[A.B].W > // note: greater than, so we prefer connections between larger spheres
+							InMedialSkeleton.Spheres[B.A].W + InMedialSkeleton.Spheres[B.B].W;
+				}
+			);
+		}
+		else // EdgeLength
+		{
+			Edges.Sort([&InMedialSkeleton](const FIndex2i& A, const FIndex2i& B)
+				{
+					return	FVector::DistSquared(InMedialSkeleton.Spheres[A.A].Center, InMedialSkeleton.Spheres[A.B].Center) <
+							FVector::DistSquared(InMedialSkeleton.Spheres[B.A].Center, InMedialSkeleton.Spheres[B.B].Center);
+				}
+			);
+		}
+
+		// Use a disjoint set to track the connectivity as we add one edge at a time
+		FSizedDisjointSet MedialForest;
+		MedialForest.Init(InMedialSkeleton.ClusterNeighbors.Num());
+
+		// KeptNbrs will hold the subset of the medial skeleton ClusterNeighbors that form a spanning tree
+		TArray<TArray<int32, TInlineAllocator<8>>> KeptNbrs;
+		KeptNbrs.SetNum(NumClusters);
+		for (const FIndex2i& Edge : Edges)
+		{
+			// add the edge only if the clusters aren't already indirectly connected by a previously-added edge
+			if (MedialForest.Find(Edge.A) != MedialForest.Find(Edge.B))
+			{
+				MedialForest.Union(Edge.A, Edge.B);
+				KeptNbrs[Edge.A].Add(Edge.B);
+				KeptNbrs[Edge.B].Add(Edge.A);
+			}
+		}
+
+		int32 RootIdx = InMedialSkeleton.Spheres.IsValidIndex(InRootIndex) ? InRootIndex : 0;
+
+		if (Options.MergeDisconnectedMethod == EMergeDisconnectedMethod::ConnectClosestBones)
+		{
+			int32 IterLimit = NumClusters; // sanity check to avoid infinite loop, medial forest may take several iterations to merge to a single tree, but should never take more iterations than there are clusters (and typically, hopefully, much fewer!)
+			while (MedialForest.GetSize(0) != NumClusters && ensure(IterLimit-- > 0))
+			{
+				const int32 RootTree = MedialForest.Find(RootIdx);
+				struct FLinkCandidate
+				{
+					double DistSq;
+					FIndex2i Edge;
+				};
+				TMap<int32, FLinkCandidate> BestLinks;
+				for (int32 Idx = 0; Idx < NumClusters; ++Idx)
+				{
+					const int32 IdxTree = MedialForest.Find(Idx);
+					if (IdxTree != RootTree)
+					{
+						// find the closest neighbor by brute force
+						double ClosestDistSq = FMathd::MaxReal;
+						if (FLinkCandidate* BestLink = BestLinks.Find(IdxTree))
+						{
+							ClosestDistSq = BestLink->DistSq;
+						}
+
+						int32 Closest = INDEX_NONE;
+						for (int32 NbrIdx = 0; NbrIdx < NumClusters; ++NbrIdx)
+						{
+							// Any disconnected tree is a candidate for a new edge (not just the root tree)
+							if (Idx != NbrIdx && MedialForest.Find(NbrIdx) != IdxTree)
+							{
+								double DSq = FVector3d::DistSquared(InMedialSkeleton.Spheres[Idx].Center, InMedialSkeleton.Spheres[NbrIdx].Center);
+								if (DSq < ClosestDistSq)
+								{
+									ClosestDistSq = DSq;
+									Closest = NbrIdx;
+								}
+							}
+						}
+						if (Closest != INDEX_NONE)
+						{
+							BestLinks.Add(IdxTree, FLinkCandidate{ ClosestDistSq, FIndex2i(Idx, Closest) });
+						}
+					}
+				}
+				for (const TPair<int32, FLinkCandidate>& Link : BestLinks)
+				{
+					MedialForest.Union(Link.Value.Edge.A, Link.Value.Edge.B);
+					KeptNbrs[Link.Value.Edge.A].Add(Link.Value.Edge.B);
+					KeptNbrs[Link.Value.Edge.B].Add(Link.Value.Edge.A);
+				}
+			}
+		}
+
+		
+		FVector3d RootCenter = InMedialSkeleton.Spheres[RootIdx].Center;
+
+		TArray<bool> Visited;
+		Visited.SetNumZeroed(NumClusters);
+
+		TArray<FIndex2i> MedialQueue;
+
+		if (MedialForest.GetSize(RootIdx) != NumClusters)
+		{
+			ensure(Options.MergeDisconnectedMethod == EMergeDisconnectedMethod::AddTopLevelRoot); // should only reach here w/ this merge method; otherwise, expect we have already merged to a single root
+
+			// add the virtual root at origin
+			OutParent.Add(INDEX_NONE);
+			OutTransform.Add(FTransform::Identity);
+			UseIndexToMedial->Add(INDEX_NONE); // virtual root doesn't correspond to a medial skeleton cluster
+
+			TSet<int32> SeenTrees;
+			for (int32 Idx = 0; Idx < NumClusters; ++Idx)
+			{
+				int32 TreeIdx = MedialForest.Find(Idx);
+				bool bWasInSet = false;
+				SeenTrees.FindOrAdd(TreeIdx, &bWasInSet);
+				if (!bWasInSet)
+				{
+					MedialQueue.Emplace(Idx, 0);
+				}
+			}
+		}
+
+		// If above multi-root handling hasn't already added to the queue, add the root
+		if (MedialQueue.IsEmpty())
+		{
+			MedialQueue.Emplace(RootIdx, INDEX_NONE);
+		}
+
+		while (!MedialQueue.IsEmpty())
+		{
+			FIndex2i SrcIdxOutParent = MedialQueue.Pop(EAllowShrinking::No);
+			// note: should not have already visited, because we enforced a tree structure on the graph by construction above
+			if (!ensure(!Visited[SrcIdxOutParent.A]))
+			{
+				continue;
+			}
+			Visited[SrcIdxOutParent.A] = true;
+			int32 OutIdx = UseIndexToMedial->Add(SrcIdxOutParent.A);
+			OutParent.Add(SrcIdxOutParent.B);
+			FVector3d ParentPos(0.);
+			if (SrcIdxOutParent.B != INDEX_NONE && (*UseIndexToMedial)[SrcIdxOutParent.B] != INDEX_NONE)
+			{
+				ParentPos = InMedialSkeleton.Spheres[(*UseIndexToMedial)[SrcIdxOutParent.B]].Center;
+			}
+			OutTransform.Add(FTransform(InMedialSkeleton.Spheres[SrcIdxOutParent.A].Center - ParentPos));
+			for (int32 Nbr : KeptNbrs[SrcIdxOutParent.A])
+			{
+				if (!Visited[Nbr])
+				{
+					MedialQueue.Emplace(Nbr, OutIdx);
+				}
+			}
+		}
+
+		if (OutMedialIndexToBoneIndex)
+		{
+			OutMedialIndexToBoneIndex->Init(INDEX_NONE, NumClusters);
+			for (int32 BoneIdx = 0; BoneIdx < UseIndexToMedial->Num(); ++BoneIdx)
+			{
+				int32 MedialIdx = (*UseIndexToMedial)[BoneIdx];
+				if (OutMedialIndexToBoneIndex->IsValidIndex(MedialIdx))
+				{
+					(*OutMedialIndexToBoneIndex)[MedialIdx] = BoneIdx;
+				}
+			}
+		}
+	}
+
+	void FMedialSkeletonToTreeSkeletonOptions::ToDynamicMeshSkeleton(const FMedialSkeleton& InMedialSkeleton, FDynamicMesh3& TargetMesh, FName WeightProfileName, FMedialSkeletonToTreeSkeletonOptions Options, int32 RootIndex)
+	{
+		TArray<int32> BoneRemap;
+		TArray<int32> Parents;
+		TArray<FTransform> Poses;
+
+		ToHierarchy(InMedialSkeleton, Parents, Poses, Options, RootIndex, nullptr, &BoneRemap);
+
+		if (Poses.IsEmpty())
+		{
+			// no skeleton to add
+			return;
+		}
+
+		// make sure attributes are enabled so we can add skeleton data
+		TargetMesh.EnableAttributes();
+
+		// transfer hierarchy to mesh bones
+		const int32 NumBones = Parents.Num();
+		FDynamicMeshAttributeSet* AttributeSet = TargetMesh.Attributes();
+		{
+			AttributeSet->EnableBones(NumBones);
+			FDynamicMeshBoneNameAttribute* BoneNameAttrib = AttributeSet->GetBoneNames();
+			FDynamicMeshBoneParentIndexAttribute* BoneParentIndices = AttributeSet->GetBoneParentIndices();
+			FDynamicMeshBonePoseAttribute* BonePoses = AttributeSet->GetBonePoses();
+			FDynamicMeshBoneColorAttribute* BoneColors = AttributeSet->GetBoneColors();
+
+			for (int32 BoneIdx = 0; BoneIdx < NumBones; ++BoneIdx)
+			{
+				FName BoneName("Bone", BoneIdx);
+				BoneNameAttrib->SetValue(BoneIdx, BoneName);
+				BoneParentIndices->SetValue(BoneIdx, Parents[BoneIdx]);
+				BonePoses->SetValue(BoneIdx, Poses[BoneIdx]);
+				BoneColors->SetValue(BoneIdx, FVector4f::One());
+			}
+		}
+
+		// Remove any existing skin weight attributes; they would generally not be compatible with the newly-created bones
+		TargetMesh.Attributes()->RemoveAllSkinWeightsAttributes();
+		// Add our new skin weight attribute
+		FDynamicMeshVertexSkinWeightsAttribute* Attribute = new FDynamicMeshVertexSkinWeightsAttribute(&TargetMesh);
+		TargetMesh.Attributes()->AttachSkinWeightsAttribute(WeightProfileName, Attribute);
+
+		// Transfer cluster assignments as rigid weights
+		using FBoneWeights = UE::AnimationCore::FBoneWeights;
+		using FBoneWeight = UE::AnimationCore::FBoneWeight;
+		ensureMsgf(TargetMesh.MaxVertexID() == InMedialSkeleton.VIDtoClusterIndex.Num(), TEXT("Medial skeleton does not match target mesh max ID: %d vs %d"), InMedialSkeleton.VIDtoClusterIndex.Num(), TargetMesh.MaxVertexID());
+		for (int32 VID = 0; VID < InMedialSkeleton.VIDtoClusterIndex.Num(); ++VID)
+		{
+			if (!TargetMesh.IsVertex(VID))
+			{
+				continue;
+			}
+			int32 SourceClusterIndex = InMedialSkeleton.VIDtoClusterIndex[VID];
+			if (!BoneRemap.IsValidIndex(SourceClusterIndex))
+			{
+				continue;
+			}
+			int32 BoneIndex = BoneRemap[SourceClusterIndex];
+			if (!ensureMsgf(BoneIndex >= 0 && BoneIndex <= TNumericLimits<FBoneIndexType>::Max(), TEXT("Invalid bone index %d"), BoneIndex))
+			{
+				continue;
+			}
+			FBoneWeights Wts;
+			Wts.SetBoneWeight(static_cast<FBoneIndexType>(BoneIndex), 1.f);
+			Attribute->SetValue(VID, Wts);
+		}
+	}
+
+
+
+
 
 	// Instantiate templated ComputeSkeleton methods for the supported mesh types
 	template FMedialSkeleton FSkeletonViaSampling::ComputeSkeleton<FDynamicMesh3>(const TMeshAABBTree3<FDynamicMesh3>& MeshBVH, TFunctionRef<FVector3d(int32)> GetVertexNormal, TFunctionRef<int32(int32)> VIDtoWeldVIDFn, const TFastWindingTree<FDynamicMesh3>* MeshFWTree);
diff --git a/Engine/Plugins/Runtime/GeometryProcessing/Source/GeometryAlgorithms/Public/Skeletonization/MeshMedialAxisSampling.h b/Engine/Plugins/Runtime/GeometryProcessing/Source/GeometryAlgorithms/Public/Skeletonization/MeshMedialAxisSampling.h
--- a/Engine/Plugins/Runtime/GeometryProcessing/Source/GeometryAlgorithms/Public/Skeletonization/MeshMedialAxisSampling.h
+++ b/Engine/Plugins/Runtime/GeometryProcessing/Source/GeometryAlgorithms/Public/Skeletonization/MeshMedialAxisSampling.h
@@ -321,6 +321,61 @@ namespace MedialAxis
 		}
 	};
 
+
+
+	struct FMedialSkeletonToTreeSkeletonOptions
+	{
+		// Method to select which edges of the medial skeleton are prioritized for inclusion in the hierarchical skeleton
+		enum class EEdgeWeightMethod
+		{
+			// Favor adding shorter edges first
+			EdgeLength,
+			// Favor edges between clusters that are earlier in the medial skeleton array. Because medial spheres are added incrementally in the highest error locations, their order can approximate importance.
+			ArrayOrder,
+			// Favor adding edges between larger medial spheres first
+			AvgRadius
+		};
+
+		// Method to merge disconnected components of the medial skeleton into a single hierarchy 
+		enum class EMergeDisconnectedMethod
+		{
+			// Connect closest bones across disconnected components until there is a single hierarchy
+			ConnectClosestBones,
+			// Add a top-level root node, and connect the selected root and all disconnected components to this node
+			AddTopLevelRoot
+		};
+
+		EEdgeWeightMethod EdgeWeightMethod = EEdgeWeightMethod::ArrayOrder;
+		EMergeDisconnectedMethod MergeDisconnectedMethod = EMergeDisconnectedMethod::ConnectClosestBones;
+
+		/**
+		 * Remap a medial skeleton to a hierarchical skeleton suitable for using in a URefSkeleton for animation
+		 * Note the indices of the hierarchical skeleton bones will be remapped from the medial skeleton, because URefSkeleton requires that parent bone indices be smaller than child bone indices
+		 * 
+		 * @param InMedialSkeleton Input medial skeleton
+		 * @param OutParent Output hierarchy, as parent indices
+		 * @param OutTransform Relative transforms for each bone. 1:1 with OutParent.
+		 * @param Options Additional options to control the skeleton conversion
+		 * @param RootIndex Specify the medial skeleton cluster to prefer as the root bone. If INDEX_NONE, a default root will be selected automatically.
+		 * @param OutBoneIndexToMedialIndex Optional mapping from hierarchical skeleton bone indices to source medial skeleton cluster indices
+		 * @param OutMedialIndexToBoneIndex Optional mapping from source medial skeleton cluster indices to hierarchical skeleton bone indices
+		 */
+		GEOMETRYALGORITHMS_API static void ToHierarchy(const FMedialSkeleton& InMedialSkeleton, TArray<int32>& OutParent, TArray<FTransform>& OutTransform, FMedialSkeletonToTreeSkeletonOptions Options, int32 RootIndex = INDEX_NONE, 
+			TArray<int32>* OutBoneIndexToMedialIndex = nullptr, TArray<int32>* OutMedialIndexToBoneIndex = nullptr);
+
+		/**
+		 * Transfer a medial skeleton to a dynamic mesh as a hierarchical skeleton, so it can be converted to a USkeletalMesh
+		 * 
+		 * @param InMedialSkeleton Input medial skeleton
+		 * @param TargetMesh Mesh to add skeleton a skin weights to
+		 * @param WeightProfile Name of skin weight profile to use. Note: Other profiles will be cleared, since they would likely not be compatible with the new skeleton.
+		 * @param Options Additional options to control the skeleton conversion
+		 * @param RootIndex Specify the medial skeleton cluster to prefer as the root bone. If INDEX_NONE, a default root will be selected automatically.
+		 */
+		GEOMETRYALGORITHMS_API static void ToDynamicMeshSkeleton(const FMedialSkeleton& InMedialSkeleton, FDynamicMesh3& TargetMesh, FName WeightProfile, FMedialSkeletonToTreeSkeletonOptions Options, int32 RootIndex = INDEX_NONE);
+	};
+
+
 	// Implementation/adaptation of "Dynamic Skeletonization via Variational Medial Skeletal Sampling" by Huang et al
 	// with some surface-focused additions -- optionally uses raycasts to guide refinement, and fast winding to guide medial-sphere projection
 	struct FSkeletonViaSampling
