[GeometryProcessing] Making mesh plannar symmetry more robust when handling meshes with overlapping verts on both sides of the mirror plane.

#Synced-CL 50896465

#Synced-CL 50896465

diff --git a/Engine/Plugins/Animation/SkeletalMeshModelingTools/Source/SkeletalMeshModelingTools/Private/Components/SKMBackedDynaMeshComponent.cpp b/Engine/Plugins/Animation/SkeletalMeshModelingTools/Source/SkeletalMeshModelingTools/Private/Components/SKMBackedDynaMeshComponent.cpp
--- a/Engine/Plugins/Animation/SkeletalMeshModelingTools/Source/SkeletalMeshModelingTools/Private/Components/SKMBackedDynaMeshComponent.cpp
+++ b/Engine/Plugins/Animation/SkeletalMeshModelingTools/Source/SkeletalMeshModelingTools/Private/Components/SKMBackedDynaMeshComponent.cpp
@@ -867,9 +867,11 @@ void USkeletalMeshBackedDynamicMeshComponent::UpdateSymmetry()
 		PreferSymmetryPlaneNormal.Add(FVector3d::UnitX());
 
 		FFrame3d SymmetryFrame(FVector3d::Zero(), FVector3d::UnitX());		
+
+		constexpr bool bForceComputeVertexPairs = true;	
 			
 		ReadOnlySymmetry = MakePimpl<FMeshPlanarSymmetry>();
-		ReadOnlySymmetry->Initialize(const_cast<FDynamicMesh3*>(&Mesh), Bounds, SymmetryFrame);
+		ReadOnlySymmetry->Initialize(const_cast<FDynamicMesh3*>(&Mesh), Bounds, SymmetryFrame, bForceComputeVertexPairs);
 	});
 }
 
diff --git a/Engine/Plugins/Runtime/GeometryProcessing/Source/DynamicMesh/Private/Parameterization/MeshPlanarSymmetry.cpp b/Engine/Plugins/Runtime/GeometryProcessing/Source/DynamicMesh/Private/Parameterization/MeshPlanarSymmetry.cpp
--- a/Engine/Plugins/Runtime/GeometryProcessing/Source/DynamicMesh/Private/Parameterization/MeshPlanarSymmetry.cpp
+++ b/Engine/Plugins/Runtime/GeometryProcessing/Source/DynamicMesh/Private/Parameterization/MeshPlanarSymmetry.cpp
@@ -489,57 +489,100 @@ bool FMeshPlanarSymmetry::AssignMatches(const FDynamicMesh3* Mesh, const TPointH
 	double OnPlaneTolerance = OnPlaneToleranceFactor * ErrorScale;
 	double MatchVertexTolerance = MatchVertexToleranceFactor * ErrorScale;
 
+	TArray<TArray<int32, TInlineAllocator<4>>> CandidatesPerPositiveVertex;
+	CandidatesPerPositiveVertex.SetNum(PositiveVerts.Num());
+	
 	// Compute unique matching vertices on mirror side (negative side of plane) for positive-side vertices
 	ParallelFor(PositiveVerts.Num(), [&](int32 k)
 		{
 			int32 VertexID = PositiveVerts[k];
 
-
 			FVector3d Position = Mesh->GetVertex(VertexID);
 			FVector3d MirrorPosition = GetMirroredPosition(Position);
 
 			// find vertex that is closest to mirrored position
 			FVector3d Feature = InvariantFeatures[VertexID];
 			double MatchFeaturesToleranceSqr = MatchFeaturesTolerance * MatchFeaturesTolerance;
-			TPair<int32, double> NearestInfo = VertexHash.FindNearestInRadius(MirrorPosition, MatchVertexTolerance, [MirrorPosition, Mesh](const int32& OtherVID)
-				{
-					return FVector3d::DistSquared(MirrorPosition, Mesh->GetVertex(OtherVID));
-				}
-				, [Feature, VertexID, this, &InvariantFeatures, MatchFeaturesToleranceSqr](const int32& OtherVID) -> bool // Ignore filter fn returns true for vertices we want to skip
-				{
-					return OtherVID == VertexID // don't match to self
-						|| Vertices[OtherVID].bIsSourceVertex // don't match to other positive vertices
-						|| FVector3d::DistSquared(InvariantFeatures[OtherVID], Feature) > MatchFeaturesToleranceSqr; // don't match if the features don't match
-				}
-				);
-			int32 NearestVID = NearestInfo.Key;
-			if (NearestVID != INDEX_NONE)
+
+			TArray<int32> Candidates;
+			VertexHash.FindPointsInBall(MirrorPosition, MatchVertexTolerance, [MirrorPosition, Mesh](const int32& OtherVID)
+			{
+				return FVector3d::DistSquared(MirrorPosition, Mesh->GetVertex(OtherVID));
+			}, 
+			Candidates,
+			[Feature, VertexID, this, &InvariantFeatures, MatchFeaturesToleranceSqr](const int32& OtherVID) -> bool // Ignore filter fn returns true for vertices we want to skip
 			{
-				Vertices[NearestVID].PairedVertex = VertexID;
-				Vertices[VertexID].PairedVertex = NearestVID;
+				return OtherVID == VertexID // don't match to self
+					|| Vertices[OtherVID].bIsSourceVertex // don't match to other positive vertices
+					|| FVector3d::DistSquared(InvariantFeatures[OtherVID], Feature) > MatchFeaturesToleranceSqr; // don't match if the features don't match
 			}
-			// If a failed match would break symmetry overall, we can retry without feature-based filtering
-			else if (bReattemptMatchesOnFailure && Vertices[VertexID].PlaneSignedDistance >= OnPlaneTolerance)
+			);
+
+			if (Candidates.IsEmpty())
 			{
-				TPair<int32, double> FeaturelessNearestInfo = VertexHash.FindNearestInRadius(MirrorPosition, MatchVertexTolerance, [MirrorPosition, Mesh](const int32& OtherVID)
+				// If a failed match would break symmetry overall, we can retry without feature-based filtering
+				if (bReattemptMatchesOnFailure && Vertices[VertexID].PlaneSignedDistance >= OnPlaneTolerance)
+				{
+					VertexHash.FindPointsInBall(MirrorPosition, MatchVertexTolerance, [MirrorPosition, Mesh](const int32& OtherVID)
 					{
 						return FVector3d::DistSquared(MirrorPosition, Mesh->GetVertex(OtherVID));
-					}
-					, [Feature, VertexID, this, &InvariantFeatures, MatchFeaturesToleranceSqr](const int32& OtherVID) -> bool // Ignore filter fn returns true for vertices we want to skip
+					}, 
+					Candidates,
+					[Feature, VertexID, this, &InvariantFeatures, MatchFeaturesToleranceSqr](const int32& OtherVID) -> bool // Ignore filter fn returns true for vertices we want to skip
 					{
 						return OtherVID == VertexID // don't match to self
 							|| Vertices[OtherVID].bIsSourceVertex; // don't match to other positive vertices
 					}
-					);
-				NearestVID = FeaturelessNearestInfo.Key;
-				if (NearestVID != INDEX_NONE)
-				{
-					Vertices[NearestVID].PairedVertex = VertexID;
-					Vertices[VertexID].PairedVertex = NearestVID;
+					);	
 				}
 			}
+
+			// Rank all candidates first by Euclidean distance and then by feature distance
+			Algo::Sort(Candidates,
+				[&MirrorPosition, Mesh, &InvariantFeatures, Feature](const int32& A, const int32& B)
+				{
+					double DistanceA = FVector3d::DistSquared(MirrorPosition, Mesh->GetVertex(A));
+					double FeatureDistanceA = FVector3d::DistSquared(InvariantFeatures[A], Feature);
+				
+					double DistanceB = FVector3d::DistSquared(MirrorPosition, Mesh->GetVertex(B));
+					double FeatureDistanceB = FVector3d::DistSquared(InvariantFeatures[B], Feature);
+				
+					if (DistanceA < DistanceB)
+					{
+						return true;
+					}
+					
+					if (DistanceA == DistanceB)
+					{
+						if (FeatureDistanceA < FeatureDistanceB)
+						{
+							return true;
+						}
+					}
+
+					return false;
+				});
+			
+			CandidatesPerPositiveVertex[k] = MoveTemp(Candidates); 
 		});
+	
+	for (int32 k = 0; k < PositiveVerts.Num(); k++)
+	{
+		int32 VertexID = PositiveVerts[k];
+		
+		for (int32 Candidate : CandidatesPerPositiveVertex[k])
+		{
+			// Avoid many to one matching
+			if (Vertices[Candidate].PairedVertex != INDEX_NONE)
+			{
+				continue;
+			}
 
+			Vertices[Candidate].PairedVertex = VertexID;
+			Vertices[VertexID].PairedVertex = Candidate;
+			break;
+		}
+	}
 
 	// If a vertex has no match and is within the symmetry-plane tolerance band, consider it a plane vertex.
 	// Otherwise it is a failed match
@@ -570,6 +613,12 @@ bool FMeshPlanarSymmetry::Initialize(FDynamicMesh3* Mesh, FDynamicMeshAABBTree3*
 }
 
 bool FMeshPlanarSymmetry::Initialize(FDynamicMesh3* Mesh, const FAxisAlignedBox3d& Bounds, FFrame3d SymmetryFrameIn)
+{
+	constexpr bool bForceComputeVertexPairs = false;
+	return Initialize(Mesh, Bounds, SymmetryFrameIn, bForceComputeVertexPairs);
+}
+
+bool FMeshPlanarSymmetry::Initialize(FDynamicMesh3* Mesh, const FAxisAlignedBox3d& Bounds, FFrame3d SymmetryFrameIn, bool bForceComputeVertexPairs)
 {
 	this->TargetMesh = Mesh;
 	if (!Validate(Mesh))
@@ -584,11 +633,14 @@ bool FMeshPlanarSymmetry::Initialize(FDynamicMesh3* Mesh, const FAxisAlignedBox3
 	TArray<FVector3d> Features;
 	ComputeMeshInfo(Mesh, Bounds, Features, MeshCentroid);
 
-	// Quick reject test: MeshCentroid must be on desired symmetry plane
-	double CentroidPlaneSignedDist = (MeshCentroid - SymmetryFrameIn.Origin).Dot(SymmetryFrameIn.Z());
-	if (FMath::Abs(CentroidPlaneSignedDist) > OnPlaneToleranceFactor * ErrorScale)
+	if (!bForceComputeVertexPairs)
 	{
-		return false;
+		// Quick reject test: MeshCentroid must be on desired symmetry plane
+		double CentroidPlaneSignedDist = (MeshCentroid - SymmetryFrameIn.Origin).Dot(SymmetryFrameIn.Z());
+		if (FMath::Abs(CentroidPlaneSignedDist) > OnPlaneToleranceFactor * ErrorScale)
+		{
+			return false;
+		}
 	}
 
 	// Hash the vertices.
@@ -605,7 +657,6 @@ bool FMeshPlanarSymmetry::Initialize(FDynamicMesh3* Mesh, const FAxisAlignedBox3
 }
 
 
-
 FFrame3d FMeshPlanarSymmetry::GetPositiveSideFrame(FFrame3d FromFrame) const
 {
 	FVector3d DeltaVec = (FromFrame.Origin - SymmetryFrame.Origin);
diff --git a/Engine/Plugins/Runtime/GeometryProcessing/Source/DynamicMesh/Public/Parameterization/MeshPlanarSymmetry.h b/Engine/Plugins/Runtime/GeometryProcessing/Source/DynamicMesh/Public/Parameterization/MeshPlanarSymmetry.h
--- a/Engine/Plugins/Runtime/GeometryProcessing/Source/DynamicMesh/Public/Parameterization/MeshPlanarSymmetry.h
+++ b/Engine/Plugins/Runtime/GeometryProcessing/Source/DynamicMesh/Public/Parameterization/MeshPlanarSymmetry.h
@@ -44,6 +44,14 @@ public:
 	 */
 	UE_API bool Initialize(FDynamicMesh3* Mesh, const FAxisAlignedBox3d& Bounds, FFrame3d SymmetryFrameIn);
 
+	/**
+	 * Given a Mesh and its bounding box, and a Symmetry Plane/Frame, detect any pairs of vertices with
+	 * planar/mirror-symmetry relationships, as well as "on plane" vertices
+	 * @param bForceComputeVertexPairs Computes vertex pairs even if the mesh doesn't pass quick full symmetry test
+	 * @return false if any non-on-plane vertex fails to find a match.
+	 */
+	UE_API bool Initialize(FDynamicMesh3* Mesh, const FAxisAlignedBox3d& Bounds, FFrame3d SymmetryFrameIn, bool bForceComputeVertexPairs);
+	
 	/**
 	 * Given a Mesh and its bounding box, find a Symmetry Plane/Frame and detect any pairs of vertices with
 	 * planar/mirror-symmetry relationships, as well as "on plane" vertices
