Refactor non-engine part of skin weight binding algorithm to GeometryProcessing plugin, so it can be used in more places

#Synced-CL 51082240

diff --git a/Engine/Plugins/Runtime/GeometryProcessing/Source/DynamicMesh/Private/Operations/SkinWeightBinding.cpp b/Engine/Plugins/Runtime/GeometryProcessing/Source/DynamicMesh/Private/Operations/SkinWeightBinding.cpp
new file mode 100644
--- /dev/null
+++ b/Engine/Plugins/Runtime/GeometryProcessing/Source/DynamicMesh/Private/Operations/SkinWeightBinding.cpp
@@ -0,0 +1,656 @@
+// Copyright Epic Games, Inc. All Rights Reserved.
+
+#include "Operations/SkinWeightBinding.h"
+
+#include "BoneWeights.h"
+#include "DynamicMesh/DynamicMesh3.h"
+#include "DynamicMesh/DynamicMeshAttributeSet.h"
+#include "DynamicMesh/DynamicVertexSkinWeightsAttribute.h"
+#include "DynamicMesh/DynamicBoneAttribute.h"
+#include "Spatial/DenseGrid3.h"
+#include "Spatial/MeshWindingNumberGrid.h"
+#include "Spatial/OccupancyGrid3.h"
+
+
+namespace UE::Geometry::SkinWeightBindingLocals
+{
+
+// A simple FIFO queue. Maintains a set of blocks, rather than allocate for each element.
+template<typename T, int32 BlockSize=256>
+class TFIFOQueue
+{
+public:
+	TFIFOQueue() = default;
+
+	void Push(T InElem)
+	{
+		// Container Empty?
+		if (Blocks.IsEmpty())
+		{
+			Blocks.SetNum(1);
+			Blocks[0].bIsFree = false;
+			PushIndex = -1;
+		}
+		// Or are we at the end of current block?
+		else if (PushIndex == (BlockSize - 1))
+		{
+			PushBlock = AllocateNewPushBlock();
+			PushIndex = -1;
+		}
+
+		PushIndex++;
+
+		Blocks[PushBlock].Data[PushIndex] = InElem;
+	}
+
+	T Pop()
+	{
+		check(!IsEmpty());
+
+		if (PopIndex == (BlockSize - 1))
+		{
+			// Reached the end. Free this block and move onto the next one.
+			PopBlock = FreeCurrentPopBlock();
+			PopIndex = -1;
+		}
+
+		PopIndex++;
+		return Blocks[PopBlock].Data[PopIndex];
+	}
+
+	bool TryPop(T &OutValue)
+	{
+		if (IsEmpty())
+		{
+			return false;
+		}
+
+		OutValue = Pop();
+		return true;
+	}
+
+	void Reset()
+	{
+		Blocks.Reset();
+		PushBlock = 0;
+		PushIndex = -1;
+		PopBlock = 0;
+		PopIndex = -1;
+		FreeCount = 0;
+	}
+	
+
+	bool IsEmpty() const
+	{
+		return PushBlock == PopBlock && PushIndex == PopIndex;
+	}
+
+private:
+	int32 AllocateNewPushBlock()
+	{
+		checkSlow(PushBlock != -1);
+
+		int NewBlockIndex = INDEX_NONE;
+		
+		if (FreeCount > 0)
+		{
+			// Find a free block
+			for (int32 Index = 0; Index < Blocks.Num(); Index++)
+			{
+				if (Blocks[Index].bIsFree)
+				{
+					NewBlockIndex = Index;
+					FreeCount--;
+					break;
+				}
+			}
+			checkfSlow(NewBlockIndex != INDEX_NONE, TEXT("We should have found a free block."));
+		}
+		else
+		{
+			NewBlockIndex = Blocks.Num();
+			Blocks.AddDefaulted();
+		}
+
+		Blocks[PushBlock].NextBlock = NewBlockIndex;
+		Blocks[NewBlockIndex].bIsFree = false;
+		Blocks[NewBlockIndex].NextBlock = INDEX_NONE;
+		return NewBlockIndex;
+	}
+
+	int32 FreeCurrentPopBlock()
+	{
+		const int32 NextBlock = Blocks[PopBlock].NextBlock;
+		checkSlow(NextBlock != -1);
+		Blocks[PopBlock].bIsFree = true;
+		FreeCount++;
+
+		return NextBlock;
+	}
+	
+
+	int32 PushBlock = 0;
+	int32 PushIndex = -1;
+
+	int32 PopBlock = 0;
+	int32 PopIndex = -1;
+
+	int32 FreeCount = 0;
+	
+	struct FBlock
+	{
+		FBlock()
+		{
+			Data = new T[BlockSize];
+		}
+		~FBlock()
+		{
+			delete [] Data;
+		}
+		FBlock(FBlock &&InOther) noexcept
+		{
+			Data = InOther.Data;
+			InOther.Data = nullptr;
+			NextBlock = InOther.NextBlock;
+			bIsFree = InOther.bIsFree;
+		}
+		// disallow copy/copy-assignment
+		FBlock(const FBlock&) = delete;
+		FBlock& operator=(const FBlock&) = delete;
+		
+		bool bIsFree = true;
+		int32 NextBlock = INDEX_NONE;
+		T *Data;
+	};
+	
+	TArray<FBlock> Blocks;
+};
+
+
+static float DistanceToLineSegment(const FVector& P, const FVector& A, const FVector& B)
+{
+	const FVector M = B - A;
+	const FVector T = P - A;
+
+	const float C1 = (float)FVector::DotProduct(M, T);
+	if (C1 <= 0.0f)
+	{
+		return (float)FVector::Dist(P, A);
+	}
+
+	const float C2 = (float)FVector::DotProduct(M, M);
+	if (C2 <= C1)
+	{
+		return (float)FVector::Dist(P, B);
+	}
+
+	// Project the point onto the line and get the distance between them.
+	const FVector PT = A + M * (C1 / C2);
+	return (float)FVector::Dist(P, PT);
+}
+
+
+// List of bones as used by the binding class. In this case for each bone transform, we want to
+// store a list of line segments going from the bone transform to all the child bone transforms.
+struct FTransformHierarchyQuery
+{
+	explicit FTransformHierarchyQuery(TConstArrayView<SkinBinding::FBonePoseInfo> InTransformHierarchy)
+	{
+		for (int Index = 0; Index < InTransformHierarchy.Num(); Index++)
+		{
+			FTransform Xform = InTransformHierarchy[Index].LocalTransform;
+			int32 ParentIndex = InTransformHierarchy[Index].ParentIndex;
+
+			while (ParentIndex != INDEX_NONE)
+			{
+				Xform = Xform * InTransformHierarchy[ParentIndex].LocalTransform;
+				ParentIndex = InTransformHierarchy[ParentIndex].ParentIndex;
+			}
+
+			BoneFans.Add({ Xform.GetLocation() });
+		}
+
+		// Fill in the fan tips, as needed.
+		for (int Index = 0; Index < InTransformHierarchy.Num(); Index++)
+		{
+			const int32 ParentIndex = InTransformHierarchy[Index].ParentIndex;
+			if (ParentIndex != INDEX_NONE)
+			{
+				BoneFans[ParentIndex].TipsPos.Add(BoneFans[Index].RootPos);
+			}
+		}		
+	}
+
+	bool IsEndEffector(const int32 InBoneIndex) const
+	{
+		return BoneFans[InBoneIndex].TipsPos.IsEmpty();
+	}
+
+	float GetDistanceToBoneFan(const int32 InBoneIndex, const FVector& InPoint) const
+	{
+		return BoneFans[InBoneIndex].GetDistance(InPoint);
+	}
+	
+	FBox GetBoneFanBBox(const int32 InBoneIndex) const
+	{
+		return BoneFans[InBoneIndex].GetBBox();
+	}
+
+	bool GetBoneFanIntersectsBox(const int32 InBoneIndex, const FBox &InBox) const
+	{
+		return BoneFans[InBoneIndex].IntersectsBox(InBox);
+	}
+
+private:
+	struct FBoneFan
+	{
+		FVector RootPos;
+		TArray<FVector> TipsPos;
+
+		float GetDistance(const FVector& InPoint) const
+		{
+			if (TipsPos.IsEmpty())
+			{
+				return (float)FVector::Distance(RootPos, InPoint);
+			}
+			else
+			{
+				float Distance = std::numeric_limits<float>::max();
+				for (const FVector& TipPos: TipsPos)
+				{
+					Distance = FMath::Min(Distance, DistanceToLineSegment(InPoint, RootPos, TipPos));
+				}
+				return Distance;
+			}
+		}
+
+		FBox GetBBox() const
+		{
+			FBox Box(RootPos, RootPos);
+			for (const FVector& TipPos: TipsPos)
+			{
+				Box += TipPos;
+			}
+
+			return Box;
+		}
+
+		bool IntersectsBox(const FBox &InBox) const
+		{
+			if (TipsPos.IsEmpty())
+			{
+				return FMath::PointBoxIntersection(RootPos, InBox);
+			}
+
+			if (GetBBox().Intersect(InBox))
+			{
+				for (const FVector& TipPos: TipsPos)
+				{
+					if (FMath::LineBoxIntersection(InBox, RootPos, TipPos, TipPos - RootPos))
+					{
+						return true;
+					}
+				}
+			}
+			return false;
+		}				
+	};
+
+	 TArray<FBoneFan> BoneFans;
+};
+
+
+struct FCreateSkinWeights_Closest_WorkData final :
+	TThreadSingleton<FCreateSkinWeights_Closest_WorkData>
+{
+	TArray<TPair<FBoneIndexType, float>> RawBoneWeights;
+	TArray<AnimationCore::FBoneWeight> BoneWeights;
+
+	void NormalizeWeightsAndLimitCount(const float InTotalWeight, const int32 InMaxWeights)
+	{
+		// Normalize
+		for (TPair<FBoneIndexType, float> &BoneWeight: RawBoneWeights)
+		{
+			BoneWeight.Value /= InTotalWeight;
+		}
+		
+		RawBoneWeights.Sort([](const TPair<FBoneIndexType, float> &A, const TPair<FBoneIndexType, float> &B)
+		{
+			return A.Value > B.Value;
+		});
+
+		BoneWeights.Reset(InMaxWeights);
+		for (int32 BoneIndex = 0; BoneIndex < FMath::Min(InMaxWeights, RawBoneWeights.Num()); BoneIndex++)
+		{
+			const TPair<FBoneIndexType, float>& BoneWeight = RawBoneWeights[BoneIndex];
+			BoneWeights.Add(UE::AnimationCore::FBoneWeight(BoneWeight.Key, BoneWeight.Value));
+		}
+			
+	}
+};
+
+float ComputeWeightStiffness(const float InWeight, const float InStiffness)
+{
+	return (1.0f - InStiffness) * InWeight + InStiffness * InWeight * InWeight;
+}
+
+void CopyBonesToMesh(FDynamicMesh3& Mesh, TConstArrayView<SkinBinding::FBonePoseInfo> TransformHierarchy)
+{
+	// Initialize bone attributes
+	FDynamicMeshAttributeSet* AttribSet = Mesh.Attributes();
+
+	const int32 NumBones = TransformHierarchy.Num();
+	AttribSet->EnableBones(NumBones);
+
+	FDynamicMeshBoneNameAttribute* BoneNames = AttribSet->GetBoneNames();
+	FDynamicMeshBoneParentIndexAttribute* BoneParentIndices = AttribSet->GetBoneParentIndices();
+	FDynamicMeshBonePoseAttribute* BonePoses = AttribSet->GetBonePoses();
+
+	for (int BoneIdx = 0; BoneIdx < NumBones; ++BoneIdx)
+	{
+		BoneNames->SetValue(BoneIdx, TransformHierarchy[BoneIdx].Name);
+		BoneParentIndices->SetValue(BoneIdx, TransformHierarchy[BoneIdx].ParentIndex);
+		BonePoses->SetValue(BoneIdx, TransformHierarchy[BoneIdx].LocalTransform);
+	}
+}
+
+FDynamicMeshVertexSkinWeightsAttribute* GetOrCreateSkinWeightsAttribute(
+	FDynamicMesh3& InMesh,
+	FName InProfileName
+	)
+{
+	FDynamicMeshVertexSkinWeightsAttribute *Attribute = InMesh.Attributes()->GetSkinWeightsAttribute(InProfileName);
+	if (!Attribute)
+	{
+		Attribute = new FDynamicMeshVertexSkinWeightsAttribute(&InMesh);
+		InMesh.Attributes()->AttachSkinWeightsAttribute(InProfileName, Attribute);
+	}
+	
+	return Attribute;
+}
+
+
+void CreateSkinWeights_DirectDistance(
+	FDynamicMesh3& InMesh, TConstArrayView<SkinBinding::FBonePoseInfo> TransformHierarchy, FName ProfileName,
+	float InStiffness,
+	const AnimationCore::FBoneWeightsSettings& InSettings
+	)
+{
+	using namespace AnimationCore;
+
+	FDynamicMeshVertexSkinWeightsAttribute *SkinWeights = GetOrCreateSkinWeightsAttribute(InMesh, ProfileName);
+		
+	const int32 NumVertices = InMesh.VertexCount();
+	const int32 MaxInfluences = InSettings.GetMaxWeightCount();
+
+	// Use the diagonal size of the bbox to make the bone distance falloff scale invariant.
+	const float DiagBounds = (float)FMath::Max(InMesh.GetBounds(true).DiagonalLength(), UE_SMALL_NUMBER);
+	
+	const FTransformHierarchyQuery Skeleton(TransformHierarchy);
+
+	ParallelFor(NumVertices, [&](const int32 VertexIdx)
+	{
+		const FVector3d& Pos = InMesh.GetVertex(VertexIdx);
+
+		FCreateSkinWeights_Closest_WorkData &WorkData = FCreateSkinWeights_Closest_WorkData::Get();
+		
+		if (MaxInfluences > 1)
+		{
+			WorkData.RawBoneWeights.Reset(TransformHierarchy.Num());
+		}
+		else
+		{
+			WorkData.RawBoneWeights.Init(MakeTuple(static_cast<FBoneIndexType>(0), 1.0f), 1);
+		}
+		
+		float TotalWeight = 0.0f;
+		for (int32 BoneIndex = 0; BoneIndex < TransformHierarchy.Num(); BoneIndex++)
+		{
+			// For single influences, avoid end effectors to avoid closest-distance fighting with the tips of their
+			// parent bones.
+			if (MaxInfluences == 1 && Skeleton.IsEndEffector(BoneIndex))
+			{
+				continue;
+			}
+			
+			// Normalize the distance by the diagonal size of the bbox to maintain scale invariance.
+			float Weight = Skeleton.GetDistanceToBoneFan(BoneIndex, Pos) / DiagBounds;
+
+			// Avoid div-by-zero but allow for the possibility that multiple bones may
+			// touch this vertex.
+			Weight = FMath::Max(Weight, KINDA_SMALL_NUMBER);
+
+			// Compute the actual weight, factoring in the stiffness value. W = (1/S(D))^2
+			// Where S(x) is the stiffness function.
+			Weight = FMath::Square(1.0f / ComputeWeightStiffness(Weight, InStiffness));
+
+			if (MaxInfluences > 1)
+			{
+				TotalWeight += Weight;
+				WorkData.RawBoneWeights.Add(MakeTuple(static_cast<FBoneIndexType>(BoneIndex), Weight));
+			}
+			else if (Weight > TotalWeight)
+			{
+				// For single influences, we only care about the strongest influence.
+				TotalWeight = Weight;
+				WorkData.RawBoneWeights[0] = MakeTuple(static_cast<FBoneIndexType>(BoneIndex), Weight);
+			}
+		}
+
+		WorkData.NormalizeWeightsAndLimitCount(TotalWeight, InSettings.GetMaxWeightCount());
+
+		SkinWeights->SetValue(VertexIdx, FBoneWeights::Create(WorkData.BoneWeights, InSettings));
+	});
+}
+
+
+void CreateSkinWeights_GeodesicVoxel(
+	FDynamicMesh3& InMesh, TConstArrayView<SkinBinding::FBonePoseInfo> TransformHierarchy, FName ProfileName,
+	float InStiffness, int32 InVoxelResolution,
+	const AnimationCore::FBoneWeightsSettings& InSettings
+	)
+{
+	using namespace AnimationCore;
+
+	FDynamicMeshVertexSkinWeightsAttribute *SkinWeights = GetOrCreateSkinWeightsAttribute(InMesh, ProfileName);
+		
+	const int32 NumVertices = InMesh.VertexCount();
+	const int32 MaxInfluences = InSettings.GetMaxWeightCount();
+
+	// Use the diagonal size of the bbox to make the bone distance falloff scale invariant.
+	const float DiagonalBounds = (float)FMath::Max(InMesh.GetBounds(true).DiagonalLength(), UE_SMALL_NUMBER);
+		
+	const FTransformHierarchyQuery Skeleton(TransformHierarchy);
+
+	// This is grossly inefficient but tricky to do otherwise, since each bone distance
+	// computation is done per-thread. We could possibly solve this by chunking instead
+	// and accumulating partial results.
+	TArray<float> Weights;
+	Weights.SetNumUninitialized(NumVertices * TransformHierarchy.Num());
+
+	const FOccupancyGrid3 Occupancy(InMesh, InVoxelResolution);
+
+	const FVector3i Dimensions = Occupancy.GetOccupancyStateGrid().GetDimensions();
+		
+	ParallelFor(TransformHierarchy.Num(), [&](int32 BoneIndex) {
+		// For single influences, avoid end effectors to avoid closest-distance fighting with the tips of their
+		// parent bones.
+		if (MaxInfluences == 1 && Skeleton.IsEndEffector(BoneIndex))
+		{
+			return;
+		}
+		
+		TFIFOQueue<FVector3i> WorkingSet;  
+		FDenseGrid3f BoneDistance(Dimensions.X, Dimensions.Y, Dimensions.Z, DiagonalBounds);
+
+		// Mark all the cells that the bone intersects with distance of 0 and put them
+		// on the work queue.
+		const FBox BoneBox = Skeleton.GetBoneFanBBox(BoneIndex);
+		const FVector3i BoneMin = Occupancy.GetCellIndexFromPoint(BoneBox.Min);
+		const FVector3i BoneMax = Occupancy.GetCellIndexFromPoint(BoneBox.Max);
+
+		for (int32 I = BoneMin.X; I <= BoneMax.X; I++)
+		{
+			for (int32 J = BoneMin.Y; J <= BoneMax.Y; J++)
+			{
+				for (int32 K = BoneMin.Z; K <= BoneMax.Z; K++)
+				{
+					const FVector3i Candidate(I, J, K);
+					if (BoneDistance.IsValidIndex(Candidate))
+					{
+						const FBox3d CellBox{Occupancy.GetCellBoxFromIndex(Candidate)};
+						if (Skeleton.GetBoneFanIntersectsBox(BoneIndex, CellBox))
+						{
+							WorkingSet.Push(Candidate);
+							BoneDistance[Candidate] = 0.0f;
+						}
+					}
+				}
+			}
+		}
+
+		// Iterate over all the voxels until we have constructed shortest distance paths
+		// throughout the level set.
+		while (!WorkingSet.IsEmpty())
+		{
+			const FVector3i WorkItem = WorkingSet.Pop();
+
+			// Loop through each of the neighbours (6 face neighbours, 12 edge neighbors,
+			// and 8 corner neighbours) and see if any of them are closer to the bone
+			// than their current marked distance.
+			float CurrentDistance = BoneDistance[WorkItem];
+			for (int32 N = 0; N < 26; N++)
+			{
+				FVector3i Offset(IndexUtil::GridOffsets26[N]);
+				FVector3i Candidate(WorkItem + Offset);
+
+				if (!BoneDistance.IsValidIndex(Candidate))
+				{
+					continue;
+				}
+					
+				// Ensure this entry is either a part of the interior or boundary domain.
+				if (Occupancy.GetOccupancyStateGrid()[Candidate] == FOccupancyGrid3::EDomain::Exterior)
+				{
+					continue;
+				}
+
+				const float CellDistance = (FVector3f(Offset) * Occupancy.GetCellSize()).Length();
+				const float CandidateDistance = CurrentDistance + CellDistance;
+				const float OldDistance = BoneDistance[Candidate]; 
+
+				if (OldDistance > CandidateDistance)
+				{
+					WorkingSet.Push(Candidate);
+					BoneDistance[Candidate] = CandidateDistance;
+				}
+			}
+		}
+
+		// Loop through all the vertices, find the voxel each belongs to, and compute
+		// the distance from the voxel to the vertex (assuming the distance stored in the
+		// voxel is based on traversing from voxel center to voxel center).
+		for (int32 VertexIdx = 0; VertexIdx < NumVertices; VertexIdx++)
+		{
+			const FVector3d& Pos = InMesh.GetVertex(VertexIdx);
+			const FVector3i CellIndex = Occupancy.GetCellIndexFromPoint(Pos);
+			const FVector3d CellCenter{Occupancy.GetCellCenterFromIndex(CellIndex)};
+
+			if (ensure(BoneDistance.IsValidIndex(CellIndex)))
+			{
+				float Distance = BoneDistance[CellIndex];
+				Distance += (float)FVector3d::Distance(CellCenter, Pos);
+				
+				// Normalize the distance by the diagonal size of the bbox to maintain scale invariance.
+				float Weight = Distance / DiagonalBounds;
+
+				// Avoid div-by-zero but allow for the possibility that multiple bones may
+				// touch this vertex.
+				Weight = FMath::Max(Weight, UE_KINDA_SMALL_NUMBER);
+
+				// Compute the actual weight, factoring in the stiffness value. W = (1/S(D))^2
+				// Where S(x) is the stiffness function.
+				Weight = FMath::Square(1.0f / ComputeWeightStiffness(Weight, InStiffness));
+
+				Weights[VertexIdx * TransformHierarchy.Num() + BoneIndex] = Weight;
+			}
+		}
+	});	
+
+	ParallelFor(NumVertices, [&](const int32 VertexIdx)
+	{
+		FCreateSkinWeights_Closest_WorkData &WorkData = FCreateSkinWeights_Closest_WorkData::Get();
+			
+		if (MaxInfluences > 1)
+		{
+			WorkData.RawBoneWeights.Reset(TransformHierarchy.Num());
+		}
+		else
+		{
+			WorkData.RawBoneWeights.Init(MakeTuple(static_cast<FBoneIndexType>(0), 1.0f), 1);
+		}
+
+		float TotalWeight = 0.0f;
+		for (int32 BoneIndex = 0; BoneIndex < TransformHierarchy.Num(); BoneIndex++)
+		{
+			// Ignore end effectors for single influence computation, since those we not
+			// calculated.
+			if (MaxInfluences == 1 && Skeleton.IsEndEffector(BoneIndex))
+			{
+				continue;
+			}
+			
+			const float Weight = Weights[VertexIdx * TransformHierarchy.Num() + BoneIndex];
+			
+			if (MaxInfluences > 1)
+			{
+				TotalWeight += Weight;
+				WorkData.RawBoneWeights.Add(MakeTuple(static_cast<FBoneIndexType>(BoneIndex), Weight));
+			}
+			else if (Weight > TotalWeight)
+			{
+				TotalWeight = Weight;
+				WorkData.RawBoneWeights[0] = MakeTuple(static_cast<FBoneIndexType>(BoneIndex), Weight);
+			}
+		}
+
+		WorkData.NormalizeWeightsAndLimitCount(TotalWeight, InSettings.GetMaxWeightCount());
+
+		SkinWeights->SetValue(VertexIdx, FBoneWeights::Create(WorkData.BoneWeights, InSettings));
+	});
+}
+
+}
+
+namespace UE::Geometry::SkinBinding
+{
+
+void CreateSkinWeights(FDynamicMesh3& InMesh, TConstArrayView<SkinBinding::FBonePoseInfo> TransformHierarchy, FName ProfileName, const FBindSettings& BindSettings)
+{
+	using namespace UE::Geometry::SkinWeightBindingLocals;
+
+	const float ClampedStiffness = FMath::Clamp(BindSettings.Stiffness, 0.0f, 1.0f);
+
+	UE::AnimationCore::FBoneWeightsSettings BoneSettings;
+	BoneSettings.SetMaxWeightCount(BindSettings.MaxInfluences);
+
+	switch(BindSettings.BindType)
+	{
+	default:
+		ensureMsgf(false, TEXT("Unknown BindType in CreateSkinWeights"));
+		// fall through to DirectDistance
+	case ESkinBindingType::DirectDistance:
+		CreateSkinWeights_DirectDistance(InMesh, TransformHierarchy, ProfileName, ClampedStiffness, BoneSettings);
+		break;
+			
+	case ESkinBindingType::GeodesicVoxel:
+		CreateSkinWeights_GeodesicVoxel(InMesh, TransformHierarchy, ProfileName, ClampedStiffness, BindSettings.VoxelResolution, BoneSettings);
+		break;
+	}
+
+	CopyBonesToMesh(InMesh, TransformHierarchy);
+}
+
+
+
+} // namespace UE::Geometry::SkinBinding
diff --git a/Engine/Plugins/Runtime/GeometryProcessing/Source/DynamicMesh/Public/Operations/SkinWeightBinding.h b/Engine/Plugins/Runtime/GeometryProcessing/Source/DynamicMesh/Public/Operations/SkinWeightBinding.h
new file mode 100644
--- /dev/null
+++ b/Engine/Plugins/Runtime/GeometryProcessing/Source/DynamicMesh/Public/Operations/SkinWeightBinding.h
@@ -0,0 +1,60 @@
+// Copyright Epic Games, Inc. All Rights Reserved.
+
+#pragma once
+
+#include "GeometryTypes.h"
+#include "UObject/NameTypes.h"
+#include "TransformTypes.h"
+
+#define UE_API DYNAMICMESH_API
+
+
+namespace UE::Geometry
+{
+
+template<typename ParentType> class TDynamicVertexSkinWeightsAttribute;
+class FDynamicMesh3;
+using FDynamicMeshVertexSkinWeightsAttribute = TDynamicVertexSkinWeightsAttribute<FDynamicMesh3>;
+
+
+enum class ESkinBindingType : uint8
+{
+	// Computes the binding strength by computing the Euclidean distance to the closest set of bones,
+	// where the strength of binding is proportional to the inverse distance. May cause bones to affects
+	// parts of geometry that, although close in space, may be topologically distant.
+	DirectDistance = 0,
+
+	// Computes the binding by computing the geodesic distance from each set of bones. This is slower than the
+	// direct distance.
+	GeodesicVoxel = 1,
+};
+
+
+namespace SkinBinding
+{
+
+	struct FBonePoseInfo
+	{
+		FTransform LocalTransform;
+		FName Name;
+		int32 ParentIndex;
+	};
+
+	struct FBindSettings
+	{
+		float Stiffness = 0.2f;
+		int32 MaxInfluences = 5;
+		int32 VoxelResolution = 256;
+		ESkinBindingType BindType = ESkinBindingType::DirectDistance;
+	};
+
+	// Create skin weights from the input bone hierarchy, and set them on the dynamic mesh
+	UE_API void CreateSkinWeights(FDynamicMesh3& InMesh, TConstArrayView<FBonePoseInfo> TransformHierarchy, FName ProfileName, const FBindSettings& Settings);
+
+
+};
+
+
+} // namespace UE::Geometry
+
+#undef UE_API
diff --git a/Engine/Plugins/Runtime/MeshModelingToolset/Source/ModelingOperators/Private/SkinningOps/SkinBindingOp.cpp b/Engine/Plugins/Runtime/MeshModelingToolset/Source/ModelingOperators/Private/SkinningOps/SkinBindingOp.cpp
--- a/Engine/Plugins/Runtime/MeshModelingToolset/Source/ModelingOperators/Private/SkinningOps/SkinBindingOp.cpp
+++ b/Engine/Plugins/Runtime/MeshModelingToolset/Source/ModelingOperators/Private/SkinningOps/SkinBindingOp.cpp
@@ -2,339 +2,13 @@
 
 #include "SkinningOps/SkinBindingOp.h"
 
-#include "BoneWeights.h"
 #include "DynamicMesh/DynamicMesh3.h"
-#include "DynamicMesh/DynamicMeshAttributeSet.h"
-#include "DynamicMesh/DynamicVertexSkinWeightsAttribute.h"
-#include "DynamicMesh/DynamicBoneAttribute.h"
 #include "ReferenceSkeleton.h"
-#include "Spatial/DenseGrid3.h"
-#include "Spatial/MeshWindingNumberGrid.h"
-#include "Spatial/OccupancyGrid3.h"
 
 
 namespace UE::Geometry
 {
 
-// A simple FIFO queue. Maintains a set of blocks, rather than allocate for each element.
-template<typename T, int32 BlockSize=256>
-class TFIFOQueue
-{
-public:
-	TFIFOQueue() = default;
-
-	void Push(T InElem)
-	{
-		// Container Empty?
-		if (Blocks.IsEmpty())
-		{
-			Blocks.SetNum(1);
-			Blocks[0].bIsFree = false;
-			PushIndex = -1;
-		}
-		// Or are we at the end of current block?
-		else if (PushIndex == (BlockSize - 1))
-		{
-			PushBlock = AllocateNewPushBlock();
-			PushIndex = -1;
-		}
-
-		PushIndex++;
-
-		Blocks[PushBlock].Data[PushIndex] = InElem;
-	}
-
-	T Pop()
-	{
-		check(!IsEmpty());
-
-		if (PopIndex == (BlockSize - 1))
-		{
-			// Reached the end. Free this block and move onto the next one.
-			PopBlock = FreeCurrentPopBlock();
-			PopIndex = -1;
-		}
-
-		PopIndex++;
-		return Blocks[PopBlock].Data[PopIndex];
-	}
-
-	bool TryPop(T &OutValue)
-	{
-		if (IsEmpty())
-		{
-			return false;
-		}
-
-		OutValue = Pop();
-		return true;
-	}
-
-	void Reset()
-	{
-		Blocks.Reset();
-		PushBlock = 0;
-		PushIndex = -1;
-		PopBlock = 0;
-		PopIndex = -1;
-	}
-	
-
-	bool IsEmpty() const
-	{
-		return PushBlock == PopBlock && PushIndex == PopIndex;
-	}
-
-private:
-	int32 AllocateNewPushBlock()
-	{
-		checkSlow(PushBlock != -1);
-
-		int NewBlockIndex = INDEX_NONE;
-		
-		if (FreeCount > 0)
-		{
-			// Find a free block
-			for (int32 Index = 0; Index < Blocks.Num(); Index++)
-			{
-				if (Blocks[Index].bIsFree)
-				{
-					NewBlockIndex = Index;
-					FreeCount--;
-					break;
-				}
-			}
-			checkfSlow(NewBlockIndex != INDEX_NONE, TEXT("We should have found a free block."));
-		}
-		else
-		{
-			NewBlockIndex = Blocks.Num();
-			Blocks.AddDefaulted();
-		}
-
-		Blocks[PushBlock].NextBlock = NewBlockIndex;
-		Blocks[NewBlockIndex].bIsFree = false;
-		Blocks[NewBlockIndex].NextBlock = INDEX_NONE;
-		return NewBlockIndex;
-	}
-
-	int32 FreeCurrentPopBlock()
-	{
-		const int32 NextBlock = Blocks[PopBlock].NextBlock;
-		checkSlow(NextBlock != -1);
-		Blocks[PopBlock].bIsFree = true;
-		FreeCount++;
-
-		return NextBlock;
-	}
-	
-
-	int32 PushBlock = 0;
-	int32 PushIndex = -1;
-
-	int32 PopBlock = 0;
-	int32 PopIndex = -1;
-
-	int32 FreeCount = 0;
-	
-	struct FBlock
-	{
-		FBlock()
-		{
-			Data = new T[BlockSize];
-		}
-		~FBlock()
-		{
-			delete [] Data;
-		}
-
-		FBlock(FBlock &&InOther) noexcept
-		{
-			Data = InOther.Data;
-			InOther.Data = nullptr;
-		}
-		
-		bool bIsFree = true;
-		int32 NextBlock = INDEX_NONE;
-		T *Data;
-	};
-	
-	TArray<FBlock> Blocks;
-};
-
-
-static float DistanceToLineSegment(const FVector& P, const FVector& A, const FVector& B)
-{
-	const FVector M = B - A;
-	const FVector T = P - A;
-
-	const float C1 = FVector::DotProduct(M, T);
-	if (C1 <= 0.0f)
-	{
-		return FVector::Dist(P, A);
-	}
-
-	const float C2 = FVector::DotProduct(M, M);
-	if (C2 <= C1)
-	{
-		return FVector::Dist(P, B);
-	}
-
-	// Project the point onto the line and get the distance between them.
-	const FVector PT = A + M * (C1 / C2);
-	return FVector::Dist(P, PT);
-}
-
-
-// List of bones as used by the binding class. In this case for each bone transform, we want to
-// store a list of line segments going from the bone transform to all the child bone transforms.
-struct FTransformHierarchyQuery
-{
-	explicit FTransformHierarchyQuery(const TArray<TPair<FTransform, FMeshBoneInfo>>& InTransformHierarchy)
-	{
-		for (int Index = 0; Index < InTransformHierarchy.Num(); Index++)
-		{
-			FTransform Xform = InTransformHierarchy[Index].Key;
-			int32 ParentIndex = InTransformHierarchy[Index].Value.ParentIndex;
-
-			while (ParentIndex != INDEX_NONE)
-			{
-				Xform = Xform * InTransformHierarchy[ParentIndex].Key;
-				ParentIndex = InTransformHierarchy[ParentIndex].Value.ParentIndex;
-			}
-
-			BoneFans.Add({ Xform.GetLocation() });
-		}
-
-		// Fill in the fan tips, as needed.
-		for (int Index = 0; Index < InTransformHierarchy.Num(); Index++)
-		{
-			const int32 ParentIndex = InTransformHierarchy[Index].Value.ParentIndex;
-			if (ParentIndex != INDEX_NONE)
-			{
-				BoneFans[ParentIndex].TipsPos.Add(BoneFans[Index].RootPos);
-			}
-		}		
-	}
-
-	bool IsEndEffector(const int32 InBoneIndex) const
-	{
-		return BoneFans[InBoneIndex].TipsPos.IsEmpty();
-	}
-
-	float GetDistanceToBoneFan(const int32 InBoneIndex, const FVector& InPoint) const
-	{
-		return BoneFans[InBoneIndex].GetDistance(InPoint);
-	}
-	
-	FBox GetBoneFanBBox(const int32 InBoneIndex) const
-	{
-		return BoneFans[InBoneIndex].GetBBox();
-	}
-
-	bool GetBoneFanIntersectsBox(const int32 InBoneIndex, const FBox &InBox) const
-	{
-		return BoneFans[InBoneIndex].IntersectsBox(InBox);
-	}
-
-private:
-	struct FBoneFan
-	{
-		FVector RootPos;
-		TArray<FVector> TipsPos;
-
-		float GetDistance(const FVector& InPoint) const
-		{
-			if (TipsPos.IsEmpty())
-			{
-				return FVector::Distance(RootPos, InPoint);
-			}
-			else
-			{
-				float Distance = std::numeric_limits<float>::max();
-				for (const FVector& TipPos: TipsPos)
-				{
-					Distance = FMath::Min(Distance, DistanceToLineSegment(InPoint, RootPos, TipPos));
-				}
-				return Distance;
-			}
-		}
-
-		FBox GetBBox() const
-		{
-			FBox Box(RootPos, RootPos);
-			for (const FVector& TipPos: TipsPos)
-			{
-				Box += TipPos;
-			}
-
-			return Box;
-		}
-
-		bool IntersectsBox(const FBox &InBox) const
-		{
-			if (TipsPos.IsEmpty())
-			{
-				return FMath::PointBoxIntersection(RootPos, InBox);
-			}
-
-			if (GetBBox().Intersect(InBox))
-			{
-				for (const FVector& TipPos: TipsPos)
-				{
-					if (FMath::LineBoxIntersection(InBox, RootPos, TipPos, TipPos - RootPos))
-					{
-						return true;
-					}
-				}
-			}
-			return false;
-		}				
-	};
-
-	 TArray<FBoneFan> BoneFans;
-};
-
-
-namespace
-{
-	struct FCreateSkinWeights_Closest_WorkData final :
-		TThreadSingleton<FCreateSkinWeights_Closest_WorkData>
-	{
-		TArray<TPair<FBoneIndexType, float>> RawBoneWeights;
-		TArray<AnimationCore::FBoneWeight> BoneWeights;
-
-		void NormalizeWeightsAndLimitCount(const float InTotalWeight, const int32 InMaxWeights)
-		{
-			// Normalize
-			for (TPair<FBoneIndexType, float> &BoneWeight: RawBoneWeights)
-			{
-				BoneWeight.Value /= InTotalWeight;
-			}
-		
-			RawBoneWeights.Sort([](const TPair<FBoneIndexType, float> &A, const TPair<FBoneIndexType, float> &B)
-			{
-				return A.Value > B.Value;
-			});
-
-			BoneWeights.Reset(InMaxWeights);
-			for (int32 BoneIndex = 0; BoneIndex < FMath::Min(InMaxWeights, RawBoneWeights.Num()); BoneIndex++)
-			{
-				const TPair<FBoneIndexType, float>& BoneWeight = RawBoneWeights[BoneIndex];
-				BoneWeights.Add(UE::AnimationCore::FBoneWeight(BoneWeight.Key, BoneWeight.Value));
-			}
-			
-		}
-	};
-
-	float ComputeWeightStiffness(const float InWeight, const float InStiffness)
-	{
-		return (1.0f - InStiffness) * InWeight + InStiffness * InWeight * InWeight;
-	}
-
-}
-
-
 void FSkinBindingOp::SetTransformHierarchyFromReferenceSkeleton(
 	const FReferenceSkeleton& InRefSkeleton
 	)
@@ -343,11 +17,12 @@ void FSkinBindingOp::SetTransformHierarchyFromReferenceSkeleton(
 	const TArray<FMeshBoneInfo>& BoneInfo = InRefSkeleton.GetRawRefBoneInfo();
 	const TArray<FTransform>& BonePose = InRefSkeleton.GetRawRefBonePose();
 		
-	TransformHierarchy.Reset(BoneInfo.Num());
+	BoneTransformHierarchy.Reset(BoneInfo.Num());
 		
 	for (int32 Index = 0; Index < BoneInfo.Num(); Index++)
 	{
-		TransformHierarchy.Add(MakeTuple(BonePose[Index], BoneInfo[Index]));
+		SkinBinding::FBonePoseInfo Info{ .LocalTransform = BonePose[Index], .Name = BoneInfo[Index].Name, .ParentIndex = BoneInfo[Index].ParentIndex };
+		BoneTransformHierarchy.Add(Info);
 	}
 }
 
@@ -366,303 +41,15 @@ void FSkinBindingOp::CalculateResult(FProgressCancel* InProgress)
 		return;
 	}
 
-	const float ClampedStiffness = FMath::Clamp(Stiffness, 0.0f, 1.0f);	
-
-	UE::AnimationCore::FBoneWeightsSettings Settings;
-	Settings.SetMaxWeightCount(MaxInfluences);
-
-	switch(BindType)
-	{
-	case ESkinBindingType::DirectDistance:
-		CreateSkinWeights_DirectDistance(*ResultMesh, ClampedStiffness, Settings);
-		break;
-			
-	case ESkinBindingType::GeodesicVoxel:
-		CreateSkinWeights_GeodesicVoxel(*ResultMesh, ClampedStiffness, Settings);
-		break;
-	}
-
-	// Initialize bone attributes
-	FDynamicMeshAttributeSet* AttribSet = ResultMesh->Attributes();
-
-	const int32 NumBones = TransformHierarchy.Num();
-	AttribSet->EnableBones(NumBones);
-
-	FDynamicMeshBoneNameAttribute* BoneNames = AttribSet->GetBoneNames();
-	FDynamicMeshBoneParentIndexAttribute* BoneParentIndices = AttribSet->GetBoneParentIndices();
-	FDynamicMeshBonePoseAttribute* BonePoses = AttribSet->GetBonePoses();
-
-	for (int BoneIdx = 0; BoneIdx < NumBones; ++BoneIdx)
-	{
-		BoneNames->SetValue(BoneIdx, TransformHierarchy[BoneIdx].Value.Name);
-		BoneParentIndices->SetValue(BoneIdx, TransformHierarchy[BoneIdx].Value.ParentIndex);
-		BonePoses->SetValue(BoneIdx, TransformHierarchy[BoneIdx].Key);
-	}
-}
-
-
-FDynamicMeshVertexSkinWeightsAttribute* FSkinBindingOp::GetOrCreateSkinWeightsAttribute(
-	FDynamicMesh3& InMesh,
-	FName InProfileName
-	)
-{
-	FDynamicMeshVertexSkinWeightsAttribute *Attribute = InMesh.Attributes()->GetSkinWeightsAttribute(InProfileName);
-	if (!Attribute)
-	{
-		Attribute = new FDynamicMeshVertexSkinWeightsAttribute(&InMesh);
-		InMesh.Attributes()->AttachSkinWeightsAttribute(InProfileName, Attribute);
-	}
-	
-	return Attribute;
-}
-
-
-void FSkinBindingOp::CreateSkinWeights_DirectDistance(
-	FDynamicMesh3& InMesh,
-	float InStiffness,
-	const AnimationCore::FBoneWeightsSettings& InSettings
-	) const
-{
-	using namespace AnimationCore;
-
-	FDynamicMeshVertexSkinWeightsAttribute *SkinWeights = GetOrCreateSkinWeightsAttribute(InMesh, ProfileName);
-		
-	const int32 NumVertices = InMesh.VertexCount();
-
-	// Use the diagonal size of the bbox to make the bone distance falloff scale invariant.
-	const float DiagBounds = InMesh.GetBounds(true).DiagonalLength();
-	
-	const FTransformHierarchyQuery Skeleton(TransformHierarchy);
-
-	ParallelFor(NumVertices, [&](const int32 VertexIdx)
+	SkinBinding::FBindSettings Settings
 	{
-		const FVector3d& Pos = InMesh.GetVertex(VertexIdx);
-
-		FCreateSkinWeights_Closest_WorkData &WorkData = FCreateSkinWeights_Closest_WorkData::Get();
-		
-		if (MaxInfluences > 1)
-		{
-			WorkData.RawBoneWeights.Reset(TransformHierarchy.Num());
-		}
-		else
-		{
-			WorkData.RawBoneWeights.Init(MakeTuple(0, 1.0), 1);
-		}
-		
-		float TotalWeight = 0.0f;
-		for (int32 BoneIndex = 0; BoneIndex < TransformHierarchy.Num(); BoneIndex++)
-		{
-			// For single influences, avoid end effectors to avoid closest-distance fighting with the tips of their
-			// parent bones.
-			if (MaxInfluences == 1 && Skeleton.IsEndEffector(BoneIndex))
-			{
-				continue;
-			}
-			
-			// Normalize the distance by the diagonal size of the bbox to maintain scale invariance.
-			float Weight = Skeleton.GetDistanceToBoneFan(BoneIndex, Pos) / DiagBounds;
-
-			// Avoid div-by-zero but allow for the possibility that multiple bones may
-			// touch this vertex.
-			Weight = FMath::Max(Weight, KINDA_SMALL_NUMBER);
-
-			// Compute the actual weight, factoring in the stiffness value. W = (1/S(D))^2
-			// Where S(x) is the stiffness function.
-			Weight = FMath::Square(1.0f / ComputeWeightStiffness(Weight, InStiffness));
-
-			if (MaxInfluences > 1)
-			{
-				TotalWeight += Weight;
-				WorkData.RawBoneWeights.Add(MakeTuple(static_cast<FBoneIndexType>(BoneIndex), Weight));
-			}
-			else if (Weight > TotalWeight)
-			{
-				// For single influences, we only care about the strongest influence.
-				TotalWeight = Weight;
-				WorkData.RawBoneWeights[0] = MakeTuple(static_cast<FBoneIndexType>(BoneIndex), Weight);
-			}
-		}
-
-		WorkData.NormalizeWeightsAndLimitCount(TotalWeight, InSettings.GetMaxWeightCount());
-
-		SkinWeights->SetValue(VertexIdx, FBoneWeights::Create(WorkData.BoneWeights, InSettings));
-	});
+		.Stiffness = Stiffness,
+		.MaxInfluences = MaxInfluences,
+		.VoxelResolution = VoxelResolution,
+		.BindType = BindType,
+	};
+	SkinBinding::CreateSkinWeights(*ResultMesh, BoneTransformHierarchy, ProfileName, Settings);
 }
 
 
-void FSkinBindingOp::CreateSkinWeights_GeodesicVoxel(
-	FDynamicMesh3& InMesh,
-	float InStiffness,
-	const AnimationCore::FBoneWeightsSettings& InSettings
-	) const
-{
-	using namespace AnimationCore;
-
-	FDynamicMeshVertexSkinWeightsAttribute *SkinWeights = GetOrCreateSkinWeightsAttribute(InMesh, ProfileName);
-		
-	const int32 NumVertices = InMesh.VertexCount();
-
-	// Use the diagonal size of the bbox to make the bone distance falloff scale invariant.
-	const float DiagonalBounds = InMesh.GetBounds(true).DiagonalLength();
-		
-	const FTransformHierarchyQuery Skeleton(TransformHierarchy);
-
-	// This is grossly inefficient but tricky to do otherwise, since each bone distance
-	// computation is done per-thread. We could possibly solve this by chunking instead
-	// and accumulating partial results.
-	TArray<float> Weights;
-	Weights.SetNumUninitialized(NumVertices * TransformHierarchy.Num());
-
-	const FOccupancyGrid3 Occupancy(InMesh, VoxelResolution);
-
-	const FVector3i Dimensions = Occupancy.GetOccupancyStateGrid().GetDimensions();
-		
-	ParallelFor(TransformHierarchy.Num(), [&](int32 BoneIndex) {
-		// For single influences, avoid end effectors to avoid closest-distance fighting with the tips of their
-		// parent bones.
-		if (MaxInfluences == 1 && Skeleton.IsEndEffector(BoneIndex))
-		{
-			return;
-		}
-		
-		TFIFOQueue<FVector3i> WorkingSet;  
-		FDenseGrid3f BoneDistance(Dimensions.X, Dimensions.Y, Dimensions.Z, DiagonalBounds);
-
-		// Mark all the cells that the bone intersects with distance of 0 and put them
-		// on the work queue.
-		const FBox BoneBox = Skeleton.GetBoneFanBBox(BoneIndex);
-		const FVector3i BoneMin = Occupancy.GetCellIndexFromPoint(BoneBox.Min);
-		const FVector3i BoneMax = Occupancy.GetCellIndexFromPoint(BoneBox.Max);
-
-		for (int32 I = BoneMin.X; I <= BoneMax.X; I++)
-		{
-			for (int32 J = BoneMin.Y; J <= BoneMax.Y; J++)
-			{
-				for (int32 K = BoneMin.Z; K <= BoneMax.Z; K++)
-				{
-					const FVector3i Candidate(I, J, K);
-					if (BoneDistance.IsValidIndex(Candidate))
-					{
-						const FBox3d CellBox{Occupancy.GetCellBoxFromIndex(Candidate)};
-						if (Skeleton.GetBoneFanIntersectsBox(BoneIndex, CellBox))
-						{
-							WorkingSet.Push(Candidate);
-							BoneDistance[Candidate] = 0.0f;
-						}
-					}
-				}
-			}
-		}
-
-		// Iterate over all the voxels until we have constructed shortest distance paths
-		// throughout the level set.
-		while (!WorkingSet.IsEmpty())
-		{
-			const FVector3i WorkItem = WorkingSet.Pop();
-
-			// Loop through each of the neighbours (6 face neighbours, 12 edge neighbors,
-			// and 8 corner neighbours) and see if any of them are closer to the bone
-			// than their current marked distance.
-			float CurrentDistance = BoneDistance[WorkItem];
-			for (int32 N = 0; N < 26; N++)
-			{
-				FVector3i Offset(IndexUtil::GridOffsets26[N]);
-				FVector3i Candidate(WorkItem + Offset);
-
-				if (!BoneDistance.IsValidIndex(Candidate))
-				{
-					continue;
-				}
-					
-				// Ensure this entry is either a part of the interior or boundary domain.
-				if (Occupancy.GetOccupancyStateGrid()[Candidate] == FOccupancyGrid3::EDomain::Exterior)
-				{
-					continue;
-				}
-
-				const float CellDistance = (FVector3f(Offset) * Occupancy.GetCellSize()).Length();
-				const float CandidateDistance = CurrentDistance + CellDistance;
-				const float OldDistance = BoneDistance[Candidate]; 
-
-				if (OldDistance > CandidateDistance)
-				{
-					WorkingSet.Push(Candidate);
-					BoneDistance[Candidate] = CandidateDistance;
-				}
-			}
-		}
-
-		// Loop through all the vertices, find the voxel each belongs to, and compute
-		// the distance from the voxel to the vertex (assuming the distance stored in the
-		// voxel is based on traversing from voxel center to voxel center).
-		for (int32 VertexIdx = 0; VertexIdx < NumVertices; VertexIdx++)
-		{
-			const FVector3d& Pos = InMesh.GetVertex(VertexIdx);
-			const FVector3i CellIndex = Occupancy.GetCellIndexFromPoint(Pos);
-			const FVector3d CellCenter{Occupancy.GetCellCenterFromIndex(CellIndex)};
-
-			if (ensure(BoneDistance.IsValidIndex(CellIndex)))
-			{
-				float Distance = BoneDistance[CellIndex];
-				Distance += FVector3d::Distance(CellCenter, Pos);
-				
-				// Normalize the distance by the diagonal size of the bbox to maintain scale invariance.
-				float Weight = Distance / DiagonalBounds;
-
-				// Avoid div-by-zero but allow for the possibility that multiple bones may
-				// touch this vertex.
-				Weight = FMath::Max(Weight, UE_KINDA_SMALL_NUMBER);
-
-				// Compute the actual weight, factoring in the stiffness value. W = (1/S(D))^2
-				// Where S(x) is the stiffness function.
-				Weight = FMath::Square(1.0f / ComputeWeightStiffness(Weight, InStiffness));
-
-				Weights[VertexIdx * TransformHierarchy.Num() + BoneIndex] = Weight;
-			}
-		}
-	});	
-
-	ParallelFor(NumVertices, [&](const int32 VertexIdx)
-	{
-		FCreateSkinWeights_Closest_WorkData &WorkData = FCreateSkinWeights_Closest_WorkData::Get();
-			
-		if (MaxInfluences > 1)
-		{
-			WorkData.RawBoneWeights.Reset(TransformHierarchy.Num());
-		}
-		else
-		{
-			WorkData.RawBoneWeights.Init(MakeTuple(0, 1.0), 1);
-		}
-
-		float TotalWeight = 0.0f;
-		for (int32 BoneIndex = 0; BoneIndex < TransformHierarchy.Num(); BoneIndex++)
-		{
-			// Ignore end effectors for single influence computation, since those we not
-			// calculated.
-			if (MaxInfluences == 1 && Skeleton.IsEndEffector(BoneIndex))
-			{
-				continue;
-			}
-			
-			const float Weight = Weights[VertexIdx * TransformHierarchy.Num() + BoneIndex];
-			
-			if (MaxInfluences > 1)
-			{
-				TotalWeight += Weight;
-				WorkData.RawBoneWeights.Add(MakeTuple(static_cast<FBoneIndexType>(BoneIndex), Weight));
-			}
-			else if (Weight > TotalWeight)
-			{
-				TotalWeight = Weight;
-				WorkData.RawBoneWeights[0] = MakeTuple(static_cast<FBoneIndexType>(BoneIndex), Weight);
-			}
-		}
-
-		WorkData.NormalizeWeightsAndLimitCount(TotalWeight, InSettings.GetMaxWeightCount());
-
-		SkinWeights->SetValue(VertexIdx, FBoneWeights::Create(WorkData.BoneWeights, InSettings));
-	});
-}
-
 } // namespace UE::Geometry
diff --git a/Engine/Plugins/Runtime/MeshModelingToolset/Source/ModelingOperators/Public/SkinningOps/SkinBindingOp.h b/Engine/Plugins/Runtime/MeshModelingToolset/Source/ModelingOperators/Public/SkinningOps/SkinBindingOp.h
--- a/Engine/Plugins/Runtime/MeshModelingToolset/Source/ModelingOperators/Public/SkinningOps/SkinBindingOp.h
+++ b/Engine/Plugins/Runtime/MeshModelingToolset/Source/ModelingOperators/Public/SkinningOps/SkinBindingOp.h
@@ -5,6 +5,7 @@
 #include "ModelingOperators.h"
 #include "ReferenceSkeleton.h"
 #include "SkeletalMeshAttributes.h"
+#include "Operations/SkinWeightBinding.h"
 
 #define UE_API MODELINGOPERATORS_API
 
@@ -19,27 +20,19 @@ class FDynamicMesh3;
 using FDynamicMeshVertexSkinWeightsAttribute = TDynamicVertexSkinWeightsAttribute<FDynamicMesh3>;
 
 
-enum class ESkinBindingType : uint8
-{
-	// Computes the binding strength by computing the Euclidean distance to the closest set of bones,
-	// where the strength of binding is proportional to the inverse distance. May cause bones to affects
-	// parts of geometry that, although close in space, may be topologically distant.
-	DirectDistance = 0,
-
-	// Computes the binding by computing the geodesic distance from each set of bones. This is slower than the
-	// direct distance.
-	GeodesicVoxel = 1,
-};
-
-
 class FSkinBindingOp : public FDynamicMeshOperator
 {
 public:
+
+	// note we disable deprecation warnings here due to the deprecated TransformHierarchy member below
+PRAGMA_DISABLE_DEPRECATION_WARNINGS
 	virtual ~FSkinBindingOp() override = default;
+PRAGMA_ENABLE_DEPRECATION_WARNINGS
 	
-	// The transform hierarchy to bind to. Listed in the same order as the bones in the
-	// reference skeleton that this skeletal mesh is tied to.
+
 	TSharedPtr<FDynamicMesh3, ESPMode::ThreadSafe> OriginalMesh;
+
+	UE_DEPRECATED(5.8, "TransformHierarchy is now unused and should not be accessed.")
 	TArray<TPair<FTransform, FMeshBoneInfo>> TransformHierarchy;
 
 	FName ProfileName = FSkeletalMeshAttributes::DefaultSkinWeightProfileName;
@@ -54,24 +47,7 @@ public:
 	UE_API virtual void CalculateResult(FProgressCancel* InProgress) override;
 
 private:
-	static UE_API FDynamicMeshVertexSkinWeightsAttribute* GetOrCreateSkinWeightsAttribute(
-		FDynamicMesh3& InMesh,
-		FName InProfileName
-		);
-
-
-	UE_API void CreateSkinWeights_DirectDistance(
-		FDynamicMesh3& InMesh,
-		float InStiffness,
-		const AnimationCore::FBoneWeightsSettings& InSettings
-		) const;
-
-
-	UE_API void CreateSkinWeights_GeodesicVoxel(
-		FDynamicMesh3& InMesh,
-		float InStiffness,
-		const AnimationCore::FBoneWeightsSettings& InSettings
-		) const; 
+	TArray<SkinBinding::FBonePoseInfo> BoneTransformHierarchy;
 };
 
 
diff --git a/Engine/Source/Runtime/GeometryCore/Public/Spatial/OccupancyGrid3.h b/Engine/Source/Runtime/GeometryCore/Public/Spatial/OccupancyGrid3.h
--- a/Engine/Source/Runtime/GeometryCore/Public/Spatial/OccupancyGrid3.h
+++ b/Engine/Source/Runtime/GeometryCore/Public/Spatial/OccupancyGrid3.h
@@ -60,10 +60,15 @@ struct FOccupancyGrid3
 			FVector3i(1, 1, 1),
 		};
 
-		ParallelFor(Occupancy.Num(), [&](const int32 OccupancyId)
+		if (!ensure(Occupancy.Num() <= (int64)MAX_int32))
+		{
+			// TODO: support occupancy grids of larger sizes
+			return;
+		}
+		ParallelFor(static_cast<int32>(Occupancy.Num()), [&](const int32 OccupancyId)
 		{
 			const FVector3i OccupancyIndex(Occupancy.GetCoords(OccupancyId));
-			const int32 Count = Algo::CountIf(CornerOffsets, [&WindingGrid, OccupancyIndex](FVector3i CornerOffset)
+			const SIZE_T Count = Algo::CountIf(CornerOffsets, [&WindingGrid, OccupancyIndex](FVector3i CornerOffset)
 			{
 				const FVector3i CornerIndex(OccupancyIndex + CornerOffset);
 				return WindingGrid.GetValue(CornerIndex) >= WindingGrid.WindingIsoValue;
