handle unset elements in the color overlay when transferring vertex color attributes

#Synced-CL 51328354

diff --git a/Engine/Plugins/Runtime/GeometryProcessing/Source/DynamicMesh/Private/Operations/TransferDynamicMeshAttributes.cpp b/Engine/Plugins/Runtime/GeometryProcessing/Source/DynamicMesh/Private/Operations/TransferDynamicMeshAttributes.cpp
--- a/Engine/Plugins/Runtime/GeometryProcessing/Source/DynamicMesh/Private/Operations/TransferDynamicMeshAttributes.cpp
+++ b/Engine/Plugins/Runtime/GeometryProcessing/Source/DynamicMesh/Private/Operations/TransferDynamicMeshAttributes.cpp
@@ -551,6 +551,14 @@ bool FTransferVertexColorAttribute::TransferColorToPoint(FVector4f& OutColor, co
 
 	const FDynamicMeshColorOverlay* SourceColors = SourceMesh->Attributes()->PrimaryColors();
 	const FIndex3i ColorTriElements = SourceColors->GetTriangle(TriID);
+	if (ColorTriElements.A == IndexConstants::InvalidID)
+	{
+		return false;
+	}
+	if (!ensure(ColorTriElements.B != IndexConstants::InvalidID && ColorTriElements.C != IndexConstants::InvalidID))
+	{
+		return false;
+	}
 	if (SearchRadius < 0 && NormalThreshold < 0)
 	{
 		// If the radius and normals are ignored, simply interpolate the values and return the result
