Make position error / normal error wt tradeoff be a lerp rather than just scaling position error weight, so it's easier to think about separately from the error threshold (doesn't increase the scale on measured error overall)

#Synced-CL 51328354

diff --git a/Engine/Plugins/Experimental/GeometryCollectionPlugin/Source/GeometryCollectionNodes/Public/Dataflow/GeometryCollectionSamplingNodes.h b/Engine/Plugins/Experimental/GeometryCollectionPlugin/Source/GeometryCollectionNodes/Public/Dataflow/GeometryCollectionSamplingNodes.h
--- a/Engine/Plugins/Experimental/GeometryCollectionPlugin/Source/GeometryCollectionNodes/Public/Dataflow/GeometryCollectionSamplingNodes.h
+++ b/Engine/Plugins/Experimental/GeometryCollectionPlugin/Source/GeometryCollectionNodes/Public/Dataflow/GeometryCollectionSamplingNodes.h
@@ -385,8 +385,8 @@ private:
 	UPROPERTY(EditAnywhere, Category = Options, AdvancedDisplay, meta = (DataflowInput, ClampMin = 1, UIMax = 100))
 	float ErrorScaleNearEdgeSurfaceIntersection = 10.f;
 
-	// Weight for the position error term. Relative to plane error term, which has weight 1.
-	UPROPERTY(EditAnywhere, Category = Options, AdvancedDisplay, meta = (DataflowInput, UIMin = .1, ClampMin = .001, UIMax = 2))
+	// Weight for the position error term. Relative to plane error term, so normal and position error weights sum to 1. Should be in the (0,1] range.
+	UPROPERTY(EditAnywhere, Category = Options, AdvancedDisplay, meta = (DisplayName = "Position Error Weight", DataflowInput, UIMin = .1, ClampMin = .001, ClampMax = 1))
 	float PosErrorWt = .2f;
 
 	virtual void Evaluate(UE::Dataflow::FContext& Context, const FDataflowOutput* Out) const override;
diff --git a/Engine/Plugins/Runtime/GeometryProcessing/Source/GeometryAlgorithms/Private/Skeletonization/MeshMedialAxisSampling.cpp b/Engine/Plugins/Runtime/GeometryProcessing/Source/GeometryAlgorithms/Private/Skeletonization/MeshMedialAxisSampling.cpp
--- a/Engine/Plugins/Runtime/GeometryProcessing/Source/GeometryAlgorithms/Private/Skeletonization/MeshMedialAxisSampling.cpp
+++ b/Engine/Plugins/Runtime/GeometryProcessing/Source/GeometryAlgorithms/Private/Skeletonization/MeshMedialAxisSampling.cpp
@@ -268,21 +268,22 @@ namespace UE::Geometry::MedialAxis
 			return Result;
 		}
 
-
+		PosErrorWt = FMath::Clamp(PosErrorWt, 0., 1.);
+		double NormalErrorWt = 1. - PosErrorWt;
 
 		// Area-Weighted Average of face-Normal errors, using FSQEMs
-		auto GetVertexAvgSquaredError = [&Mesh, &VtxSQEMs, &VtxAreas, this](int32 VID, const FSphere& Sphere) -> double
+		auto GetVertexAvgSquaredError = [&Mesh, &VtxSQEMs, &VtxAreas, this, NormalErrorWt](int32 VID, const FSphere& Sphere) -> double
 		{
 			FVector3d V = Mesh.GetVertex(VID);
 			double PosErr = (V - Sphere.Center).Length() - Sphere.W;
 			double PosErrContrib = PosErr * PosErr * PosErrorWt;
-			double NormalErrContrib = VtxAreas[VID] > 0 ? VtxSQEMs[VID].Eval(Sphere) / VtxAreas[VID] : 0.;
+			double NormalErrContrib = VtxAreas[VID] > 0 ? NormalErrorWt * VtxSQEMs[VID].Eval(Sphere) / VtxAreas[VID] : 0.;
 			return PosErrContrib + NormalErrContrib;
 		};
 
 		// Area-Weighted Average of face-Normal errors, without using FSQEM (should be equivalent to the above GetVertexAvgSquaredError, kept here for comparison -- currently this appears to be slower)
 		auto GetVertexAvgSquaredError_NoSQEM = 
-			[&Mesh, this, &VtxAreas, &TriNormals, &TriAreas, &VIDtoTIDs]
+			[&Mesh, this, &VtxAreas, &TriNormals, &TriAreas, &VIDtoTIDs, NormalErrorWt]
 			(int32 VID, const FSphere& Sphere) -> double
 		{
 			FVector3d V = Mesh.GetVertex(VID);
@@ -298,7 +299,7 @@ namespace UE::Geometry::MedialAxis
 					double SignedDist = N.Dot(V - Sphere.Center) - Sphere.W;
 					NormalErrContrib += TriAreas[TID] * SignedDist * SignedDist;
 				}
-				TotalErr += NormalErrContrib / AreaSum;
+				TotalErr += NormalErrorWt * NormalErrContrib / AreaSum;
 			}
 			return TotalErr;
 		};
@@ -628,7 +629,7 @@ namespace UE::Geometry::MedialAxis
 			// Optimize the sphere for each cluster
 			ParallelFor(ClustersToUpdate.Num(),
 				[this, &MeshBVH, &MeshFWTree, &Mesh, &Result, &VIDtoTIDs, &TriAreas, &TriNormals, InitialRadius, &MaxCenterBounds, MaxRadius, &FullPassesToUpdateCluster, &ClustersToUpdate, &VtxSQEMs,
-				&ClusterVertStarts, &ClusterSize, &VertsByCluster](int32 ClustersToUpdateIdx)
+				&ClusterVertStarts, &ClusterSize, &VertsByCluster, NormalErrorWt](int32 ClustersToUpdateIdx)
 			{
 				int32 ClusterIdx = ClustersToUpdate[ClustersToUpdateIdx];
 				if (!ensure(FullPassesToUpdateCluster[ClusterIdx] > 0))
@@ -649,6 +650,7 @@ namespace UE::Geometry::MedialAxis
 					int32 VID = VertsByCluster[Idx];
 					A_SQEM += VtxSQEMs[VID].A;
 				}
+				A_SQEM *= NormalErrorWt;
 				for (int32 Iter = 0; Iter < MaxGaussNewtonIterations; ++Iter)
 				{
 					FEigMat4 A = FEigMat4::Zero();
@@ -666,7 +668,7 @@ namespace UE::Geometry::MedialAxis
 							FVector3d TriNormal = (FVector3d)TriNormals[TID];
 							FEigVec4 N_1(TriNormal.X, TriNormal.Y, TriNormal.Z, 1.);
 							// Plane-error contribution to A added below via A_SQEM
-							B += N_1 * (AreaWt * ((VertPos - Result.Spheres[ClusterIdx].Center).Dot((FVector3d)TriNormal) - Result.Spheres[ClusterIdx].W));
+							B += N_1 * (NormalErrorWt * AreaWt * ((VertPos - Result.Spheres[ClusterIdx].Center).Dot((FVector3d)TriNormal) - Result.Spheres[ClusterIdx].W));
 						}
 						FVector3d FromCenterDir = (VertPos - Result.Spheres[ClusterIdx].Center);
 						double FromCenterDist = Normalize(FromCenterDir);
diff --git a/Engine/Plugins/Runtime/GeometryProcessing/Source/GeometryAlgorithms/Public/Skeletonization/MeshMedialAxisSampling.h b/Engine/Plugins/Runtime/GeometryProcessing/Source/GeometryAlgorithms/Public/Skeletonization/MeshMedialAxisSampling.h
--- a/Engine/Plugins/Runtime/GeometryProcessing/Source/GeometryAlgorithms/Public/Skeletonization/MeshMedialAxisSampling.h
+++ b/Engine/Plugins/Runtime/GeometryProcessing/Source/GeometryAlgorithms/Public/Skeletonization/MeshMedialAxisSampling.h
@@ -386,7 +386,7 @@ namespace MedialAxis
 		int32 MaxSpheres = 10;
 		// Do not split a cluster if its max error is below this threshold
 		double MinClusterErrorToSplit = .1;
-		// Weight for the position error term. Relative to plane error term, which has weight 1.
+		// Weight for the position error term. Relative to plane error term, so normal and position error weights sum to 1. Should be in the (0,1] range.
 		double PosErrorWt = .2;
 		// Whether to test medial skeleton edges for intersection w/ the input surface, and try to refine where these intersections are found
 		bool bSplitClustersIfEdgesIntersectSurface = true;
