fix perf bug: minimal hole filler accidentally passes mesh by value on CheckIfEdgeFlipCreatesFlip helper function

#Synced-CL 51695643

diff --git a/Engine/Plugins/Runtime/GeometryProcessing/Source/DynamicMesh/Private/Operations/MinimalHoleFiller.cpp b/Engine/Plugins/Runtime/GeometryProcessing/Source/DynamicMesh/Private/Operations/MinimalHoleFiller.cpp
--- a/Engine/Plugins/Runtime/GeometryProcessing/Source/DynamicMesh/Private/Operations/MinimalHoleFiller.cpp
+++ b/Engine/Plugins/Runtime/GeometryProcessing/Source/DynamicMesh/Private/Operations/MinimalHoleFiller.cpp
@@ -137,7 +137,7 @@ namespace
 	/// If before a flip we have normals (n1,n2) and after we have (m1,m2), check if
 	/// the dot between any of the 4 pairs changes sign after the flip, or is
 	/// less than the dot-product tolerance (i.e. angle tolerance)
-	bool CheckIfEdgeFlipCreatesFlip(const FDynamicMesh3 Mesh, int EdgeID, double FlipDotTol = 0.0)
+	bool CheckIfEdgeFlipCreatesFlip(const FDynamicMesh3& Mesh, int EdgeID, double FlipDotTol = 0.0)
 	{
 		check(Mesh.IsBoundaryEdge(EdgeID) == false);
 		const FDynamicMesh3::FEdge Edge = Mesh.GetEdge(EdgeID);
