Add medial skeleton simplifier to reduce skeletons computed by medial skeleton sampling
+ some basic test coverage

#Synced-CL 51925774

diff --git a/Engine/Plugins/Runtime/GeometryProcessing/Source/GeometryAlgorithms/Private/Skeletonization/MeshMedialAxisSampling.cpp b/Engine/Plugins/Runtime/GeometryProcessing/Source/GeometryAlgorithms/Private/Skeletonization/MeshMedialAxisSampling.cpp
--- a/Engine/Plugins/Runtime/GeometryProcessing/Source/GeometryAlgorithms/Private/Skeletonization/MeshMedialAxisSampling.cpp
+++ b/Engine/Plugins/Runtime/GeometryProcessing/Source/GeometryAlgorithms/Private/Skeletonization/MeshMedialAxisSampling.cpp
@@ -4,6 +4,7 @@
 
 #include <atomic>
 
+#include "Algo/RemoveIf.h"
 #include "VertexConnectedComponents.h"
 #include "MeshAdapter.h"
 #include "MeshAdapter/MeshVertexNormals.h"
@@ -13,6 +14,7 @@
 #include "DynamicMesh/DynamicVertexSkinWeightsAttribute.h"
 
 #include "Util/SizedDisjointSet.h"
+#include "VectorUtil.h"
 
 // Include Eigen/Dense for a 4x4 matrix solve
 #if defined(_MSC_VER) && USING_CODE_ANALYSIS
@@ -31,8 +33,724 @@ PRAGMA_DEFAULT_VISIBILITY_END
 #pragma warning(pop)
 #endif
 
+namespace MeshMedialAxisSamplingLocals
+{
+	using FEigMat4 = Eigen::Matrix4d;
+	using FEigVec4 = Eigen::Vector4d;
+
+	// Struct to sum/track/eval the spherical quadric error metric
+	struct FSQEM
+	{
+		// Terms for quadric error, s^T A s - s^T B + C
+		FEigMat4 A;
+		FEigVec4 B;
+		double C;
+		FSQEM() : A(FEigMat4::Zero()), B(FEigVec4::Zero()), C(0) {}
+		FSQEM(const FVector3d& Centroid, const FVector3d& Normal, double Wt)
+		{
+			FEigVec4 N_1(Normal.X, Normal.Y, Normal.Z, 1.);
+			A = (Wt * N_1) * N_1.transpose();
+			double NdP = Normal.Dot(Centroid);
+			B = N_1 * (2. * NdP * Wt);
+			C = NdP * NdP * Wt;
+		}
+		static FEigVec4 SphereToEig(const FSphere& Sphere)
+		{
+			return FEigVec4(Sphere.Center.X, Sphere.Center.Y, Sphere.Center.Z, Sphere.W);
+		}
+		double Eval(const FSphere& Sphere) const
+		{
+			return Eval(SphereToEig(Sphere));
+		}
+		double Eval(const FEigVec4& EigSphere) const
+		{
+			// Compute the spherical quadric error, s^T A s - s^T B + C
+			double Aterm = EigSphere.dot(A * EigSphere);
+			double Bterm = EigSphere.dot(B);
+			return Aterm - Bterm + C;
+		}
+		void Add(const FSQEM& Other)
+		{
+			A += Other.A;
+			B += Other.B;
+			C += Other.C;
+		}
+		void Add(const FSQEM& Other, double Wt)
+		{
+			A += Other.A * Wt;
+			B += Other.B * Wt;
+			C += Other.C * Wt;
+		}
+		// 'regularization' term biasing towards a specific sphere -- to avoid singular solve
+		void AddRegularizationToSphere(const FSphere& Sphere, double Wt)
+		{
+			A += FEigMat4::Identity() * Wt;
+			B += SphereToEig(Sphere) * (2. * Wt);
+		}
+		// plane error for sphere center; used to keep close to the original skeleton triangles
+		void AddCenterNormalQEM(const FVector3d& Point, const FVector3d& Normal, double Wt)
+		{
+			FEigVec4 N_0(Normal.X, Normal.Y, Normal.Z, 0.);
+			A += (Wt * N_0) * N_0.transpose();
+			double NdP = Normal.Dot(Point);
+			B += N_0 * (2. * NdP * Wt);
+			C += NdP * NdP * Wt;
+		}
+		// line error for sphere center; used to keep close to original skeleton edges
+		void AddCenterLineQEM(const FVector3d& Point, const FVector3d& NormalizedLineDir, double Wt)
+		{
+			FVector3d N1, N2;
+			UE::Geometry::VectorUtil::MakePerpVectors(NormalizedLineDir, N1, N2);
+			AddCenterNormalQEM(Point, N1, Wt * .5);
+			AddCenterNormalQEM(Point, N2, Wt * .5);
+		}
+		[[nodiscard]] bool SolveSphere(FSphere& OutSphere, double MinRadius = FMathd::ZeroTolerance) const
+		{
+			Eigen::LDLT<FEigMat4> LDLT = A.ldlt();
+			if (LDLT.info() != Eigen::Success)
+			{
+				return false;
+			}
+			FEigVec4 EigSphere = LDLT.solve(.5 * B);
+			if (!EigSphere.allFinite())
+			{
+				return false;
+			}
+			EigSphere[3] = FMath::Max(MinRadius, EigSphere[3]);
+			OutSphere = FSphere(FVector(EigSphere[0], EigSphere[1], EigSphere[2]), EigSphere[3]);
+			return true;
+		}
+	};
+
+	// helpers to get an edge or tri in canonical sorted form
+	UE::Geometry::FIndex2i AsEdge(int32 A, int32 B)
+	{
+		UE::Geometry::FIndex2i Edge(A, B); Edge.Sort();
+		return Edge;
+	}
+	UE::Geometry::FIndex3i AsTri(int32 A, int32 B, int32 C)
+	{
+		UE::Geometry::FIndex3i Tri(A, B, C); Tri.Sort();
+		return Tri;
+	}
+	// helper to interpolate two spheres via lerp of centers and radii
+	FSphere LerpSphere(const FSphere& A, const FSphere& B, double T)
+	{
+		return FSphere(
+			FMath::Lerp(A.Center, B.Center, T),
+			FMath::Lerp(A.W, B.W, T)
+		);
+	}
+
+
+	// Build up standard mesh QEM error terms for the skeleton geometry
+	// Uses cluster areas rather than the skeleton geometry for weighting
+	void AccumulateSkeletonDistanceQEMs(
+		const UE::Geometry::MedialAxis::FMedialSkeleton& Skeleton,
+		TConstArrayView<double> ClusterAreas,
+		int32 NumClusters,
+		TFunctionRef<bool(const UE::Geometry::FIndex2i&)> EdgeHasTris,
+		double SkeletonDistanceWt,
+		TArrayView<FSQEM> ClusterSQEMs)
+	{
+		using namespace UE::Geometry;
+
+		// Sum weights of all neighbors, so we can normalize QEM weights below
+		TArray<double> NeighborFeatureAreaSums;
+		NeighborFeatureAreaSums.SetNumZeroed(NumClusters);
+
+		for (const FIndex3i& Tri : Skeleton.ClusterTriangles)
+		{
+			for (int32 Idx = 0, Prev = 2, Next = 1; Idx < 3; Next = Prev, Prev = Idx++)
+			{
+				NeighborFeatureAreaSums[Tri[Idx]] += ClusterAreas[Tri[Prev]] + ClusterAreas[Tri[Next]];
+			}
+		}
+		for (int32 ClusterIdx = 0; ClusterIdx < NumClusters; ++ClusterIdx)
+		{
+			for (int32 NbrIdx : Skeleton.ClusterNeighbors[ClusterIdx])
+			{
+				if (NbrIdx <= ClusterIdx)
+				{
+					continue;
+				}
+				FIndex2i Edge(ClusterIdx, NbrIdx);
+				if (!EdgeHasTris(Edge))
+				{
+					NeighborFeatureAreaSums[ClusterIdx] += ClusterAreas[NbrIdx];
+					NeighborFeatureAreaSums[NbrIdx] += ClusterAreas[ClusterIdx];
+				}
+			}
+		}
+
+		// Add the triangle QEM terms
+		for (const FIndex3i& Tri : Skeleton.ClusterTriangles)
+		{
+			const FVector3d PA = Skeleton.Spheres[Tri.A].Center;
+			const FVector3d PB = Skeleton.Spheres[Tri.B].Center;
+			const FVector3d PC = Skeleton.Spheres[Tri.C].Center;
+			FVector3d Normal = VectorUtil::Normal(PA, PB, PC);
+			const FVector3d Centroid = (PA + PB + PC) / 3.0;
+			for (int32 Idx = 0, Prev = 2, Next = 1; Idx < 3; Next = Prev, Prev = Idx++)
+			{
+				const int32 VID = Tri[Idx];
+				if (NeighborFeatureAreaSums[VID] > 0.)
+				{
+					double OtherAreaSum = ClusterAreas[Tri[Prev]] + ClusterAreas[Tri[Next]];
+					const double W = (OtherAreaSum / NeighborFeatureAreaSums[VID]) * ClusterAreas[VID] * SkeletonDistanceWt;
+					ClusterSQEMs[VID].AddCenterNormalQEM(Centroid, Normal, W);
+				}
+			}
+		}
+		// Add the edge QEM terms (only on edge's w/ no triangle attached)
+		for (int32 ClusterIdx = 0; ClusterIdx < NumClusters; ++ClusterIdx)
+		{
+			for (int32 NbrIdx : Skeleton.ClusterNeighbors[ClusterIdx])
+			{
+				if (NbrIdx <= ClusterIdx)
+				{
+					continue;
+				}
+				FIndex2i Edge(ClusterIdx, NbrIdx);
+				if (EdgeHasTris(Edge))
+				{
+					continue;
+				}
+				const FVector3d PA = Skeleton.Spheres[ClusterIdx].Center;
+				const FVector3d PB = Skeleton.Spheres[NbrIdx].Center;
+				FVector3d EdgeDir = PB - PA;
+				const double EdgeLen = Normalize(EdgeDir);
+				if (EdgeLen < FMathd::ZeroTolerance)
+				{
+					continue;
+				}
+				if (NeighborFeatureAreaSums[ClusterIdx] > 0.)
+				{
+					const double WA = (ClusterAreas[NbrIdx] / NeighborFeatureAreaSums[ClusterIdx]) * ClusterAreas[ClusterIdx] * SkeletonDistanceWt;
+					ClusterSQEMs[ClusterIdx].AddCenterLineQEM(PA, EdgeDir, WA);
+				}
+				if (NeighborFeatureAreaSums[NbrIdx] > 0.)
+				{
+					const double WB = (ClusterAreas[ClusterIdx] / NeighborFeatureAreaSums[NbrIdx]) * ClusterAreas[NbrIdx] * SkeletonDistanceWt;
+					ClusterSQEMs[NbrIdx].AddCenterLineQEM(PA, EdgeDir, WB);
+				}
+			}
+		}
+	}
+
+	// Helper to compute the 'weld mapping' for vertices, if needed (ImplicitVertexNormalWeldThreshold > 0)
+	// (returns an empty array otherwise, we fall back to identity mapping in that case)
+	template<typename MeshType>
+	TArray<int32> ComputeWeldIdx(
+		const MeshType& Mesh,
+		TFunctionRef<bool(int32)> OptionalBoundaryVertexFilter,
+		double ImplicitVertexNormalWeldThreshold)
+	{
+		using namespace UE::Geometry;
+		TArray<int32> WeldIdx;
+		if (ImplicitVertexNormalWeldThreshold > 0)
+		{
+			FVertexConnectedComponents WeldVerts(Mesh.MaxVertexID());
+			WeldVerts.ConnectCloseFilteredVertices(Mesh, ImplicitVertexNormalWeldThreshold, OptionalBoundaryVertexFilter);
+			WeldIdx.SetNumUninitialized(Mesh.MaxVertexID());
+			ParallelFor(Mesh.MaxVertexID(), [&WeldIdx, &WeldVerts](int32 VID) { WeldIdx[VID] = WeldVerts.GetComponentThreadsafe(VID); });
+		}
+		return WeldIdx;
+	}
+	// Helper to compute per-triangle normals/areas and per-vertex area sums and SQEMs.
+	// OutTriNormals and OutTriAreas are filled for use by other parts of ComputeSkeleton.
+	// Returns total mesh surface area.
+	template<typename MeshType>
+	double ComputeVertexAreasAndSQEMs(
+		const MeshType& Mesh,
+		TFunctionRef<int32(int32)> VIDtoWeldVIDFn,
+		TArray<FVector3f>& OutTriNormals,
+		TArray<float>& OutTriAreas,
+		TArray<FSQEM>& OutVtxSQEMs,
+		TArray<float>& OutVtxAreas)
+	{
+		using namespace UE::Geometry;
+		OutTriNormals.SetNumUninitialized(Mesh.MaxTriangleID());
+		OutTriAreas.SetNumUninitialized(Mesh.MaxTriangleID());
+		OutVtxSQEMs.SetNumZeroed(Mesh.MaxVertexID());
+		OutVtxAreas.SetNumZeroed(Mesh.MaxVertexID());
+		double TotalArea = 0;
+		for (int32 TID = 0; TID < Mesh.MaxTriangleID(); ++TID)
+		{
+			if (!Mesh.IsTriangle(TID))
+			{
+				continue;
+			}
+			FVector3d Centroid, Normal;
+			double Area;
+			TMeshQueries<MeshType>::GetTriNormalAreaCentroid(Mesh, TID, Normal, Area, Centroid);
+			TotalArea += Area;
+			FIndex3i Tri = Mesh.GetTriangle(TID);
+			OutTriNormals[TID] = (FVector3f)Normal;
+			OutTriAreas[TID] = (float)Area;
+			for (int32 SubIdx = 0; SubIdx < 3; ++SubIdx)
+			{
+				// Note: vertex areas are per-triangle sums (not divided by 3); 1/3 factors cancel out in usages
+				OutVtxAreas[VIDtoWeldVIDFn(Tri[SubIdx])] += OutTriAreas[TID];
+				OutVtxSQEMs[VIDtoWeldVIDFn(Tri[SubIdx])].Add(FSQEM(Centroid, Normal, Area));
+			}
+		}
+		return TotalArea;
+	}
+	template<typename MeshType>
+	void SkeletonSimplify(const UE::Geometry::MedialAxis::FSkeletonSimplifier& Simplifier, UE::Geometry::MedialAxis::FMedialSkeleton& Skeleton, const UE::Geometry::TMeshAABBTree3<MeshType>& MeshBVH,
+		TConstArrayView<FSQEM> VtxSQEMs, TConstArrayView<float> VtxAreas)
+	{
+		using namespace UE::Geometry;
+
+		const MeshType& Mesh = *MeshBVH.GetMesh();
+
+		const int32 InitialNumClusters = Skeleton.Spheres.Num();
+
+		using FNbrTriArray = TArray<FIndex2i, TInlineAllocator<8>>;
+		TArray<FNbrTriArray> NbrTris;
+		NbrTris.SetNum(InitialNumClusters);
+		for (const FIndex3i& Tri : Skeleton.ClusterTriangles)
+		{
+			for (int32 Idx = 0, Prev = 2, Next = 1; Idx < 3; Next = Prev, Prev = Idx++)
+			{
+				const int32 V = Tri[Idx];
+				NbrTris[V].Add(AsEdge(Tri[Prev], Tri[Next]));
+			}
+		}
+
+		auto EdgeHasTris = [&NbrTris](const FIndex2i& Edge) -> bool
+		{
+			for (const FIndex2i& OppE : NbrTris[Edge.A])
+			{
+				if (OppE.Contains(Edge.B))
+				{
+					return true;
+				}
+			}
+			return false;
+		};
+
+		// Track which cluster each cluster has been merged into so far
+		TArray<int32> ClusterMergeParent;
+		ClusterMergeParent.SetNumUninitialized(InitialNumClusters);
+		for (int32 Idx = 0; Idx < InitialNumClusters; ++Idx)
+		{
+			ClusterMergeParent[Idx] = Idx; // by convention, self-index == no merge
+		}
+
+		// Tracking int incremented when cluster is updated (so we can know if our computed error is stale)
+		TArray<int32> ClusterVersion;
+		ClusterVersion.SetNumZeroed(InitialNumClusters);
+
+		TArray<FSQEM> ClusterSQEMs;
+		ClusterSQEMs.SetNumZeroed(InitialNumClusters);
+		TArray<double> ClusterBaseErrors;
+		ClusterBaseErrors.SetNumZeroed(InitialNumClusters);
+		TArray<double> ClusterAreas;
+		ClusterAreas.SetNumZeroed(InitialNumClusters);
+		const int32 MeshMaxVID = Mesh.MaxVertexID();
+
+		double UseSkelDistWt = FMath::Clamp(Simplifier.ClusterSkeletonDistanceWt, 0., 1.);
+		double UseSphereWt = 1. - UseSkelDistWt;
+		double UseRegWt = Simplifier.ClusterRegularizeWeight;
+		// If we don't have sphere weights we *must* have regularization to have some error term care about the sphere radii
+		if (UseSphereWt < FMathd::ZeroTolerance)
+		{
+			UseRegWt = FMath::Max(FMathd::ZeroTolerance, UseRegWt);
+		}
+
+		for (int32 VID = 0; VID < MeshMaxVID; ++VID)
+		{
+			int32 ClusterIdx = Skeleton.VIDtoClusterIndex[VID];
+			if (ClusterIdx != INDEX_NONE)
+			{
+				ClusterSQEMs[ClusterIdx].Add(VtxSQEMs[VID], UseSphereWt);
+				ClusterAreas[ClusterIdx] += VtxAreas[VID];
+			}
+		}
+		if (UseSkelDistWt > 0.)
+		{
+			AccumulateSkeletonDistanceQEMs(Skeleton, ClusterAreas, InitialNumClusters, EdgeHasTris, UseSkelDistWt, ClusterSQEMs);
+		}
+		if (UseRegWt > 0.)
+		{
+			for (int32 ClusterIdx = 0; ClusterIdx < InitialNumClusters; ++ClusterIdx)
+			{
+				ClusterSQEMs[ClusterIdx].AddRegularizationToSphere(Skeleton.Spheres[ClusterIdx], UseRegWt);
+			}
+		}
+
+		// Compute base errors to compare against, so we worry about the error introduced by collapses rather than the total error
+		for (int32 ClusterIdx = 0; ClusterIdx < InitialNumClusters; ++ClusterIdx)
+		{
+			ClusterBaseErrors[ClusterIdx] = ClusterSQEMs[ClusterIdx].Eval(Skeleton.Spheres[ClusterIdx]);
+		}
+
+		// track dead-end clusters that we should be allowed to collapse even if bOnlySimplifySurfaces==true
+		// because they were created from a collapsed surface (and are often not 'real' features)
+		TSet<int32> CollapseInducedDeadEnds;
+
+		auto CollapseEdge = [&Simplifier, &ClusterMergeParent, &ClusterVersion, &ClusterSQEMs, &ClusterBaseErrors, &Skeleton, &NbrTris, &CollapseInducedDeadEnds](FIndex2i Edge, const FSphere& NewSphere)
+		{
+			ensure(ClusterMergeParent[Edge.A] == Edge.A && ClusterMergeParent[Edge.B] == Edge.B);
+			// always sort and remove the higher vertex, so relative ordering from construction is approximately preserved through simplification
+			Edge.Sort();
+			ClusterMergeParent[Edge.B] = Edge.A;
+			ClusterVersion[Edge.A]++;
+			ClusterVersion[Edge.B]++;
+			ClusterSQEMs[Edge.A].Add(ClusterSQEMs[Edge.B]);
+			ClusterBaseErrors[Edge.A] += ClusterBaseErrors[Edge.B];
+			Skeleton.Spheres[Edge.A] = NewSphere;
+
+			// flag used for the 'CollapseInducedDeadEnds' tagging, below;
+			// heuristically, if the triangle already looked kind of like a dead-end, it's more of a 'real' feature, so
+			// we will get better results not tagging it as an 'induced'/collapsable edge.
+			bool bAorBWasTerminal = Skeleton.ClusterNeighbors[Edge.A].Num() < 3 || Skeleton.ClusterNeighbors[Edge.B].Num() < 3;
+
+			// update neighbor lists of B's neighbors to (1) remove B, and (2) add A & A's back-ref (unless the neighbor was A itself)
+			for (int32 BNbr : Skeleton.ClusterNeighbors[Edge.B])
+			{
+				Skeleton.ClusterNeighbors[BNbr].Remove(Edge.B);
+				if (BNbr != Edge.A)
+				{
+					Skeleton.ClusterNeighbors[BNbr].AddUnique(Edge.A);
+					Skeleton.ClusterNeighbors[Edge.A].AddUnique(BNbr);
+				}
+			}
+			// update NbrTris of B
+			for (int32 Idx = 0, Num = NbrTris[Edge.B].Num(); Idx < Num; ++Idx)
+			{
+				FIndex2i OppEdge = NbrTris[Edge.B][Idx];
+				int32 ASubIdx = OppEdge.IndexOf(Edge.A);
+				// tri includes AB edge, should be deleted by collapse
+				if (ASubIdx != INDEX_NONE)
+				{
+					int32 OtherV = OppEdge[1 - ASubIdx];
+					NbrTris[OtherV].Remove(Edge); // ok b/c edges canonically have vertices in sorted order
+					FIndex2i EdgeFromA(Edge.B, OtherV); EdgeFromA.Sort();
+					NbrTris[Edge.A].Remove(EdgeFromA); // (note: could do this in a separate pass if there may be a lot of tris on AB ...)
+					// If cluster on OtherV has only one neighbor, it'll become a new 'dead end' branch -- conditionally tag it as such
+					if (Simplifier.bOnlySimplifySurfaces && !bAorBWasTerminal && Skeleton.ClusterNeighbors[OtherV].Num() == 1)
+					{
+						CollapseInducedDeadEnds.Add(OtherV);
+					}
+				}
+				else // tri not on AB edge, need to remap B to A and add to A's neighbors
+				{
+					NbrTris[Edge.A].Add(OppEdge);
+					// remap the two copies of tri on edge's vertices
+					for (int32 SubIdx = 0; SubIdx < 2; ++SubIdx)
+					{
+						FIndex2i ToUpdate(Edge.B, OppEdge[1 - SubIdx]);
+						ToUpdate.Sort();
+						FIndex2i UpdateTo(Edge.A, OppEdge[1 - SubIdx]);
+						UpdateTo.Sort();
+						for (FIndex2i& SearchEdge : NbrTris[OppEdge[SubIdx]])
+						{
+							if (SearchEdge == ToUpdate)
+							{
+								SearchEdge = UpdateTo;
+								break;
+							}
+						}
+					}
+				}
+			}
+			Skeleton.ClusterNeighbors[Edge.B].Empty();
+			NbrTris[Edge.B].Empty();
+		};
+
+		struct FEdgeInfo
+		{
+			FIndex2i Edge, Version;
+			FSphere MergeSphere;
+			float Error;
+		};
+
+		TArray<FEdgeInfo> EdgeQ;
+
+		auto ComputeMergeError = [&Skeleton, &ClusterSQEMs, &ClusterBaseErrors](const FIndex2i& Edge, FSphere& OutSphere) -> float
+		{
+			double BaseError = ClusterBaseErrors[Edge.A] + ClusterBaseErrors[Edge.B];
+			FSQEM SQEM = ClusterSQEMs[Edge.A];
+			SQEM.Add(ClusterSQEMs[Edge.B]);
+			if (!SQEM.SolveSphere(OutSphere))
+			{
+				OutSphere = LerpSphere(Skeleton.Spheres[Edge.A], Skeleton.Spheres[Edge.B], .5);
+			}
+			return (float)(SQEM.Eval(OutSphere) - BaseError);
+		};
+
+		double UseEdgeLengthThresholdSq = Simplifier.EdgeLengthThreshold * Simplifier.EdgeLengthThreshold;
+		auto ShouldSkipEdge = [&Simplifier, &Skeleton, UseEdgeLengthThresholdSq, &EdgeHasTris, &CollapseInducedDeadEnds](const FIndex2i& Edge) -> bool
+		{
+			if (Simplifier.bOnlySimplifySurfaces && !EdgeHasTris(Edge) &&
+				// Allow collapse if this dead-end was artificially introduced by a prior triangle collapse.
+				!CollapseInducedDeadEnds.Contains(Edge.A) && !CollapseInducedDeadEnds.Contains(Edge.B))
+			{
+				return true;
+			}
+			if (Simplifier.EdgeLengthThreshold > 0 &&
+				FVector3d::DistSquared(Skeleton.Spheres[Edge.A].Center, Skeleton.Spheres[Edge.B].Center) > UseEdgeLengthThresholdSq)
+			{
+				return true;
+			}
+			if (Simplifier.EdgeRadiusThreshold > 0 &&
+				Skeleton.Spheres[Edge.A].W > Simplifier.EdgeRadiusThreshold && Skeleton.Spheres[Edge.B].W > Simplifier.EdgeRadiusThreshold)
+			{
+				return true;
+			}
+			return false;
+		};
+
+		for (int32 ClusterIdx = 0; ClusterIdx < InitialNumClusters; ++ClusterIdx)
+		{
+			for (int32 NbrIdx : Skeleton.ClusterNeighbors[ClusterIdx])
+			{
+				if (NbrIdx <= ClusterIdx)
+				{
+					continue;
+				}
+				FIndex2i Edge(ClusterIdx, NbrIdx);
+				if (ShouldSkipEdge(Edge))
+				{
+					continue;
+				}
+				FEdgeInfo Info;
+				Info.Edge = Edge;
+				Info.Version = FIndex2i(0, 0);
+				Info.Error = ComputeMergeError(Info.Edge, Info.MergeSphere);
+				if (Simplifier.QEMErrorThreshold <= 0 || Info.Error <= Simplifier.QEMErrorThreshold)
+				{
+					EdgeQ.Add(Info);
+				}
+			}
+		}
+		auto EdgeSort = [](const FEdgeInfo& A, const FEdgeInfo& B)
+		{
+			return A.Error < B.Error;
+		};
+		auto UpdateClusterIndexFromMerge = [&ClusterMergeParent](int32& Parent) -> void
+		{
+			while (ClusterMergeParent[Parent] != Parent)
+			{
+				ClusterMergeParent[Parent] = ClusterMergeParent[ClusterMergeParent[Parent]];
+				Parent = ClusterMergeParent[Parent];
+			}
+		};
+
+		EdgeQ.Heapify(EdgeSort);
+		int32 NumClusters = InitialNumClusters;
+		while (!EdgeQ.IsEmpty() && NumClusters > Simplifier.MinSpheres)
+		{
+			FEdgeInfo EdgeInfo;
+			EdgeQ.HeapPop(EdgeInfo, EdgeSort, EAllowShrinking::No);
+
+			FIndex2i& Edge = EdgeInfo.Edge;
+
+			FIndex2i& Version = EdgeInfo.Version;
+			if (ClusterVersion[Edge.A] > Version.A || ClusterVersion[Edge.B] > Version.B)
+			{
+				UpdateClusterIndexFromMerge(Edge.A);
+				UpdateClusterIndexFromMerge(Edge.B);
+				if (Edge.A == Edge.B)
+				{
+					// edge was merged away, skip
+					continue;
+				}
+				if (ShouldSkipEdge(Edge))
+				{
+					continue;
+				}
+				Version.A = ClusterVersion[Edge.A];
+				Version.B = ClusterVersion[Edge.B];
+				EdgeInfo.Error = ComputeMergeError(Edge, EdgeInfo.MergeSphere);
+				if (Simplifier.QEMErrorThreshold > 0 && EdgeInfo.Error > Simplifier.QEMErrorThreshold)
+				{
+					continue;
+				}
+
+				if (!EdgeQ.IsEmpty() && EdgeQ.HeapTop().Error < EdgeInfo.Error)
+				{
+					// after updating error, there's a cheaper edge on top of heap --
+					// need just re-insert this one and continue to pick up the smaller edge instead
+					EdgeQ.HeapPush(EdgeInfo, EdgeSort);
+
+					continue;
+				}
+			}
+			else
+			{
+				if (ShouldSkipEdge(Edge))
+				{
+					continue;
+				}
+			}
+
+			if (Simplifier.bPreventEdgeSurfaceIntersections)
+			{
+				bool bFoundHit = false;
+				for (int32 SubIdx = 0; !bFoundHit && SubIdx < 2; ++SubIdx)
+				{
+					int32 SrcClusterIdx = Edge[SubIdx];
+					int32 OtherIdx = Edge[1 - SubIdx];
+					FVector3d SrcPos = EdgeInfo.MergeSphere.Center;
+					for (int32 ToClusterIdx : Skeleton.ClusterNeighbors[SrcClusterIdx])
+					{
+						if (ToClusterIdx != OtherIdx)
+						{
+							FVector3d DestPos = Skeleton.Spheres[ToClusterIdx].Center;
+							FVector3d EdgeVec = DestPos - SrcPos;
+							double Len = Normalize(EdgeVec);
+							if (Len > 0)
+							{
+								FRay3d Ray(SrcPos, EdgeVec);
+								IMeshSpatial::FQueryOptions Options;
+								Options.MaxDistance = Len;
+								if (MeshBVH.FindNearestHitTriangle(Ray, Options) != INDEX_NONE)
+								{
+									bFoundHit = true;
+									break;
+								}
+							}
+						}
+					}
+				}
+				if (bFoundHit)
+				{
+					continue;
+				}
+			}
+
+			// we're up-to-date and have the lowest error, so can merge
+			CollapseEdge(EdgeInfo.Edge, EdgeInfo.MergeSphere);
+			NumClusters--;
+		}
+
+		// Do a final pass to compact away the removed cluster spheres/info and remap indices accordingly
+		TArray<int32> RemapClusterIDs;
+		int32 FinalNumClusters = 0;
+		RemapClusterIDs.SetNumUninitialized(InitialNumClusters);
+		for (int32 Idx = 0; Idx < InitialNumClusters; ++Idx)
+		{
+			if (ClusterMergeParent[Idx] == Idx)
+			{
+				int32 NewIdx = FinalNumClusters++;
+				RemapClusterIDs[Idx] = NewIdx;
+				if (Idx != NewIdx)
+				{
+					ensure(Idx > NewIdx);
+					Skeleton.Spheres[NewIdx] = Skeleton.Spheres[Idx];
+				}
+			}
+			else
+			{
+				ensure(ClusterMergeParent[Idx] < Idx); // by construction, parent should always be the smallest index in the merged cluster
+				RemapClusterIDs[Idx] = RemapClusterIDs[ClusterMergeParent[Idx]];
+			}
+		}
+		Skeleton.Spheres.SetNum(FinalNumClusters);
+		for (int32 Idx = 0; Idx < InitialNumClusters; ++Idx)
+		{
+			int32 NewIdx = RemapClusterIDs[Idx];
+			if (ClusterMergeParent[Idx] != Idx)
+			{
+				continue;
+			}
+			if (NewIdx != Idx)
+			{
+				Skeleton.ClusterNeighbors[NewIdx] = MoveTemp(Skeleton.ClusterNeighbors[Idx]);
+			}
+			for (int32& Nbr : Skeleton.ClusterNeighbors[NewIdx])
+			{
+				Nbr = RemapClusterIDs[Nbr];
+			}
+		}
+		Skeleton.ClusterNeighbors.SetNum(FinalNumClusters);
+		ParallelFor(MeshMaxVID, [&Skeleton, &RemapClusterIDs](int32 VID)
+		{
+			int32& ClusterIdx = Skeleton.VIDtoClusterIndex[VID];
+			if (ClusterIdx != INDEX_NONE)
+			{
+				ClusterIdx = RemapClusterIDs[ClusterIdx];
+			}
+		});
+		Skeleton.ClusterTriangles.Empty();
+		TSet<FIndex3i> NewTriSet;
+		for (int32 Idx = 0; Idx < InitialNumClusters; ++Idx)
+		{
+			for (const FIndex2i& OppE : NbrTris[Idx])
+			{
+				NewTriSet.Add(AsTri(
+					RemapClusterIDs[OppE.A],
+					RemapClusterIDs[OppE.B],
+					RemapClusterIDs[Idx])
+				);
+			}
+		}
+		Skeleton.ClusterTriangles = NewTriSet.Array();
+	}
+}
+
 namespace UE::Geometry::MedialAxis
 {
+	bool FMedialSkeleton::CheckValidity(EValidityCheckFailMode FailMode) const
+	{
+		bool bIsOk = true;
+		TFunction<void(bool)> CheckOrFailF = [&bIsOk](bool b) { bIsOk = bIsOk && b; };
+		if (FailMode == EValidityCheckFailMode::Check)
+		{
+			CheckOrFailF = [&bIsOk](bool b)
+			{
+				checkf(b, TEXT("FMedialSkeleton::CheckValidity failed!"));
+				bIsOk = bIsOk && b;
+			};
+		}
+		else if (FailMode == EValidityCheckFailMode::Ensure)
+		{
+			CheckOrFailF = [&bIsOk](bool b)
+			{
+				ensureMsgf(b, TEXT("FMedialSkeleton::CheckValidity failed!"));
+				bIsOk = bIsOk && b;
+			};
+		}
+
+		// ClusterNeighbors must be sized to match Spheres
+		CheckOrFailF(ClusterNeighbors.Num() == Spheres.Num());
+
+		// Neighbor arrays must be bidirectional and contain no self-loops
+		for (int32 Idx = 0; Idx < ClusterNeighbors.Num(); ++Idx)
+		{
+			for (int32 Nbr : ClusterNeighbors[Idx])
+			{
+				CheckOrFailF(Nbr != Idx);
+				CheckOrFailF(Nbr >= 0 && Nbr < ClusterNeighbors.Num());
+				if (Nbr >= 0 && Nbr < ClusterNeighbors.Num())
+				{
+					CheckOrFailF(ClusterNeighbors[Nbr].Contains(Idx));
+				}
+			}
+		}
+
+		// Every edge of every triangle must appear in ClusterNeighbors
+		for (const FIndex3i& Tri : ClusterTriangles)
+		{
+			for (int32 Prev = 2, Idx = 0; Idx < 3; Prev = Idx++)
+			{
+				CheckOrFailF(Tri[Prev] >= 0 && Tri[Prev] < ClusterNeighbors.Num());
+				CheckOrFailF(Tri[Idx] >= 0 && Tri[Idx] < ClusterNeighbors.Num());
+				if (Tri[Prev] >= 0 && Tri[Prev] < ClusterNeighbors.Num() &&
+					Tri[Idx] >= 0 && Tri[Idx] < ClusterNeighbors.Num())
+				{
+					CheckOrFailF(ClusterNeighbors[Tri[Prev]].Contains(Tri[Idx]));
+					CheckOrFailF(ClusterNeighbors[Tri[Idx]].Contains(Tri[Prev]));
+				}
+			}
+		}
+
+		return bIsOk;
+	}
 
 	template<typename MeshType>
 	FMedialSkeleton FSkeletonViaSampling::ComputeSkeleton(
@@ -41,9 +759,8 @@ namespace UE::Geometry::MedialAxis
 		TFunctionRef<int32(int32)> VIDtoWeldVIDFn, 
 		const TFastWindingTree<MeshType>* MeshFWTree)
 	{
-		using FEigMat4 = Eigen::Matrix4d;
-		using FEigVec4 = Eigen::Vector4d;
-		
+		using namespace MeshMedialAxisSamplingLocals;
+
 		FMedialSkeleton Result;
 
 		const MeshType& Mesh = *MeshBVH.GetMesh();
@@ -197,70 +914,13 @@ namespace UE::Geometry::MedialAxis
 				}
 			});
 		};
-		
-		// Helper compute a spherical quadric error metric for a given triangle, and to sum and evaluate them for faster error calculation
-		struct FSQEM
-		{
-			FEigMat4 A;
-			FEigVec4 B;
-			double C;
-			FSQEM() : A(FEigMat4::Zero()), B(FEigVec4::Zero()), C(0) {}
-			FSQEM(const FVector3d& Centroid, const FVector3d& Normal, double Wt)
-			{
-				FEigVec4 N_1(Normal.X, Normal.Y, Normal.Z, 1.);
-				A = (Wt * N_1) * N_1.transpose();
-				double NdP = Normal.Dot(Centroid);
-				B = N_1 * (2. * NdP * Wt);
-				C = NdP * NdP * Wt;
-			}
-			double Eval(const FSphere& Sphere)
-			{
-				// Compute the spherical quadric error, s^T A s - s^T B + C
-				FEigVec4 S(Sphere.Center.X, Sphere.Center.Y, Sphere.Center.Z, Sphere.W);
-				double Aterm = S.transpose() * A * S;
-				double Bterm = S.transpose() * B;
-				return Aterm - Bterm + C;
-			}
-			void Add(const FSQEM& Other)
-			{
-				A += Other.A;
-				B += Other.B;
-				C += Other.C;
-			}
-		};
 
 		// cache the triangle normals and areas, and vertex areas and SQEMs, for use in optimization below
 		TArray<FVector3f> TriNormals;
 		TArray<FSQEM> VtxSQEMs;
 		TArray<float> TriAreas, VtxAreas;
-		double TotalArea = 0;
-		TriNormals.SetNumUninitialized(Mesh.MaxTriangleID());
-		TriAreas.SetNumUninitialized(Mesh.MaxTriangleID());
-		VtxSQEMs.SetNumZeroed(Mesh.MaxVertexID());
-		VtxAreas.SetNumZeroed(Mesh.MaxVertexID());
-		for (int32 TID = 0; TID < Mesh.MaxTriangleID(); ++TID)
-		{
-			if (!Mesh.IsTriangle(TID))
-			{
-				continue;
-			}
-
-			FVector3d Centroid, Normal;
-			double Area;
-			TMeshQueries<MeshType>::GetTriNormalAreaCentroid(Mesh, TID, Normal, Area, Centroid);
-			TotalArea += Area;
-			FIndex3i Tri = Mesh.GetTriangle(TID);
-			TriNormals[TID] = (FVector3f)Normal;
-			TriAreas[TID] = (float)Area;
-			for (int32 SubIdx = 0; SubIdx < 3; ++SubIdx)
-			{
-				// Note we could divide by 3 for the more logical vertex area, but would then need to divide by 3
-				// when scaling triangle area contributions in a number of places below as well ...
-				// In practice these 1/3 factors will cancel out in our usages
-				VtxAreas[VIDtoWeldVIDFn(Tri[SubIdx])] += TriAreas[TID];
-				VtxSQEMs[VIDtoWeldVIDFn(Tri[SubIdx])].Add(FSQEM(Centroid, Normal, Area));
-			}
-		}
+		// Note: vertex areas are per-triangle sums (not divided by 3); 1/3 factors cancel out in usages below
+		const double TotalArea = ComputeVertexAreasAndSQEMs(Mesh, VIDtoWeldVIDFn, TriNormals, TriAreas, VtxSQEMs, VtxAreas);
 
 		// Stop early on a fully degenerate / zero-area mesh
 		if (TotalArea == 0)
@@ -744,32 +1404,71 @@ namespace UE::Geometry::MedialAxis
 			}
 		}
 
+		if (Simplifier.IsSet())
+		{
+			SkeletonSimplify(*Simplifier, Result, MeshBVH, VtxSQEMs, VtxAreas);
+		}
+
 		return Result;
 	}
 
 	template<typename MeshType>
 	FMedialSkeleton FSkeletonViaSampling::ComputeSkeleton(const TMeshAABBTree3<MeshType>& MeshBVH, TFunctionRef<bool(int32)> OptionalBoundaryVertexFilter, double ImplicitVertexNormalWeldThreshold, const TFastWindingTree<MeshType>* MeshFWTree)
 	{
+		using namespace MeshMedialAxisSamplingLocals;
 		const MeshType& Mesh = *MeshBVH.GetMesh();
 
 		TArray<FVector3f> Normals;
 		Normals.SetNumUninitialized(Mesh.MaxVertexID());
 		auto NormalFn = [&Normals](int32 VID) -> FVector3d { return (FVector3d)Normals[VID]; };
-		if (ImplicitVertexNormalWeldThreshold > 0)
+		TArray<int32> WeldIdx = ComputeWeldIdx(Mesh, OptionalBoundaryVertexFilter, ImplicitVertexNormalWeldThreshold);
+		if (WeldIdx.IsEmpty())
+		{
+			MeshNormals::ComputeVertexNormals(Mesh, Normals);
+			return ComputeSkeleton(MeshBVH, NormalFn, [](int32 VID)->int32 {return VID;}, MeshFWTree);
+		}
+		else
 		{
-			FVertexConnectedComponents WeldVerts(Mesh.MaxVertexID());
-			WeldVerts.ConnectCloseFilteredVertices(Mesh, ImplicitVertexNormalWeldThreshold, OptionalBoundaryVertexFilter);
-			// Extract the fully-compacted weld parent array, since we will access it a lot
-			TArray<int32> WeldIdx;
-			WeldIdx.SetNumUninitialized(Mesh.MaxVertexID());
-			ParallelFor(Mesh.MaxVertexID(), [&WeldIdx, &WeldVerts](int32 VID) { WeldIdx[VID] = WeldVerts.GetComponentThreadsafe(VID); });
 			MeshNormals::ComputeWeldedVertexNormals(Mesh, [&WeldIdx](int32 VID)->int32 {return WeldIdx[VID];}, Normals);
 			return ComputeSkeleton(MeshBVH, NormalFn, [&WeldIdx](int32 VID)->int32 {return WeldIdx[VID];}, MeshFWTree);
 		}
+	}
+
+	template<typename MeshType>
+	void FSkeletonSimplifier::Simplify(
+		FMedialSkeleton& Skeleton,
+		const TMeshAABBTree3<MeshType>& MeshBVH,
+		TFunctionRef<FVector3d(int32)> GetVertexNormal,
+		TFunctionRef<int32(int32)> VIDtoWeldVIDFn)
+	{
+		using namespace MeshMedialAxisSamplingLocals;
+		const MeshType& Mesh = *MeshBVH.GetMesh();
+		TArray<FVector3f> TriNormals;
+		TArray<float> TriAreas;
+		TArray<FSQEM> VtxSQEMs;
+		TArray<float> VtxAreas;
+		ComputeVertexAreasAndSQEMs(Mesh, VIDtoWeldVIDFn, TriNormals, TriAreas, VtxSQEMs, VtxAreas);
+		SkeletonSimplify(*this, Skeleton, MeshBVH, VtxSQEMs, VtxAreas);
+	}
+
+	template<typename MeshType>
+	void FSkeletonSimplifier::Simplify(
+		FMedialSkeleton& Skeleton,
+		const TMeshAABBTree3<MeshType>& MeshBVH,
+		TFunctionRef<bool(int32)> OptionalBoundaryVertexFilter,
+		double ImplicitVertexNormalWeldThreshold)
+	{
+		using namespace MeshMedialAxisSamplingLocals;
+		const MeshType& Mesh = *MeshBVH.GetMesh();
+		TArray<int32> WeldIdx = ComputeWeldIdx(Mesh, OptionalBoundaryVertexFilter, ImplicitVertexNormalWeldThreshold);
+		auto NopNormals = [](int32) -> FVector3d { return FVector3d::ZeroVector; };
+		if (WeldIdx.IsEmpty())
+		{
+			Simplify(Skeleton, MeshBVH, NopNormals, [](int32 VID) -> int32 { return VID; });
+		}
 		else
 		{
-			MeshNormals::ComputeVertexNormals(Mesh, Normals);
-			return ComputeSkeleton(MeshBVH, NormalFn, [](int32 VID)->int32 {return VID;}, MeshFWTree);
+			Simplify(Skeleton, MeshBVH, NopNormals, [&WeldIdx](int32 VID) -> int32 { return WeldIdx[VID]; });
 		}
 	}
 
@@ -1117,14 +1816,17 @@ namespace UE::Geometry::MedialAxis
 
 
 
+	// Instantiate templated ComputeSkeleton methods for the supported mesh types
+	template GEOMETRYALGORITHMS_API FMedialSkeleton FSkeletonViaSampling::ComputeSkeleton<FDynamicMesh3>(const TMeshAABBTree3<FDynamicMesh3>& MeshBVH, TFunctionRef<FVector3d(int32)> GetVertexNormal, TFunctionRef<int32(int32)> VIDtoWeldVIDFn, const TFastWindingTree<FDynamicMesh3>* MeshFWTree);
+	template GEOMETRYALGORITHMS_API FMedialSkeleton FSkeletonViaSampling::ComputeSkeleton<FDynamicMesh3>(const TMeshAABBTree3<FDynamicMesh3>& MeshBVH, TFunctionRef<bool(int32)> OptionalBoundaryVertexFilter, double ImplicitVertexNormalWeldThreshold, const TFastWindingTree<FDynamicMesh3>* MeshFWTree);
+	template GEOMETRYALGORITHMS_API FMedialSkeleton FSkeletonViaSampling::ComputeSkeleton<FTriangleMeshAdapterd>(const TMeshAABBTree3<FTriangleMeshAdapterd>& MeshBVH, TFunctionRef<FVector3d(int32)> GetVertexNormal, TFunctionRef<int32(int32)> VIDtoWeldVIDFn, const TFastWindingTree<FTriangleMeshAdapterd>* MeshFWTree);
+	template GEOMETRYALGORITHMS_API FMedialSkeleton FSkeletonViaSampling::ComputeSkeleton<FTriangleMeshAdapterd>(const TMeshAABBTree3<FTriangleMeshAdapterd>& MeshBVH, TFunctionRef<bool(int32)> OptionalBoundaryVertexFilter, double ImplicitVertexNormalWeldThreshold, const TFastWindingTree<FTriangleMeshAdapterd>* MeshFWTree);
 
+	template GEOMETRYALGORITHMS_API void FSkeletonSimplifier::Simplify<FDynamicMesh3>(FMedialSkeleton& Skeleton, const TMeshAABBTree3<FDynamicMesh3>& MeshBVH, TFunctionRef<FVector3d(int32)> GetVertexNormal, TFunctionRef<int32(int32)> VIDtoWeldVIDFn);
+	template GEOMETRYALGORITHMS_API void FSkeletonSimplifier::Simplify<FTriangleMeshAdapterd>(FMedialSkeleton& Skeleton, const TMeshAABBTree3<FTriangleMeshAdapterd>& MeshBVH, TFunctionRef<FVector3d(int32)> GetVertexNormal, TFunctionRef<int32(int32)> VIDtoWeldVIDFn);
+	template GEOMETRYALGORITHMS_API void FSkeletonSimplifier::Simplify<FDynamicMesh3>(FMedialSkeleton& Skeleton, const TMeshAABBTree3<FDynamicMesh3>& MeshBVH, TFunctionRef<bool(int32)> OptionalBoundaryVertexFilter, double ImplicitVertexNormalWeldThreshold);
+	template GEOMETRYALGORITHMS_API void FSkeletonSimplifier::Simplify<FTriangleMeshAdapterd>(FMedialSkeleton& Skeleton, const TMeshAABBTree3<FTriangleMeshAdapterd>& MeshBVH, TFunctionRef<bool(int32)> OptionalBoundaryVertexFilter, double ImplicitVertexNormalWeldThreshold);
 
-	// Instantiate templated ComputeSkeleton methods for the supported mesh types
-	template FMedialSkeleton FSkeletonViaSampling::ComputeSkeleton<FDynamicMesh3>(const TMeshAABBTree3<FDynamicMesh3>& MeshBVH, TFunctionRef<FVector3d(int32)> GetVertexNormal, TFunctionRef<int32(int32)> VIDtoWeldVIDFn, const TFastWindingTree<FDynamicMesh3>* MeshFWTree);
-	template FMedialSkeleton FSkeletonViaSampling::ComputeSkeleton<FDynamicMesh3>(const TMeshAABBTree3<FDynamicMesh3>& MeshBVH, TFunctionRef<bool(int32)> OptionalBoundaryVertexFilter, double ImplicitVertexNormalWeldThreshold, const TFastWindingTree<FDynamicMesh3>* MeshFWTree);
-	template FMedialSkeleton FSkeletonViaSampling::ComputeSkeleton<FTriangleMeshAdapterd>(const TMeshAABBTree3<FTriangleMeshAdapterd>& MeshBVH, TFunctionRef<FVector3d(int32)> GetVertexNormal, TFunctionRef<int32(int32)> VIDtoWeldVIDFn, const TFastWindingTree<FTriangleMeshAdapterd>* MeshFWTree);
-	template FMedialSkeleton FSkeletonViaSampling::ComputeSkeleton<FTriangleMeshAdapterd>(const TMeshAABBTree3<FTriangleMeshAdapterd>& MeshBVH, TFunctionRef<bool(int32)> OptionalBoundaryVertexFilter, double ImplicitVertexNormalWeldThreshold, const TFastWindingTree<FTriangleMeshAdapterd>* MeshFWTree);
-	
 
 
 } // namespace UE::Geometry::MedialAxis
\ No newline at end of file
diff --git a/Engine/Plugins/Runtime/GeometryProcessing/Source/GeometryAlgorithms/Public/Skeletonization/MeshMedialAxisSampling.h b/Engine/Plugins/Runtime/GeometryProcessing/Source/GeometryAlgorithms/Public/Skeletonization/MeshMedialAxisSampling.h
--- a/Engine/Plugins/Runtime/GeometryProcessing/Source/GeometryAlgorithms/Public/Skeletonization/MeshMedialAxisSampling.h
+++ b/Engine/Plugins/Runtime/GeometryProcessing/Source/GeometryAlgorithms/Public/Skeletonization/MeshMedialAxisSampling.h
@@ -260,6 +260,10 @@ namespace MedialAxis
 		// Triangles of the medial skeleton, where the skeleton is a surface rather than a chain. Note: Not oriented, and generally not manifold -- so cannot be directly appended to an FDynamicMesh3
 		TArray<FIndex3i> ClusterTriangles;
 
+		// Check that the skeleton data structures are internally consistent
+		// Note: Relatively expensive, intended for debugging / testing
+		GEOMETRYALGORITHMS_API bool CheckValidity(EValidityCheckFailMode FailMode = EValidityCheckFailMode::Ensure) const;
+
 
 		// Compute the ClusterTriangles array for the current FMedialSkeleton, given a Mesh and assuming VIDtoClusterIndex is already computed
 		template<typename TriangleMeshType>
@@ -375,6 +379,48 @@ namespace MedialAxis
 		GEOMETRYALGORITHMS_API static void ToDynamicMeshSkeleton(const FMedialSkeleton& InMedialSkeleton, FDynamicMesh3& TargetMesh, FName WeightProfile, FMedialSkeletonToTreeSkeletonOptions Options, int32 RootIndex = INDEX_NONE);
 	};
 
+	struct FSkeletonSimplifier
+	{
+
+		/// Common settings
+
+		// Stop simplifying at this sphere count
+		int32 MinSpheres = 1;
+
+		// If positive, will not collapse edges if it would introduce more error than this threshold. Note: Error value is defined by the internal accumulated quadric metric, *not* a straightforward distance.
+		double QEMErrorThreshold = -1;
+
+		// If positive, the maximum edge length to consider for collapse
+		double EdgeLengthThreshold = -1;
+		
+		// If positive, the maximum cluster radius to consider for collapse (will not collapse if both cluster radii on an edge are larger than this)
+		double EdgeRadiusThreshold = -1;
+
+		// If true, only simplify edges that have triangles associated
+		bool bOnlySimplifySurfaces = false;
+
+		// How much the simplifier should attempt to remain close to the initial skeleton, vs attempting to keep medial spheres close to the mesh surface.
+		// Relative term, should be in the [0,1] range. At 0, the original skeleton will not be considered; at 1, the mesh surface will not be considered.
+		double ClusterSkeletonDistanceWt = .5;
+
+		// If true, prevent collapses that would introduce intersections between skeleton edges and the mesh
+		bool bPreventEdgeSurfaceIntersections = true;
+
+
+
+		/// Advanced settings
+
+		// Regularization weight encouraging skeleton clusters to keep their original positions and radii.
+		double ClusterRegularizeWeight = FMathd::ZeroTolerance;
+
+		
+		template<typename MeshType>
+		GEOMETRYALGORITHMS_API void Simplify(FMedialSkeleton& Skeleton, const TMeshAABBTree3<MeshType>& MeshBVH, TFunctionRef<FVector3d(int32)> GetVertexNormal, TFunctionRef<int32(int32)> VIDtoWeldVIDFn);
+
+		template<typename MeshType>
+		GEOMETRYALGORITHMS_API void Simplify(FMedialSkeleton& Skeleton, const TMeshAABBTree3<MeshType>& MeshBVH, TFunctionRef<bool(int32)> OptionalBoundaryVertexFilter, double ImplicitVertexNormalWeldThreshold);
+	};
+
 
 	// Implementation/adaptation of "Dynamic Skeletonization via Variational Medial Skeletal Sampling" by Huang et al
 	// with some surface-focused additions -- optionally uses raycasts to guide refinement, and fast winding to guide medial-sphere projection
@@ -398,6 +444,8 @@ namespace MedialAxis
 		// Settings for the medial ball shrinking, used to initialize medial spheres and to project spheres into the surface
 		FShrinkMedialBall ShrinkMedial;
 
+		// Optional simplifier, to be run after generating the initial skeleton. Saves re-computing some values compared to running the simplifier separately after.
+		TOptional<FSkeletonSimplifier> Simplifier;
 
 		/// Additional parameters -- most users should not need to change these
 
