add option for medial -> anim skeleton conversion to prefer edges that don't intersect a given surface

#Synced-CL 51925774

diff --git a/Engine/Plugins/Experimental/GeometryCollectionPlugin/Source/GeometryCollectionNodes/Private/Dataflow/MeshMedialSkeletonConversionNodes.cpp b/Engine/Plugins/Experimental/GeometryCollectionPlugin/Source/GeometryCollectionNodes/Private/Dataflow/MeshMedialSkeletonConversionNodes.cpp
--- a/Engine/Plugins/Experimental/GeometryCollectionPlugin/Source/GeometryCollectionNodes/Private/Dataflow/MeshMedialSkeletonConversionNodes.cpp
+++ b/Engine/Plugins/Experimental/GeometryCollectionPlugin/Source/GeometryCollectionNodes/Private/Dataflow/MeshMedialSkeletonConversionNodes.cpp
@@ -103,6 +103,7 @@ FConvertMedialSkeletonToAnimationSkeletonDataflowNode::FConvertMedialSkeletonToA
 	: FDataflowNode(InParam, InGuid)
 {
 	RegisterInputConnection(&MedialSkeleton);
+	RegisterInputConnection(&AvoidIntersectingMesh);
 	RegisterInputConnection(&RootClusterIndex).SetCanHidePin(true).SetPinIsHidden(true);
 
 	RegisterOutputConnection(&Skeleton);
@@ -117,13 +118,23 @@ void FConvertMedialSkeletonToAnimationSkeletonDataflowNode::Evaluate(UE::Dataflo
 		const FDataflowMedialSkeleton& InMedialSkeleton = GetValue(Context, &MedialSkeleton);
 		FDataflowSkeleton OutSkeleton;
 
+		TObjectPtr<const UDynamicMesh> InMesh = GetValue(Context, &AvoidIntersectingMesh);
+		
+		TMeshAABBTree3<FDynamicMesh3> MeshBVHLocal;
+		const TMeshAABBTree3<FDynamicMesh3>* AvoidIntersectingMeshBVH = nullptr;
+		if (InMesh && InMesh->GetMeshRef().TriangleCount() > 0)
+		{
+			MeshBVHLocal.SetMesh(InMesh->GetMeshPtr());
+			AvoidIntersectingMeshBVH = &MeshBVHLocal;
+		}
+
 		TArray<int32> Parents;
 		TArray<FTransform> Poses;
 
 		UE::Geometry::MedialAxis::FMedialSkeletonToTreeSkeletonOptions Options;
 		MeshMedialSkeletonConversionLocals::SetSkeletonConversionOptions(Options, EdgeWeightMethod, MergeDisconnectedMethod);
 		const int32 UseRootIndex = GetValue(Context, &RootClusterIndex);
-		Options.ToHierarchy(InMedialSkeleton.Skeleton, Parents, Poses, Options, UseRootIndex);
+		Options.ToHierarchy(InMedialSkeleton.Skeleton, AvoidIntersectingMeshBVH, Parents, Poses, Options, UseRootIndex);
 
 		const int32 NumBones = Parents.Num();
 		if (NumBones > 0)
diff --git a/Engine/Plugins/Experimental/GeometryCollectionPlugin/Source/GeometryCollectionNodes/Private/Dataflow/MeshMedialSkeletonConversionNodes.h b/Engine/Plugins/Experimental/GeometryCollectionPlugin/Source/GeometryCollectionNodes/Private/Dataflow/MeshMedialSkeletonConversionNodes.h
--- a/Engine/Plugins/Experimental/GeometryCollectionPlugin/Source/GeometryCollectionNodes/Private/Dataflow/MeshMedialSkeletonConversionNodes.h
+++ b/Engine/Plugins/Experimental/GeometryCollectionPlugin/Source/GeometryCollectionNodes/Private/Dataflow/MeshMedialSkeletonConversionNodes.h
@@ -57,6 +57,10 @@ private:
 	UPROPERTY(meta = (DataflowInput, DataflowIntrinsic))
 	FDataflowMedialSkeleton MedialSkeleton;
 
+	/** If provided, edges not intersecting this mesh will be preferred for inclusion in the skeleton */
+	UPROPERTY(meta = (DataflowInput))
+	TObjectPtr<const UDynamicMesh> AvoidIntersectingMesh;
+
 	UPROPERTY(meta = (DataflowOutput))
 	FDataflowSkeleton Skeleton;
 
diff --git a/Engine/Plugins/Runtime/GeometryProcessing/Source/GeometryAlgorithms/Private/Skeletonization/MeshMedialAxisSampling.cpp b/Engine/Plugins/Runtime/GeometryProcessing/Source/GeometryAlgorithms/Private/Skeletonization/MeshMedialAxisSampling.cpp
--- a/Engine/Plugins/Runtime/GeometryProcessing/Source/GeometryAlgorithms/Private/Skeletonization/MeshMedialAxisSampling.cpp
+++ b/Engine/Plugins/Runtime/GeometryProcessing/Source/GeometryAlgorithms/Private/Skeletonization/MeshMedialAxisSampling.cpp
@@ -238,6 +238,26 @@ namespace MeshMedialAxisSamplingLocals
 		}
 	}
 
+	template<typename MeshType>
+	bool EdgeIntersectsMesh(const FVector3d EdgeA, const FVector3d EdgeB, const UE::Geometry::TMeshAABBTree3<MeshType>& MeshBVH)
+	{
+		using namespace UE::Geometry;
+
+		FVector3d EdgeVec = EdgeB - EdgeA;
+		double Len = Normalize(EdgeVec);
+		if (Len > 0)
+		{
+			FRay3d Ray(EdgeA, EdgeVec);
+			IMeshSpatial::FQueryOptions Options;
+			Options.MaxDistance = Len;
+			if (MeshBVH.FindNearestHitTriangle(Ray, Options) != INDEX_NONE)
+			{
+				return true;
+			}
+		}
+		return false;
+	}
+
 	// Helper to compute the 'weld mapping' for vertices, if needed (ImplicitVertexNormalWeldThreshold > 0)
 	// (returns an empty array otherwise, we fall back to identity mapping in that case)
 	template<typename MeshType>
@@ -601,19 +621,12 @@ namespace MeshMedialAxisSamplingLocals
 						if (ToClusterIdx != OtherIdx)
 						{
 							FVector3d DestPos = Skeleton.Spheres[ToClusterIdx].Center;
-							FVector3d EdgeVec = DestPos - SrcPos;
-							double Len = Normalize(EdgeVec);
-							if (Len > 0)
+							if (EdgeIntersectsMesh(SrcPos, DestPos, MeshBVH))
 							{
-								FRay3d Ray(SrcPos, EdgeVec);
-								IMeshSpatial::FQueryOptions Options;
-								Options.MaxDistance = Len;
-								if (MeshBVH.FindNearestHitTriangle(Ray, Options) != INDEX_NONE)
-								{
-									bFoundHit = true;
-									break;
-								}
+								bFoundHit = true;
+								break;
 							}
+
 						}
 					}
 				}
@@ -1538,10 +1551,12 @@ namespace UE::Geometry::MedialAxis
 		}
 	}
 
-	void FMedialSkeletonToTreeSkeletonOptions::ToHierarchy(const FMedialSkeleton& InMedialSkeleton, TArray<int32>& OutParent, TArray<FTransform>& OutTransform, FMedialSkeletonToTreeSkeletonOptions Options, int32 InRootIndex,
+	template<typename MeshType>
+	void FMedialSkeletonToTreeSkeletonOptions::ToHierarchy(const FMedialSkeleton& InMedialSkeleton, const TMeshAABBTree3<MeshType>* AvoidIntersectingMeshBVH, TArray<int32>& OutParent, TArray<FTransform>& OutTransform, FMedialSkeletonToTreeSkeletonOptions Options, int32 InRootIndex,
 		TArray<int32>* OutBoneIndexToMedialIndex, TArray<int32>* OutMedialIndexToBoneIndex)
 	{
 		using namespace UE::Geometry;
+		using namespace MeshMedialAxisSamplingLocals;
 
 		const int32 NumClusters = InMedialSkeleton.ClusterNeighbors.Num();
 		if (!ensureMsgf(NumClusters == InMedialSkeleton.Spheres.Num(), TEXT("Invalid medial skeleton data; %d spheres does not match %d cluster neighbors"), InMedialSkeleton.Spheres.Num(), NumClusters))
@@ -1580,28 +1595,54 @@ namespace UE::Geometry::MedialAxis
 				}
 			}
 		}
-		
-		if (Options.EdgeWeightMethod == FMedialSkeletonToTreeSkeletonOptions::EEdgeWeightMethod::ArrayOrder)
-		{
-			Edges.Sort([](const FIndex2i& A, const FIndex2i& B) { return A.A + A.B < B.A + B.B; });
-		}
-		else if (Options.EdgeWeightMethod == FMedialSkeletonToTreeSkeletonOptions::EEdgeWeightMethod::AvgRadius)
+
+		// If we have a BVH to test edges for intersection, move the intersecting edges to the back of the array
+		int32 IsIntersectingStart = Edges.Num();
+		if (AvoidIntersectingMeshBVH)
 		{
-			Edges.Sort([&InMedialSkeleton](const FIndex2i& A, const FIndex2i& B) 
-				{
-					return	InMedialSkeleton.Spheres[A.A].W + InMedialSkeleton.Spheres[A.B].W > // note: greater than, so we prefer connections between larger spheres
-							InMedialSkeleton.Spheres[B.A].W + InMedialSkeleton.Spheres[B.B].W;
-				}
-			);
+			// Note Algo::RemoveIf is actually 'move to back of array if'
+			IsIntersectingStart = Algo::RemoveIf(Edges, [&](const FIndex2i& Edge)
+			{
+				return EdgeIntersectsMesh(
+					InMedialSkeleton.Spheres[Edge.A].Center,
+					InMedialSkeleton.Spheres[Edge.B].Center,
+					*AvoidIntersectingMeshBVH
+				);
+			});
 		}
-		else // EdgeLength
+		
+		// Sort the non-intersecting and intersecting ranges separately by the secondary sort criteria
+		for (int32 Range = 0; Range < 2; Range++)
 		{
-			Edges.Sort([&InMedialSkeleton](const FIndex2i& A, const FIndex2i& B)
-				{
-					return	FVector::DistSquared(InMedialSkeleton.Spheres[A.A].Center, InMedialSkeleton.Spheres[A.B].Center) <
-							FVector::DistSquared(InMedialSkeleton.Spheres[B.A].Center, InMedialSkeleton.Spheres[B.B].Center);
-				}
-			);
+			int32 Start = Range * IsIntersectingStart;
+			int32 Count = (Range ? Edges.Num() : IsIntersectingStart) - Start;;
+			TConstArrayView<FIndex2i> EdgesView(Edges.GetData() + Start, Count);
+			if (EdgesView.IsEmpty())
+			{
+				continue;
+			}
+			if (Options.EdgeWeightMethod == FMedialSkeletonToTreeSkeletonOptions::EEdgeWeightMethod::ArrayOrder)
+			{
+				EdgesView.Sort([](const FIndex2i& A, const FIndex2i& B) { return A.A + A.B < B.A + B.B; });
+			}
+			else if (Options.EdgeWeightMethod == FMedialSkeletonToTreeSkeletonOptions::EEdgeWeightMethod::AvgRadius)
+			{
+				EdgesView.Sort([&InMedialSkeleton](const FIndex2i& A, const FIndex2i& B)
+					{
+						return	InMedialSkeleton.Spheres[A.A].W + InMedialSkeleton.Spheres[A.B].W > // note: greater than, so we prefer connections between larger spheres
+								InMedialSkeleton.Spheres[B.A].W + InMedialSkeleton.Spheres[B.B].W;
+					}
+				);
+			}
+			else // EdgeLength
+			{
+				EdgesView.Sort([&InMedialSkeleton](const FIndex2i& A, const FIndex2i& B)
+					{
+						return	FVector::DistSquared(InMedialSkeleton.Spheres[A.A].Center, InMedialSkeleton.Spheres[A.B].Center) <
+								FVector::DistSquared(InMedialSkeleton.Spheres[B.A].Center, InMedialSkeleton.Spheres[B.B].Center);
+					}
+				);
+			}
 		}
 
 		// Use a disjoint set to track the connectivity as we add one edge at a time
@@ -1835,6 +1876,9 @@ namespace UE::Geometry::MedialAxis
 	template GEOMETRYALGORITHMS_API void FSkeletonSimplifier::Simplify<FDynamicMesh3>(FMedialSkeleton& Skeleton, const TMeshAABBTree3<FDynamicMesh3>& MeshBVH, TFunctionRef<bool(int32)> OptionalBoundaryVertexFilter, double ImplicitVertexNormalWeldThreshold);
 	template GEOMETRYALGORITHMS_API void FSkeletonSimplifier::Simplify<FTriangleMeshAdapterd>(FMedialSkeleton& Skeleton, const TMeshAABBTree3<FTriangleMeshAdapterd>& MeshBVH, TFunctionRef<bool(int32)> OptionalBoundaryVertexFilter, double ImplicitVertexNormalWeldThreshold);
 
+	template GEOMETRYALGORITHMS_API void FMedialSkeletonToTreeSkeletonOptions::ToHierarchy<FDynamicMesh3>(const FMedialSkeleton& InMedialSkeleton, const TMeshAABBTree3<FDynamicMesh3>* AvoidIntersectingMeshBVH, TArray<int32>& OutParent, TArray<FTransform>& OutTransform, FMedialSkeletonToTreeSkeletonOptions Options, int32 RootIndex, TArray<int32>* OutBoneIndexToMedialIndex, TArray<int32>* OutMedialIndexToBoneIndex);
+	template GEOMETRYALGORITHMS_API void FMedialSkeletonToTreeSkeletonOptions::ToHierarchy<FTriangleMeshAdapterd>(const FMedialSkeleton& InMedialSkeleton, const TMeshAABBTree3<FTriangleMeshAdapterd>* AvoidIntersectingMeshBVH, TArray<int32>& OutParent, TArray<FTransform>& OutTransform, FMedialSkeletonToTreeSkeletonOptions Options, int32 RootIndex, TArray<int32>* OutBoneIndexToMedialIndex, TArray<int32>* OutMedialIndexToBoneIndex);
+
 
 
 } // namespace UE::Geometry::MedialAxis
\ No newline at end of file
diff --git a/Engine/Plugins/Runtime/GeometryProcessing/Source/GeometryAlgorithms/Public/Skeletonization/MeshMedialAxisSampling.h b/Engine/Plugins/Runtime/GeometryProcessing/Source/GeometryAlgorithms/Public/Skeletonization/MeshMedialAxisSampling.h
--- a/Engine/Plugins/Runtime/GeometryProcessing/Source/GeometryAlgorithms/Public/Skeletonization/MeshMedialAxisSampling.h
+++ b/Engine/Plugins/Runtime/GeometryProcessing/Source/GeometryAlgorithms/Public/Skeletonization/MeshMedialAxisSampling.h
@@ -374,6 +374,7 @@ namespace MedialAxis
 		 * Note the indices of the hierarchical skeleton bones will be remapped from the medial skeleton, because URefSkeleton requires that parent bone indices be smaller than child bone indices
 		 * 
 		 * @param InMedialSkeleton Input medial skeleton
+		 * @param AvoidIntersectingMeshBVH Surface that the output skeleton should prefer not to cross. (Surface-intersecting edges will still be added as-needed, if the non-intersecting edges are not sufficient to create a spanning tree.)
 		 * @param OutParent Output hierarchy, as parent indices
 		 * @param OutTransform Relative transforms for each bone. 1:1 with OutParent.
 		 * @param Options Additional options to control the skeleton conversion
@@ -381,9 +382,17 @@ namespace MedialAxis
 		 * @param OutBoneIndexToMedialIndex Optional mapping from hierarchical skeleton bone indices to source medial skeleton cluster indices
 		 * @param OutMedialIndexToBoneIndex Optional mapping from source medial skeleton cluster indices to hierarchical skeleton bone indices
 		 */
-		GEOMETRYALGORITHMS_API static void ToHierarchy(const FMedialSkeleton& InMedialSkeleton, TArray<int32>& OutParent, TArray<FTransform>& OutTransform, FMedialSkeletonToTreeSkeletonOptions Options, int32 RootIndex = INDEX_NONE, 
+		template<typename MeshType>
+		GEOMETRYALGORITHMS_API static void ToHierarchy(const FMedialSkeleton& InMedialSkeleton, const TMeshAABBTree3<MeshType>* AvoidIntersectingMeshBVH, TArray<int32>& OutParent, TArray<FTransform>& OutTransform, FMedialSkeletonToTreeSkeletonOptions Options, int32 RootIndex = INDEX_NONE,
 			TArray<int32>* OutBoneIndexToMedialIndex = nullptr, TArray<int32>* OutMedialIndexToBoneIndex = nullptr);
 
+		// Same as above ToHierarchy method, but without the AvoidIntersectingMeshBVH parameter
+		static void ToHierarchy(const FMedialSkeleton& InMedialSkeleton, TArray<int32>& OutParent, TArray<FTransform>& OutTransform, FMedialSkeletonToTreeSkeletonOptions Options, int32 RootIndex = INDEX_NONE,
+			TArray<int32>* OutBoneIndexToMedialIndex = nullptr, TArray<int32>* OutMedialIndexToBoneIndex = nullptr)
+		{
+			return ToHierarchy<FDynamicMesh3>(InMedialSkeleton, nullptr, OutParent, OutTransform, Options, RootIndex, OutBoneIndexToMedialIndex, OutMedialIndexToBoneIndex);
+		}
+
 		/**
 		 * Transfer a medial skeleton to a dynamic mesh as a hierarchical skeleton, so it can be converted to a USkeletalMesh
 		 * 
