fix dead-end tags incorrectly propagating in medial skeleton simplification

#Synced-CL 52024969

diff --git a/Engine/Plugins/Runtime/GeometryProcessing/Source/GeometryAlgorithms/Private/Skeletonization/MeshMedialAxisSampling.cpp b/Engine/Plugins/Runtime/GeometryProcessing/Source/GeometryAlgorithms/Private/Skeletonization/MeshMedialAxisSampling.cpp
--- a/Engine/Plugins/Runtime/GeometryProcessing/Source/GeometryAlgorithms/Private/Skeletonization/MeshMedialAxisSampling.cpp
+++ b/Engine/Plugins/Runtime/GeometryProcessing/Source/GeometryAlgorithms/Private/Skeletonization/MeshMedialAxisSampling.cpp
@@ -427,6 +427,9 @@ namespace MeshMedialAxisSamplingLocals
 			// heuristically, if the triangle already looked kind of like a dead-end, it's more of a 'real' feature, so
 			// we will get better results not tagging it as an 'induced'/collapsable edge.
 			bool bAorBWasTerminal = Skeleton.ClusterNeighbors[Edge.A].Num() < 3 || Skeleton.ClusterNeighbors[Edge.B].Num() < 3;
+			// Clear any dead-end tags on collapsed edge vertices
+			CollapseInducedDeadEnds.Remove(Edge.A);
+			CollapseInducedDeadEnds.Remove(Edge.B);
 
 			// update neighbor lists of B's neighbors to (1) remove B, and (2) add A & A's back-ref (unless the neighbor was A itself)
 			for (int32 BNbr : Skeleton.ClusterNeighbors[Edge.B])
