Add const to the IsCompatibleWithMesh medial skeleton method

#Synced-CL 52179965

diff --git a/Engine/Plugins/Runtime/GeometryProcessing/Source/GeometryAlgorithms/Public/Skeletonization/MeshMedialAxisSampling.h b/Engine/Plugins/Runtime/GeometryProcessing/Source/GeometryAlgorithms/Public/Skeletonization/MeshMedialAxisSampling.h
--- a/Engine/Plugins/Runtime/GeometryProcessing/Source/GeometryAlgorithms/Public/Skeletonization/MeshMedialAxisSampling.h
+++ b/Engine/Plugins/Runtime/GeometryProcessing/Source/GeometryAlgorithms/Public/Skeletonization/MeshMedialAxisSampling.h
@@ -265,7 +265,7 @@ namespace MedialAxis
 		GEOMETRYALGORITHMS_API bool CheckValidity(EValidityCheckFailMode FailMode = EValidityCheckFailMode::Ensure) const;
 
 		template<typename TriangleMeshType>
-		bool IsCompatibleWithMesh(const TriangleMeshType& Mesh)
+		bool IsCompatibleWithMesh(const TriangleMeshType& Mesh) const
 		{
 			if (Mesh.MaxVertexID() != VIDtoClusterIndex.Num())
 			{
