Generalize SkinWeightBinding operation to support 'constrained' binding, where vertices can only bind to a known subset of bones.

(+ fix handling of non-compact meshes)

#Synced-CL 52243365

diff --git a/Engine/Plugins/Runtime/GeometryProcessing/Source/DynamicMesh/Private/Operations/SkinWeightBinding.cpp b/Engine/Plugins/Runtime/GeometryProcessing/Source/DynamicMesh/Private/Operations/SkinWeightBinding.cpp
--- a/Engine/Plugins/Runtime/GeometryProcessing/Source/DynamicMesh/Private/Operations/SkinWeightBinding.cpp
+++ b/Engine/Plugins/Runtime/GeometryProcessing/Source/DynamicMesh/Private/Operations/SkinWeightBinding.cpp
@@ -2,6 +2,7 @@
 
 #include "Operations/SkinWeightBinding.h"
 
+#include "Algo/RemoveIf.h"
 #include "BoneWeights.h"
 #include "DynamicMesh/DynamicMesh3.h"
 #include "DynamicMesh/DynamicMeshAttributeSet.h"
@@ -369,27 +370,30 @@ FDynamicMeshVertexSkinWeightsAttribute* GetOrCreateSkinWeightsAttribute(
 	return Attribute;
 }
 
-
-void CreateSkinWeights_DirectDistance(
-	FDynamicMesh3& InMesh, TConstArrayView<SkinBinding::FBonePoseInfo> TransformHierarchy, FName ProfileName,
+// Generalized helper for direct-distance bone weights, with enumerator to customize what bones are considered for each vertex
+template<typename EnumerateType>
+void CreateSkinWeightsImpl_DirectDistance_General(FDynamicMesh3& InMesh, TConstArrayView<SkinBinding::FBonePoseInfo> TransformHierarchy, FName ProfileName,
 	float InStiffness,
-	const AnimationCore::FBoneWeightsSettings& InSettings
-	)
+	const AnimationCore::FBoneWeightsSettings& InSettings,
+	const FTransformHierarchyQuery& Skeleton,
+	EnumerateType EnumerateBonesForVertex)
 {
 	using namespace AnimationCore;
 
 	FDynamicMeshVertexSkinWeightsAttribute *SkinWeights = GetOrCreateSkinWeightsAttribute(InMesh, ProfileName);
 		
-	const int32 NumVertices = InMesh.VertexCount();
+	const int32 MaxVertexID = InMesh.MaxVertexID();
 	const int32 MaxInfluences = InSettings.GetMaxWeightCount();
 
 	// Use the diagonal size of the bbox to make the bone distance falloff scale invariant.
 	const float DiagBounds = (float)FMath::Max(InMesh.GetBounds(true).DiagonalLength(), UE_SMALL_NUMBER);
 	
-	const FTransformHierarchyQuery Skeleton(TransformHierarchy);
-
-	ParallelFor(NumVertices, [&](const int32 VertexIdx)
+	ParallelFor(MaxVertexID, [&](const int32 VertexIdx)
 	{
+		if (!InMesh.IsVertex(VertexIdx))
+		{
+			return;
+		}
 		const FVector3d& Pos = InMesh.GetVertex(VertexIdx);
 
 		FCreateSkinWeights_Closest_WorkData &WorkData = FCreateSkinWeights_Closest_WorkData::Get();
@@ -404,15 +408,8 @@ void CreateSkinWeights_DirectDistance(
 		}
 		
 		float TotalWeight = 0.0f;
-		for (int32 BoneIndex = 0; BoneIndex < TransformHierarchy.Num(); BoneIndex++)
+		EnumerateBonesForVertex(VertexIdx, [&](int32 BoneIndex)
 		{
-			// For single influences, avoid end effectors to avoid closest-distance fighting with the tips of their
-			// parent bones.
-			if (MaxInfluences == 1 && Skeleton.IsEndEffector(BoneIndex))
-			{
-				continue;
-			}
-			
 			// Normalize the distance by the diagonal size of the bbox to maintain scale invariance.
 			float Weight = Skeleton.GetDistanceToBoneFan(BoneIndex, Pos) / DiagBounds;
 
@@ -435,7 +432,7 @@ void CreateSkinWeights_DirectDistance(
 				TotalWeight = Weight;
 				WorkData.RawBoneWeights[0] = MakeTuple(static_cast<FBoneIndexType>(BoneIndex), Weight);
 			}
-		}
+		});
 
 		WorkData.NormalizeWeightsAndLimitCount(TotalWeight, InSettings.GetMaxWeightCount());
 
@@ -444,38 +441,136 @@ void CreateSkinWeights_DirectDistance(
 }
 
 
-void CreateSkinWeights_GeodesicVoxel(
+// Direct Distance skin weight implementation where vertices can bind to every bone
+void CreateSkinWeightsImpl_DirectDistance(
+	FDynamicMesh3& InMesh, TConstArrayView<SkinBinding::FBonePoseInfo> TransformHierarchy, FName ProfileName,
+	float InStiffness,
+	const AnimationCore::FBoneWeightsSettings& InSettings)
+{
+	const int32 MaxInfluences = InSettings.GetMaxWeightCount();
+	const FTransformHierarchyQuery Skeleton(TransformHierarchy);
+
+	CreateSkinWeightsImpl_DirectDistance_General(InMesh, TransformHierarchy, ProfileName, InStiffness, InSettings, Skeleton,
+	[MaxInfluences, &Skeleton, NumBones = TransformHierarchy.Num()](int32 VertexIndex, TFunctionRef<void(int32)> ProcessBone)
+	{
+		for (int32 BoneIndex = 0; BoneIndex < NumBones; ++BoneIndex)
+		{
+			// For single influences, avoid end effectors to avoid closest-distance fighting with the tips of their
+			// parent bones.
+			if (MaxInfluences == 1 && Skeleton.IsEndEffector(BoneIndex))
+			{
+				continue;
+			}
+
+			ProcessBone(BoneIndex);
+		}
+	});
+}
+
+// helper to filter the relevant bones when using the Constrained/bone-groups binding methods --
+// When MaxInfluences == 1, filter end effectors from bone groups that also contain non-end-effector bones.
+// OutFilteredGroups is populated only when filtering is needed; the returned view points into either
+// OutFilteredGroups or the original BoneGroups.
+static TConstArrayView<TArray<int32, TInlineAllocator<16>>> FilterBoneGroupEndEffectors(
+	int32 MaxInfluences,
+	TConstArrayView<TArray<int32, TInlineAllocator<16>>> BoneGroups,
+	const FTransformHierarchyQuery& Skeleton,
+	TArray<TArray<int32, TInlineAllocator<16>>>& OutFilteredGroups
+)
+{
+	if (MaxInfluences == 1)
+	{
+		OutFilteredGroups.Append(BoneGroups);
+		for (TArray<int32, TInlineAllocator<16>>&BoneGroup : OutFilteredGroups)
+		{
+			if (BoneGroup.Num() > 1)
+			{
+				int32 NonEndEffectors = 0;
+				for (int32 Bone : BoneGroup)
+				{
+					NonEndEffectors += (int32)!Skeleton.IsEndEffector(Bone);
+				}
+				// if there are non-end-effector bones in this group, remove the end effectors
+				if (NonEndEffectors > 0 && NonEndEffectors < BoneGroup.Num())
+				{
+					BoneGroup.SetNum(Algo::RemoveIf(BoneGroup, [&Skeleton](int32 Bone) {return Skeleton.IsEndEffector(Bone);}));
+				}
+			}
+		}
+		return OutFilteredGroups;
+	}
+	return BoneGroups;
+}
+
+
+// Direct Distance skin weight implementation where vertices can only bind to bones in limited 'bone groups'
+void CreateSkinWeightsImpl_DirectDistance_Constrained(
+	FDynamicMesh3& InMesh, TConstArrayView<SkinBinding::FBonePoseInfo> TransformHierarchy, FName ProfileName,
+	float InStiffness,
+	const AnimationCore::FBoneWeightsSettings& InSettings,
+	TConstArrayView<TArray<int32, TInlineAllocator<16>>> BoneGroups, TConstArrayView<int32> VIDtoGroupIndex
+	)
+{
+	if (VIDtoGroupIndex.IsEmpty() || BoneGroups.IsEmpty() || !ensure(VIDtoGroupIndex.Num() == InMesh.MaxVertexID()))
+	{
+		// fall back to unconstrained
+		CreateSkinWeightsImpl_DirectDistance(InMesh, TransformHierarchy, ProfileName, InStiffness, InSettings);
+		return;
+	}
+
+	const int32 MaxInfluences = InSettings.GetMaxWeightCount();
+	const FTransformHierarchyQuery Skeleton(TransformHierarchy);
+	TArray<TArray<int32, TInlineAllocator<16>>> FilteredBoneGroups;
+	TConstArrayView<TArray<int32, TInlineAllocator<16>>> UseBoneGroups =
+		FilterBoneGroupEndEffectors(MaxInfluences, BoneGroups, Skeleton, FilteredBoneGroups);
+
+	CreateSkinWeightsImpl_DirectDistance_General(InMesh, TransformHierarchy, ProfileName, InStiffness, InSettings, Skeleton,
+	[&VIDtoGroupIndex, &UseBoneGroups](int32 VertexIndex, TFunctionRef<void(int32)> ProcessBone)
+	{
+		int32 GroupIndex = VIDtoGroupIndex[VertexIndex];
+		for (int32 BoneIndex : UseBoneGroups[GroupIndex])
+		{
+			ProcessBone(BoneIndex);
+		}
+	});
+}
+
+// Generalized helper for geodesic-voxel bone weights, with predicates to customize which bones are computed and considered per vertex.
+// Note: EnumerateBonesForVertex *must* only enumerate bones for which ShouldComputeBone returned true,
+// otherwise Phase 2 will read uninitialized weight data.
+template<typename ShouldComputeBoneType, typename EnumerateType>
+void CreateSkinWeightsImpl_GeodesicVoxel_General(
 	FDynamicMesh3& InMesh, TConstArrayView<SkinBinding::FBonePoseInfo> TransformHierarchy, FName ProfileName,
 	float InStiffness, int32 InVoxelResolution,
-	const AnimationCore::FBoneWeightsSettings& InSettings
+	const AnimationCore::FBoneWeightsSettings& InSettings,
+	const FTransformHierarchyQuery& Skeleton,
+	ShouldComputeBoneType ShouldComputeBone,
+	EnumerateType EnumerateBonesForVertex
 	)
 {
 	using namespace AnimationCore;
 
 	FDynamicMeshVertexSkinWeightsAttribute *SkinWeights = GetOrCreateSkinWeightsAttribute(InMesh, ProfileName);
 		
-	const int32 NumVertices = InMesh.VertexCount();
+	const int32 MaxVertexID = InMesh.MaxVertexID();
 	const int32 MaxInfluences = InSettings.GetMaxWeightCount();
 
 	// Use the diagonal size of the bbox to make the bone distance falloff scale invariant.
 	const float DiagonalBounds = (float)FMath::Max(InMesh.GetBounds(true).DiagonalLength(), UE_SMALL_NUMBER);
-		
-	const FTransformHierarchyQuery Skeleton(TransformHierarchy);
 
 	// This is grossly inefficient but tricky to do otherwise, since each bone distance
 	// computation is done per-thread. We could possibly solve this by chunking instead
 	// and accumulating partial results.
 	TArray<float> Weights;
-	Weights.SetNumUninitialized(NumVertices * TransformHierarchy.Num());
+	Weights.SetNumUninitialized(MaxVertexID * TransformHierarchy.Num());
 
 	const FOccupancyGrid3 Occupancy(InMesh, InVoxelResolution);
 
 	const FVector3i Dimensions = Occupancy.GetOccupancyStateGrid().GetDimensions();
-		
+
+	// Phase 1: Compute per-bone geodesic distance fields via BFS through the voxel grid.
 	ParallelFor(TransformHierarchy.Num(), [&](int32 BoneIndex) {
-		// For single influences, avoid end effectors to avoid closest-distance fighting with the tips of their
-		// parent bones.
-		if (MaxInfluences == 1 && Skeleton.IsEndEffector(BoneIndex))
+		if (!ShouldComputeBone(BoneIndex))
 		{
 			return;
 		}
@@ -550,8 +645,12 @@ void CreateSkinWeights_GeodesicVoxel(
 		// Loop through all the vertices, find the voxel each belongs to, and compute
 		// the distance from the voxel to the vertex (assuming the distance stored in the
 		// voxel is based on traversing from voxel center to voxel center).
-		for (int32 VertexIdx = 0; VertexIdx < NumVertices; VertexIdx++)
+		for (int32 VertexIdx = 0; VertexIdx < MaxVertexID; VertexIdx++)
 		{
+			if (!InMesh.IsVertex(VertexIdx))
+			{
+				continue;
+			}
 			const FVector3d& Pos = InMesh.GetVertex(VertexIdx);
 			const FVector3i CellIndex = Occupancy.GetCellIndexFromPoint(Pos);
 			const FVector3d CellCenter{Occupancy.GetCellCenterFromIndex(CellIndex)};
@@ -577,8 +676,13 @@ void CreateSkinWeights_GeodesicVoxel(
 		}
 	});	
 
-	ParallelFor(NumVertices, [&](const int32 VertexIdx)
+	// Phase 2: Aggregate per-vertex weights across bones and normalize.
+	ParallelFor(MaxVertexID, [&](const int32 VertexIdx)
 	{
+		if (!InMesh.IsVertex(VertexIdx))
+		{
+			return;
+		}
 		FCreateSkinWeights_Closest_WorkData &WorkData = FCreateSkinWeights_Closest_WorkData::Get();
 			
 		if (MaxInfluences > 1)
@@ -591,15 +695,8 @@ void CreateSkinWeights_GeodesicVoxel(
 		}
 
 		float TotalWeight = 0.0f;
-		for (int32 BoneIndex = 0; BoneIndex < TransformHierarchy.Num(); BoneIndex++)
+		EnumerateBonesForVertex(VertexIdx, [&](int32 BoneIndex)
 		{
-			// Ignore end effectors for single influence computation, since those we not
-			// calculated.
-			if (MaxInfluences == 1 && Skeleton.IsEndEffector(BoneIndex))
-			{
-				continue;
-			}
-			
 			const float Weight = Weights[VertexIdx * TransformHierarchy.Num() + BoneIndex];
 			
 			if (MaxInfluences > 1)
@@ -612,7 +709,7 @@ void CreateSkinWeights_GeodesicVoxel(
 				TotalWeight = Weight;
 				WorkData.RawBoneWeights[0] = MakeTuple(static_cast<FBoneIndexType>(BoneIndex), Weight);
 			}
-		}
+		});
 
 		WorkData.NormalizeWeightsAndLimitCount(TotalWeight, InSettings.GetMaxWeightCount());
 
@@ -620,6 +717,87 @@ void CreateSkinWeights_GeodesicVoxel(
 	});
 }
 
+
+// Geodesic Voxel skin weight implementation where vertices can bind to every bone
+void CreateSkinWeightsImpl_GeodesicVoxel(
+	FDynamicMesh3& InMesh, TConstArrayView<SkinBinding::FBonePoseInfo> TransformHierarchy, FName ProfileName,
+	float InStiffness, int32 InVoxelResolution,
+	const AnimationCore::FBoneWeightsSettings& InSettings
+	)
+{
+	const int32 MaxInfluences = InSettings.GetMaxWeightCount();
+	const FTransformHierarchyQuery Skeleton(TransformHierarchy);
+
+	CreateSkinWeightsImpl_GeodesicVoxel_General(InMesh, TransformHierarchy, ProfileName, InStiffness, InVoxelResolution, InSettings, Skeleton,
+		// ShouldComputeBone: skip end effectors for single influence
+		[MaxInfluences, &Skeleton](int32 BoneIndex) -> bool
+		{
+			return !(MaxInfluences == 1 && Skeleton.IsEndEffector(BoneIndex));
+		},
+		// EnumerateBonesForVertex: iterate all bones, skipping end effectors for single influence
+		[MaxInfluences, &Skeleton, NumBones = TransformHierarchy.Num()](int32 VertexIndex, TFunctionRef<void(int32)> ProcessBone)
+		{
+			for (int32 BoneIndex = 0; BoneIndex < NumBones; ++BoneIndex)
+			{
+				if (MaxInfluences == 1 && Skeleton.IsEndEffector(BoneIndex))
+				{
+					continue;
+				}
+				ProcessBone(BoneIndex);
+			}
+		});
+}
+
+
+// Geodesic Voxel skin weight implementation where vertices can only bind to bones in limited 'bone groups'
+void CreateSkinWeightsImpl_GeodesicVoxel_Constrained(
+	FDynamicMesh3& InMesh, TConstArrayView<SkinBinding::FBonePoseInfo> TransformHierarchy, FName ProfileName,
+	float InStiffness, int32 InVoxelResolution,
+	const AnimationCore::FBoneWeightsSettings& InSettings,
+	TConstArrayView<TArray<int32, TInlineAllocator<16>>> BoneGroups, TConstArrayView<int32> VIDtoGroupIndex
+	)
+{
+	if (VIDtoGroupIndex.IsEmpty() || BoneGroups.IsEmpty() || !ensure(VIDtoGroupIndex.Num() == InMesh.MaxVertexID()))
+	{
+		// fall back to unconstrained
+		CreateSkinWeightsImpl_GeodesicVoxel(InMesh, TransformHierarchy, ProfileName, InStiffness, InVoxelResolution, InSettings);
+		return;
+	}
+
+	const int32 MaxInfluences = InSettings.GetMaxWeightCount();
+	const FTransformHierarchyQuery Skeleton(TransformHierarchy);
+	TArray<TArray<int32, TInlineAllocator<16>>> FilteredBoneGroups;
+	TConstArrayView<TArray<int32, TInlineAllocator<16>>> UseBoneGroups =
+		FilterBoneGroupEndEffectors(MaxInfluences, BoneGroups, Skeleton, FilteredBoneGroups);
+
+	// Track whether bones appear in any group, so Phase 1 can skip unused bones
+	TArray<bool> ActiveBones;
+	ActiveBones.SetNumZeroed(TransformHierarchy.Num());
+	for (const TArray<int32, TInlineAllocator<16>>& Group : UseBoneGroups)
+	{
+		for (int32 Bone : Group)
+		{
+			ActiveBones[Bone] = true;
+		}
+	}
+
+	CreateSkinWeightsImpl_GeodesicVoxel_General(InMesh, TransformHierarchy, ProfileName, InStiffness, InVoxelResolution, InSettings, Skeleton,
+		// ShouldComputeBone: only compute bones that appear in at least one group
+		[&ActiveBones](int32 BoneIndex) -> bool
+		{
+			return ActiveBones[BoneIndex];
+		},
+		// EnumerateBonesForVertex: only iterate bones in the vertex's assigned group
+		[&VIDtoGroupIndex, &UseBoneGroups](int32 VertexIndex, TFunctionRef<void(int32)> ProcessBone)
+		{
+			int32 GroupIndex = VIDtoGroupIndex[VertexIndex];
+			for (int32 BoneIndex : UseBoneGroups[GroupIndex])
+			{
+				ProcessBone(BoneIndex);
+			}
+		});
+}
+
 }
 
 namespace UE::Geometry::SkinBinding
@@ -640,11 +818,39 @@ void CreateSkinWeights(FDynamicMesh3& InMesh, TConstArrayView<SkinBinding::FBone
 		ensureMsgf(false, TEXT("Unknown BindType in CreateSkinWeights"));
 		// fall through to DirectDistance
 	case ESkinBindingType::DirectDistance:
-		CreateSkinWeights_DirectDistance(InMesh, TransformHierarchy, ProfileName, ClampedStiffness, BoneSettings);
+		CreateSkinWeightsImpl_DirectDistance(InMesh, TransformHierarchy, ProfileName, ClampedStiffness, BoneSettings);
 		break;
 			
 	case ESkinBindingType::GeodesicVoxel:
-		CreateSkinWeights_GeodesicVoxel(InMesh, TransformHierarchy, ProfileName, ClampedStiffness, BindSettings.VoxelResolution, BoneSettings);
+		CreateSkinWeightsImpl_GeodesicVoxel(InMesh, TransformHierarchy, ProfileName, ClampedStiffness, BindSettings.VoxelResolution, BoneSettings);
+		break;
+	}
+
+	CopyBonesToMesh(InMesh, TransformHierarchy);
+}
+
+void CreateConstrainedSkinWeights(FDynamicMesh3& InMesh, TConstArrayView<FBonePoseInfo> TransformHierarchy, FName ProfileName, const FBindSettings& BindSettings,
+	TConstArrayView<TArray<int32, TInlineAllocator<16>>> BoneGroups, TConstArrayView<int32> VIDtoGroupIndex)
+{
+	using namespace UE::Geometry::SkinWeightBindingLocals;
+
+	const float ClampedStiffness = FMath::Clamp(BindSettings.Stiffness, 0.0f, 1.0f);
+
+	UE::AnimationCore::FBoneWeightsSettings BoneSettings;
+	BoneSettings.SetMaxWeightCount(BindSettings.MaxInfluences);
+
+	switch(BindSettings.BindType)
+	{
+	default:
+		ensureMsgf(false, TEXT("Unknown BindType in CreateConstrainedSkinWeights"));
+		// fall through to DirectDistance
+	case ESkinBindingType::DirectDistance:
+		CreateSkinWeightsImpl_DirectDistance_Constrained(
+			InMesh, TransformHierarchy, ProfileName, ClampedStiffness, BoneSettings, BoneGroups, VIDtoGroupIndex);
+		break;
+	case ESkinBindingType::GeodesicVoxel:
+		CreateSkinWeightsImpl_GeodesicVoxel_Constrained(
+			InMesh, TransformHierarchy, ProfileName, ClampedStiffness, BindSettings.VoxelResolution, BoneSettings, BoneGroups, VIDtoGroupIndex);
 		break;
 	}
 
diff --git a/Engine/Plugins/Runtime/GeometryProcessing/Source/DynamicMesh/Public/Operations/SkinWeightBinding.h b/Engine/Plugins/Runtime/GeometryProcessing/Source/DynamicMesh/Public/Operations/SkinWeightBinding.h
--- a/Engine/Plugins/Runtime/GeometryProcessing/Source/DynamicMesh/Public/Operations/SkinWeightBinding.h
+++ b/Engine/Plugins/Runtime/GeometryProcessing/Source/DynamicMesh/Public/Operations/SkinWeightBinding.h
@@ -51,6 +51,9 @@ namespace SkinBinding
 	// Create skin weights from the input bone hierarchy, and set them on the dynamic mesh
 	UE_API void CreateSkinWeights(FDynamicMesh3& InMesh, TConstArrayView<FBonePoseInfo> TransformHierarchy, FName ProfileName, const FBindSettings& Settings);
 
+	// Create skin weights from the input bone hierarchy, with vertices constrained to only allow weights on specific subsets of bones
+	UE_API void CreateConstrainedSkinWeights(FDynamicMesh3& InMesh, TConstArrayView<FBonePoseInfo> TransformHierarchy, FName ProfileName, const FBindSettings& Settings,
+		TConstArrayView<TArray<int32, TInlineAllocator<16>>> BoneGroups, TConstArrayView<int32> VIDtoGroupIndex);
 
 };
 
