Add optional pass to the skeleton simplifier where we split edges of some triangles to add more degrees of freedom, especially to help collapse long skinny triangles into two edges.

Adjust collapse queue to avoid collapsing the useful degrees of freedom away too early.

#Synced-CL 52243346

diff --git a/Engine/Plugins/Runtime/GeometryProcessing/Source/GeometryAlgorithms/Private/Skeletonization/MeshMedialAxisSampling.cpp b/Engine/Plugins/Runtime/GeometryProcessing/Source/GeometryAlgorithms/Private/Skeletonization/MeshMedialAxisSampling.cpp
--- a/Engine/Plugins/Runtime/GeometryProcessing/Source/GeometryAlgorithms/Private/Skeletonization/MeshMedialAxisSampling.cpp
+++ b/Engine/Plugins/Runtime/GeometryProcessing/Source/GeometryAlgorithms/Private/Skeletonization/MeshMedialAxisSampling.cpp
@@ -142,6 +142,196 @@ namespace MeshMedialAxisSamplingLocals
 		);
 	}
 
+	// Info about a split edge: original edge A-B was split into A-I and I-B
+	struct FSplitEdgeInfo
+	{
+		FSplitEdgeInfo(int32 A, int32 B, int32 I) : A(A), B(B), I(I) {}
+		int32 A, B, I;
+		TArray<int32, TInlineAllocator<2>> OppositeVerts; // vertices opposite the source AB edge
+	};
+
+	// Helper to split skinny triangles -- to give the simplifier more DOFs to collapse away surfaces
+	template<typename MeshType>
+	void SplitThinSkeletonTriEdges(
+		UE::Geometry::MedialAxis::FMedialSkeleton& Skeleton,
+		TArray<FSplitEdgeInfo>& OutSplitEdges,
+		const MeshType& Mesh,
+		TConstArrayView<FSQEM> VtxSQEMs,
+		TConstArrayView<float> VtxAreas,
+		double AngleThresholdDeg,
+		double PosErrorWt)
+	{
+		using namespace UE::Geometry;
+
+		const double UseCosAngleThreshold = FMath::Cos(FMath::Clamp(AngleThresholdDeg, 90. + FMathd::ZeroTolerance, 180. - FMathd::ZeroTolerance) * FMathd::DegToRad);
+
+		// Build edge to opposite vertex map for all tris
+		TMap<FIndex2i, TArray<int32>> EdgeToOppVerts;
+		for (const FIndex3i& Tri : Skeleton.ClusterTriangles)
+		{
+			for (int32 Idx = 0, Prev = 2, Next = 1; Idx < 3; Next = Prev, Prev = Idx++)
+			{
+				EdgeToOppVerts.FindOrAdd(AsEdge(Tri[Next], Tri[Prev])).AddUnique(Tri[Idx]);
+			}
+		}
+
+		// helper to replace a target edge-opposite vertex with a new vertex
+		// no-op if the edge was already removed
+		auto ReplaceEdgeOppVert = [&EdgeToOppVerts](FIndex2i Edge, int32 OldV, int32 NewV)
+		{
+			if (TArray<int32>* OppVerts = EdgeToOppVerts.Find(Edge))
+			{
+				int32 OldVIdx = OppVerts->Find(OldV);
+				if (ensure(OldVIdx != INDEX_NONE))
+				{
+					(*OppVerts)[OldVIdx] = NewV;
+				}
+			}
+		};
+
+		// Helpers to add/remove edges from Skeleton data
+		auto AddEdge = [&Skeleton](int32 A, int32 B)
+		{
+			Skeleton.ClusterNeighbors[A].AddUnique(B);
+			Skeleton.ClusterNeighbors[B].AddUnique(A);
+		};
+
+		// SplitEdge book-keeping data -- used in a post-process to fix up skeleton data
+		// The skeleton triangle changes to apply after
+		TSet<FIndex3i> TrianglesToRemove;
+		TArray<FIndex3i> TrianglesToAdd;
+		// The clusters that may need local vertex reassignments after
+		TSet<int32> ClusterCanRemapVertices;
+
+		// Helper to split an edge by inserting an (already-created) new cluster, I, and do associated book-keeping
+		auto SplitEdge = [&OutSplitEdges, &EdgeToOppVerts, &TrianglesToRemove, &TrianglesToAdd, &ClusterCanRemapVertices, &AddEdge, &ReplaceEdgeOppVert, &Skeleton](FIndex2i Edge, int32 I)
+		{
+			int32 A = Edge.A, B = Edge.B;
+
+			// Replace every triangle (A,B,Cx) with (A,I,Cx) and (B,I,Cx)
+			const TArray<int32>& OppVerts = EdgeToOppVerts[Edge];
+			FSplitEdgeInfo& SplitInfo = OutSplitEdges.Emplace_GetRef(A, B, I);
+			SplitInfo.OppositeVerts.Append(OppVerts);
+			for (int32 Cx : OppVerts)
+			{
+				TrianglesToRemove.Add(AsTri(A, B, Cx));
+				TrianglesToAdd.Add(AsTri(A, I, Cx));
+				TrianglesToAdd.Add(AsTri(B, I, Cx));
+				AddEdge(Cx, I);
+
+				ReplaceEdgeOppVert(AsEdge(A, Cx), B, I);
+				ReplaceEdgeOppVert(AsEdge(B, Cx), A, I);
+			}
+
+			// Fix Skeleton.ClusterNeighbors w/ new edges
+			AddEdge(A, I);
+			AddEdge(I, B);
+			Skeleton.ClusterNeighbors[A].Remove(B);
+			Skeleton.ClusterNeighbors[B].Remove(A);
+
+			// Mark that clusters may need to re-assign vertices to new cluster(s)
+			ClusterCanRemapVertices.Add(A);
+			ClusterCanRemapVertices.Add(B);
+			ClusterCanRemapVertices.Append(OppVerts);
+
+			EdgeToOppVerts.Remove(Edge); // we only split edges once, can clear the stale ones
+			// (and likewise don't need AI, IB edges in EdgeToOppVerts)
+		};
+
+		// Get candidates edges w/ attached skinny triangles to split
+		struct FCandidateSplit
+		{
+			FIndex2i Edge;
+			double CosAngle, ProjT;
+		};
+		TArray<FCandidateSplit> Candidates;
+		for (const auto& Entry : EdgeToOppVerts)
+		{
+			const FIndex2i& Edge = Entry.Key;
+			const TArray<int32>& OppVerts = Entry.Value;
+			const FVector3d& PA = Skeleton.Spheres[Edge.A].Center;
+			const FVector3d& PB = Skeleton.Spheres[Edge.B].Center;
+			const FVector3d AB = PB - PA;
+			const double ABLenSq = AB.SquaredLength();
+			if (ABLenSq < FMathd::ZeroTolerance)
+			{
+				continue;
+			}
+			double MinCosAngle = UseCosAngleThreshold;
+			int32 MinCosOppV = INDEX_NONE;
+			for (int32 Cx : OppVerts)
+			{
+				const FVector3d& PC = Skeleton.Spheres[Cx].Center;
+				const double CosAngleAtCx = Normalized(PA - PC).Dot(Normalized(PB - PC));
+				if (CosAngleAtCx < MinCosAngle)
+				{
+					MinCosAngle = CosAngleAtCx;
+					MinCosOppV = Cx;
+				}
+			}
+			if (MinCosOppV != INDEX_NONE)
+			{
+				double ProjT = (Skeleton.Spheres[MinCosOppV].Center - PA).Dot(AB) / ABLenSq;
+				if (ProjT > FMathd::ZeroTolerance && ProjT < 1. - FMathd::ZeroTolerance)
+				{
+					Candidates.Add({ Edge, MinCosAngle, ProjT });
+				}
+			}
+		}
+
+		// Process largest-angle (smallest cosine) triangles first, since split order can affect results
+		Candidates.Sort([](const FCandidateSplit& A, const FCandidateSplit& B) { return A.CosAngle < B.CosAngle; });
+
+		const int32 StartNewClusters = Skeleton.Spheres.Num();
+		// Apply the edges splits
+		for (const FCandidateSplit& Cand : Candidates)
+		{
+			const int32 A = Cand.Edge.A, B = Cand.Edge.B;
+
+			// New cluster I, projected onto AB
+			const int32 I = Skeleton.Spheres.Add(LerpSphere(Skeleton.Spheres[A], Skeleton.Spheres[B], Cand.ProjT));
+			Skeleton.ClusterNeighbors.AddDefaulted();
+
+			SplitEdge(Cand.Edge, I);
+		}
+
+		// Error metric that produces equivalent rankings to ComputeSkeleton's vertex-cluster assignment method, using PosErrorWt
+		// Note: area-scaled (not normalized), so only valid for comparing across clusters for the same vertex
+		auto ClusterVertexError = [&Mesh, &VtxAreas, &VtxSQEMs, PosErrorWt](int32 VID, const FSphere& S) -> double
+		{
+			const double PosErr = (Mesh.GetVertex(VID) - S.Center).Length() - S.W;
+			return PosErrorWt * VtxAreas[VID] * PosErr * PosErr + (1. - PosErrorWt) * VtxSQEMs[VID].Eval(S);
+		};
+		// Pass through vertices and remap cluster assignments to new vertices, if they are on a marked-as-can-remap cluster
+		ParallelFor(Mesh.MaxVertexID(), [&Skeleton, &ClusterCanRemapVertices, &ClusterVertexError, StartNewClusters](int32 VID)
+		{
+			const int32 CurCluster = Skeleton.VIDtoClusterIndex[VID];
+			if (ClusterCanRemapVertices.Contains(CurCluster))
+			{
+				double BestClusterErr = ClusterVertexError(VID, Skeleton.Spheres[CurCluster]);
+				int32 BestCluster = CurCluster;
+				for (int32 NbrCluster : Skeleton.ClusterNeighbors[CurCluster])
+				{
+					// Only consider remapping to new clusters (assume vertex was already mapped to the 'best' cluster of previously-existing ones)
+					if (NbrCluster >= StartNewClusters)
+					{
+						double NbrClusterErr = ClusterVertexError(VID, Skeleton.Spheres[NbrCluster]);
+						if (NbrClusterErr < BestClusterErr)
+						{
+							BestClusterErr = NbrClusterErr;
+							BestCluster = NbrCluster;
+						}
+					}
+				}
+				Skeleton.VIDtoClusterIndex[VID] = BestCluster;
+			}
+		});
+
+		// Apply accumulated triangle modifications
+		Skeleton.ClusterTriangles.SetNum(Algo::RemoveIf(Skeleton.ClusterTriangles,
+			[&TrianglesToRemove](const FIndex3i& Tri) { return TrianglesToRemove.Contains(Tri); }), EAllowShrinking::No);
+		Skeleton.ClusterTriangles.Append(TrianglesToAdd);
+	}
 
 	// Build up standard mesh QEM error terms for the skeleton geometry
 	// Uses cluster areas rather than the skeleton geometry for weighting
@@ -325,7 +515,15 @@ namespace MeshMedialAxisSamplingLocals
 
 		const MeshType& Mesh = *MeshBVH.GetMesh();
 
-		const int32 InitialNumClusters = Skeleton.Spheres.Num();
+		// Optionally pass through skinny triangles and split them to add extra degrees of freedom for simplification
+		TArray<FSplitEdgeInfo> SplitEdges;
+		if (Simplifier.bSplitThinTriEdges)
+		{
+			const double UsePosErrorWt = FMath::Clamp(Simplifier.SplitThinTriEdgePosErrorWt, UE_DOUBLE_KINDA_SMALL_NUMBER, 1);
+			SplitThinSkeletonTriEdges(Skeleton, SplitEdges, Mesh, VtxSQEMs, VtxAreas, Simplifier.SplitThinTriEdgeAngleThresholdDeg, UsePosErrorWt);
+		}
+
+		const int32 InitialNumClusters = Skeleton.Spheres.Num(); // includes clusters added by SplitThinSkeletonTriEdges
 
 		using FNbrTriArray = TArray<FIndex2i, TInlineAllocator<8>>;
 		TArray<FNbrTriArray> NbrTris;
@@ -489,6 +687,14 @@ namespace MeshMedialAxisSamplingLocals
 			FIndex2i Edge, Version;
 			FSphere MergeSphere;
 			float Error;
+
+			bool operator<(const FEdgeInfo& Other) const
+			{
+				return Error < Other.Error;
+			}
+
+			FEdgeInfo(FIndex2i InEdge = FIndex2i::Invalid())
+				: Edge(InEdge), Version(0, 0) {}
 		};
 
 		TArray<FEdgeInfo> EdgeQ;
@@ -505,6 +711,26 @@ namespace MeshMedialAxisSamplingLocals
 			return (float)(SQEM.Eval(OutSphere) - BaseError);
 		};
 
+		// Helper to add an edge to the EdgeQ, if it passes the error threshold; returns true if added
+		auto AddEdgeToQueue = [&Simplifier, &EdgeQ, &ComputeMergeError](const FIndex2i& Edge, bool bHeapPush = false, float ErrorFloor = -FMathf::MaxReal) -> bool
+		{
+			FEdgeInfo Info(Edge);
+			Info.Error = FMath::Max(ComputeMergeError(Info.Edge, Info.MergeSphere), ErrorFloor);
+			if (Simplifier.QEMErrorThreshold <= 0 || Info.Error <= Simplifier.QEMErrorThreshold)
+			{
+				if (bHeapPush)
+				{
+					EdgeQ.HeapPush(Info);
+				}
+				else
+				{
+					EdgeQ.Add(Info);
+				}
+				return true;
+			}
+			return false;
+		};
+
 		double UseEdgeLengthThresholdSq = Simplifier.EdgeLengthThreshold * Simplifier.EdgeLengthThreshold;
 		auto ShouldSkipEdge = [&Simplifier, &Skeleton, UseEdgeLengthThresholdSq, &EdgeHasTris, &CollapseInducedDeadEnds](const FIndex2i& Edge) -> bool
 		{
@@ -527,6 +753,22 @@ namespace MeshMedialAxisSamplingLocals
 			return false;
 		};
 
+		// Pre-compute error floors for split sub-edges: for each sub-edge (AI, IB) of edge AB,
+		// inflate priority to Min(CollapseError(AB), min_C(CollapseError(IC)) + Epsilon)
+		// so the sub-edges aren't considered 'too soon' before IC would be considered
+		TMap<FIndex2i, float> SplitSubEdgeErrorFloors;
+		for (const FSplitEdgeInfo& Split : SplitEdges)
+		{
+			FSphere TmpSphere;
+			float ErrorFloor = ComputeMergeError(FIndex2i(Split.A, Split.B), TmpSphere);
+			for (int32 C : Split.OppositeVerts)
+			{
+				ErrorFloor = FMath::Min(ErrorFloor, ComputeMergeError(AsEdge(Split.I, C), TmpSphere) + UE_KINDA_SMALL_NUMBER);
+			}
+			SplitSubEdgeErrorFloors.Add(AsEdge(Split.A, Split.I), ErrorFloor);
+			SplitSubEdgeErrorFloors.Add(AsEdge(Split.B, Split.I), ErrorFloor);
+		}
+
 		for (int32 ClusterIdx = 0; ClusterIdx < InitialNumClusters; ++ClusterIdx)
 		{
 			for (int32 NbrIdx : Skeleton.ClusterNeighbors[ClusterIdx])
@@ -540,20 +782,15 @@ namespace MeshMedialAxisSamplingLocals
 				{
 					continue;
 				}
-				FEdgeInfo Info;
-				Info.Edge = Edge;
-				Info.Version = FIndex2i(0, 0);
-				Info.Error = ComputeMergeError(Info.Edge, Info.MergeSphere);
-				if (Simplifier.QEMErrorThreshold <= 0 || Info.Error <= Simplifier.QEMErrorThreshold)
+				float UseErrorFloor = -FMathf::MaxReal;
+				if (const float* ErrorFloor = SplitSubEdgeErrorFloors.Find(Edge))
 				{
-					EdgeQ.Add(Info);
+					UseErrorFloor = *ErrorFloor;
 				}
+				AddEdgeToQueue(Edge, false, UseErrorFloor);
 			}
 		}
-		auto EdgeSort = [](const FEdgeInfo& A, const FEdgeInfo& B)
-		{
-			return A.Error < B.Error;
-		};
+
 		auto UpdateClusterIndexFromMerge = [&ClusterMergeParent](int32& Parent) -> void
 		{
 			while (ClusterMergeParent[Parent] != Parent)
@@ -563,85 +800,115 @@ namespace MeshMedialAxisSamplingLocals
 			}
 		};
 
-		EdgeQ.Heapify(EdgeSort);
+		EdgeQ.Heapify();
 		int32 NumClusters = InitialNumClusters;
-		while (!EdgeQ.IsEmpty() && NumClusters > Simplifier.MinSpheres)
-		{
-			FEdgeInfo EdgeInfo;
-			EdgeQ.HeapPop(EdgeInfo, EdgeSort, EAllowShrinking::No);
-
-			FIndex2i& Edge = EdgeInfo.Edge;
 
-			FIndex2i& Version = EdgeInfo.Version;
-			if (ClusterVersion[Edge.A] > Version.A || ClusterVersion[Edge.B] > Version.B)
+		auto ProcessEdgeQueue = [&EdgeQ, &NumClusters, &Simplifier, &ClusterVersion, &ClusterMergeParent, &ShouldSkipEdge, &UpdateClusterIndexFromMerge, &ComputeMergeError, &Skeleton, &MeshBVH, &CollapseEdge]()
+		{
+			while (!EdgeQ.IsEmpty() && NumClusters > Simplifier.MinSpheres)
 			{
-				UpdateClusterIndexFromMerge(Edge.A);
-				UpdateClusterIndexFromMerge(Edge.B);
-				if (Edge.A == Edge.B)
-				{
-					// edge was merged away, skip
-					continue;
-				}
-				if (ShouldSkipEdge(Edge))
-				{
-					continue;
-				}
-				Version.A = ClusterVersion[Edge.A];
-				Version.B = ClusterVersion[Edge.B];
-				EdgeInfo.Error = ComputeMergeError(Edge, EdgeInfo.MergeSphere);
-				if (Simplifier.QEMErrorThreshold > 0 && EdgeInfo.Error > Simplifier.QEMErrorThreshold)
-				{
-					continue;
-				}
+				FEdgeInfo EdgeInfo;
+				EdgeQ.HeapPop(EdgeInfo, EAllowShrinking::No);
 
-				if (!EdgeQ.IsEmpty() && EdgeQ.HeapTop().Error < EdgeInfo.Error)
+				FIndex2i& Edge = EdgeInfo.Edge;
+				FIndex2i& Version = EdgeInfo.Version;
+				const bool bIsStale = (ClusterVersion[Edge.A] > Version.A || ClusterVersion[Edge.B] > Version.B);
+				if (bIsStale)
 				{
-					// after updating error, there's a cheaper edge on top of heap --
-					// need just re-insert this one and continue to pick up the smaller edge instead
-					EdgeQ.HeapPush(EdgeInfo, EdgeSort);
+					UpdateClusterIndexFromMerge(Edge.A);
+					UpdateClusterIndexFromMerge(Edge.B);
+					if (Edge.A == Edge.B)
+					{
+						// edge was merged away, skip
+						continue;
+					}
+					if (ShouldSkipEdge(Edge))
+					{
+						continue;
+					}
+					Version.A = ClusterVersion[Edge.A];
+					Version.B = ClusterVersion[Edge.B];
+					EdgeInfo.Error = ComputeMergeError(Edge, EdgeInfo.MergeSphere);
+					if (Simplifier.QEMErrorThreshold > 0 && EdgeInfo.Error > Simplifier.QEMErrorThreshold)
+					{
+						continue;
+					}
 
-					continue;
+					if (!EdgeQ.IsEmpty() && EdgeQ.HeapTop().Error < EdgeInfo.Error)
+					{
+						// after updating error, there's a cheaper edge on top of heap --
+						// need just re-insert this one and continue to pick up the smaller edge instead
+						EdgeQ.HeapPush(EdgeInfo);
+						continue;
+					}
 				}
-			}
-			else
-			{
-				if (ShouldSkipEdge(Edge))
+				else
 				{
-					continue;
+					if (ShouldSkipEdge(Edge))
+					{
+						continue;
+					}
 				}
-			}
 
-			if (Simplifier.bPreventEdgeSurfaceIntersections)
-			{
-				bool bFoundHit = false;
-				for (int32 SubIdx = 0; !bFoundHit && SubIdx < 2; ++SubIdx)
+				if (Simplifier.bPreventEdgeSurfaceIntersections)
 				{
-					int32 SrcClusterIdx = Edge[SubIdx];
-					int32 OtherIdx = Edge[1 - SubIdx];
-					FVector3d SrcPos = EdgeInfo.MergeSphere.Center;
-					for (int32 ToClusterIdx : Skeleton.ClusterNeighbors[SrcClusterIdx])
+					bool bFoundHit = false;
+					for (int32 SubIdx = 0; !bFoundHit && SubIdx < 2; ++SubIdx)
 					{
-						if (ToClusterIdx != OtherIdx)
+						int32 SrcClusterIdx = Edge[SubIdx];
+						int32 OtherIdx = Edge[1 - SubIdx];
+						FVector3d SrcPos = EdgeInfo.MergeSphere.Center;
+						for (int32 ToClusterIdx : Skeleton.ClusterNeighbors[SrcClusterIdx])
 						{
-							FVector3d DestPos = Skeleton.Spheres[ToClusterIdx].Center;
-							if (EdgeIntersectsMesh(SrcPos, DestPos, MeshBVH))
+							if (ToClusterIdx != OtherIdx)
 							{
-								bFoundHit = true;
-								break;
+								FVector3d DestPos = Skeleton.Spheres[ToClusterIdx].Center;
+								if (EdgeIntersectsMesh(SrcPos, DestPos, MeshBVH))
+								{
+									bFoundHit = true;
+									break;
+								}
 							}
-
 						}
 					}
+					if (bFoundHit)
+					{
+						continue;
+					}
 				}
-				if (bFoundHit)
+
+				CollapseEdge(EdgeInfo.Edge, EdgeInfo.MergeSphere);
+				NumClusters--;
+			}
+		};
+		ProcessEdgeQueue();
+
+		// If we're still above target sphere count, re-add split sub-edges with un-clamped errors for a cleanup pass
+		if (NumClusters > Simplifier.MinSpheres)
+		{
+			bool bAddedEdges = false;
+			for (const FSplitEdgeInfo& Split : SplitEdges)
+			{
+				for (int32 EndV : {Split.A, Split.B})
 				{
-					continue;
+					int32 InterV = Split.I;
+					UpdateClusterIndexFromMerge(EndV);
+					UpdateClusterIndexFromMerge(InterV);
+					if (EndV == InterV)
+					{
+						continue;
+					}
+					FIndex2i EdgeEI = AsEdge(EndV, InterV);
+					if (!ShouldSkipEdge(EdgeEI))
+					{
+						bAddedEdges |= AddEdgeToQueue(EdgeEI, true);
+					}
 				}
 			}
-
-			// we're up-to-date and have the lowest error, so can merge
-			CollapseEdge(EdgeInfo.Edge, EdgeInfo.MergeSphere);
-			NumClusters--;
+			if (bAddedEdges)
+			{
+				ProcessEdgeQueue();
+			}
 		}
 
 		// Do a final pass to compact away the removed cluster spheres/info and remap indices accordingly
diff --git a/Engine/Plugins/Runtime/GeometryProcessing/Source/GeometryAlgorithms/Public/Skeletonization/MeshMedialAxisSampling.h b/Engine/Plugins/Runtime/GeometryProcessing/Source/GeometryAlgorithms/Public/Skeletonization/MeshMedialAxisSampling.h
--- a/Engine/Plugins/Runtime/GeometryProcessing/Source/GeometryAlgorithms/Public/Skeletonization/MeshMedialAxisSampling.h
+++ b/Engine/Plugins/Runtime/GeometryProcessing/Source/GeometryAlgorithms/Public/Skeletonization/MeshMedialAxisSampling.h
@@ -439,12 +439,19 @@ namespace MedialAxis
 		// Regularization weight encouraging skeleton clusters to keep their original positions and radii.
 		double ClusterRegularizeWeight = FMathd::ZeroTolerance;
 
-		
+		// Whether to split edges attached to skinny triangles, as a pre-process before simplifying, which can help add useful degrees of freedom to the simplifier
+		bool bSplitThinTriEdges = true;
+		// Minimum angle of triangle opposite an edge at which to potentially split the edge
+		double SplitThinTriEdgeAngleThresholdDeg = 120;
+		// When re-clustering for split edges, weight to use for the position error term. Equivalent in meaning to PosErrorWt in FSkeletonViaSampling.
+		double SplitThinTriEdgePosErrorWt = .2;
+
+
 		template<typename MeshType>
 		GEOMETRYALGORITHMS_API void Simplify(FMedialSkeleton& Skeleton, const TMeshAABBTree3<MeshType>& MeshBVH, TFunctionRef<FVector3d(int32)> GetVertexNormal, TFunctionRef<int32(int32)> VIDtoWeldVIDFn);
 
 		template<typename MeshType>
-		GEOMETRYALGORITHMS_API void Simplify(FMedialSkeleton& Skeleton, const TMeshAABBTree3<MeshType>& MeshBVH, TFunctionRef<bool(int32)> OptionalBoundaryVertexFilter, double ImplicitVertexNormalWeldThreshold);
+		GEOMETRYALGORITHMS_API void Simplify(FMedialSkeleton& Skeleton, const TMeshAABBTree3<MeshType>& MeshBVH, TFunctionRef<bool(int32)> OptionalBoundaryVertexFilter = [](int32)->bool { return false; }, double ImplicitVertexNormalWeldThreshold = UE_DOUBLE_KINDA_SMALL_NUMBER);
 	};
 
 
