Add automatic root selection method to the medial->anim skeleton conversion

#Synced-CL 52352248

diff --git a/Engine/Plugins/Runtime/GeometryProcessing/Source/GeometryAlgorithms/Private/Skeletonization/MeshMedialAxisSampling.cpp b/Engine/Plugins/Runtime/GeometryProcessing/Source/GeometryAlgorithms/Private/Skeletonization/MeshMedialAxisSampling.cpp
--- a/Engine/Plugins/Runtime/GeometryProcessing/Source/GeometryAlgorithms/Private/Skeletonization/MeshMedialAxisSampling.cpp
+++ b/Engine/Plugins/Runtime/GeometryProcessing/Source/GeometryAlgorithms/Private/Skeletonization/MeshMedialAxisSampling.cpp
@@ -5,9 +5,12 @@
 #include <atomic>
 
 #include "Algo/RemoveIf.h"
+#include "Algo/MaxElement.h"
+#include "Algo/MinElement.h"
 #include "VertexConnectedComponents.h"
 #include "MeshAdapter.h"
 #include "MeshAdapter/MeshVertexNormals.h"
+#include "Misc/CoreMiscDefines.h"
 
 #include "DynamicMesh/DynamicBoneAttribute.h"
 #include "DynamicMesh/DynamicMeshAttributeSet.h"
@@ -1933,7 +1936,45 @@ namespace UE::Geometry::MedialAxis
 			}
 		}
 
-		int32 RootIdx = InMedialSkeleton.Spheres.IsValidIndex(InRootIndex) ? InRootIndex : 0;
+		int32 RootIdx;
+		if (InMedialSkeleton.Spheres.IsValidIndex(InRootIndex))
+		{
+			RootIdx = InRootIndex;
+		}
+		else
+		{
+			// Automatically select a root cluster based on the configured method
+			const TArray<FSphere>& Spheres = InMedialSkeleton.Spheres;
+			const FSphere* Selected = nullptr;
+			switch (Options.SelectRootMethod)
+			{
+			case ESelectRootMethod::ClosestToPoint:
+				Selected = Algo::MinElementBy(Spheres, [&](const FSphere& S) { return FVector3d::DistSquared(S.Center, Options.RootSelectionPoint); });
+				break;
+			case ESelectRootMethod::FarthestInDirection:
+				Selected = Algo::MaxElementBy(Spheres, [&](const FSphere& S) { return S.Center.Dot(Options.RootSelectionDirection); });
+				break;
+			case ESelectRootMethod::ClosestToBoundsCenter:
+			{
+				FAxisAlignedBox3d SkeletonBounds = FAxisAlignedBox3d::Empty();
+				for (const FSphere& S : Spheres)
+				{
+					SkeletonBounds.Contain(S.Center + FVector(S.W));
+					SkeletonBounds.Contain(S.Center - FVector(S.W));
+				}
+				FVector3d BoundsCenter = SkeletonBounds.Center();
+				Selected = Algo::MinElementBy(Spheres, [&](const FSphere& S) { return FVector3d::DistSquared(S.Center, BoundsCenter); });
+				break;
+			}
+			case ESelectRootMethod::LargestSphere:
+				Selected = Algo::MaxElementBy(Spheres, [](const FSphere& S) { return S.W; });
+				break;
+			case ESelectRootMethod::ArrayOrder:
+			default:
+				break;
+			}
+			RootIdx = Selected ? UE_PTRDIFF_TO_INT32(Selected - Spheres.GetData()) : 0;
+		}
 
 		if (Options.MergeDisconnectedMethod == EMergeDisconnectedMethod::ConnectClosestBones)
 		{
diff --git a/Engine/Plugins/Runtime/GeometryProcessing/Source/GeometryAlgorithms/Public/Skeletonization/MeshMedialAxisSampling.h b/Engine/Plugins/Runtime/GeometryProcessing/Source/GeometryAlgorithms/Public/Skeletonization/MeshMedialAxisSampling.h
--- a/Engine/Plugins/Runtime/GeometryProcessing/Source/GeometryAlgorithms/Public/Skeletonization/MeshMedialAxisSampling.h
+++ b/Engine/Plugins/Runtime/GeometryProcessing/Source/GeometryAlgorithms/Public/Skeletonization/MeshMedialAxisSampling.h
@@ -366,8 +366,24 @@ namespace MedialAxis
 			AddTopLevelRoot
 		};
 
+		// Method to automatically select the medial cluster to use as animation skeleton root, if a root cluster index is not specified
+		enum class ESelectRootMethod
+		{
+			ClosestToPoint,
+			FarthestInDirection,
+			ClosestToBoundsCenter,
+			LargestSphere,
+			ArrayOrder
+		};
+
 		EEdgeWeightMethod EdgeWeightMethod = EEdgeWeightMethod::ArrayOrder;
 		EMergeDisconnectedMethod MergeDisconnectedMethod = EMergeDisconnectedMethod::ConnectClosestBones;
+		ESelectRootMethod SelectRootMethod = ESelectRootMethod::ArrayOrder;
+
+		// When SelectRootMethod is ClosestToPoint, the point to use
+		FVector3d RootSelectionPoint = FVector3d::ZeroVector;
+		// When SelectRootMethod is FarthestInDirection, the direction to use
+		FVector3d RootSelectionDirection = -FVector3d::UnitZ();
 
 		/**
 		 * Remap a medial skeleton to a hierarchical skeleton suitable for using in a URefSkeleton for animation
