Add medial skeleton skin binding operation, to transfer a medial skeleton + mesh to a skeletal mesh that uses the cluster->vertex mapping to guide the skin weight binding

#Synced-CL 52392421

diff --git a/Engine/Plugins/Runtime/GeometryProcessing/Source/DynamicMesh/Private/Operations/MedialSkeletonSkinBinding.cpp b/Engine/Plugins/Runtime/GeometryProcessing/Source/DynamicMesh/Private/Operations/MedialSkeletonSkinBinding.cpp
new file mode 100644
--- /dev/null
+++ b/Engine/Plugins/Runtime/GeometryProcessing/Source/DynamicMesh/Private/Operations/MedialSkeletonSkinBinding.cpp
@@ -0,0 +1,135 @@
+// Copyright Epic Games, Inc. All Rights Reserved.
+
+#include "Operations/MedialSkeletonSkinBinding.h"
+
+#include "DynamicMesh/DynamicBoneAttribute.h"
+#include "DynamicMesh/DynamicMeshAttributeSet.h"
+#include "DynamicMesh/DynamicVertexSkinWeightsAttribute.h"
+
+namespace MedialSkeletonSkinBindingLocals
+{
+
+	static TArray<TArray<int32, TInlineAllocator<16>>> BuildMedialSkeletonBoneGroups(
+		const UE::Geometry::MedialAxis::FMedialSkeleton& InMedialSkeleton,
+		TConstArrayView<int32> MedialIndexToBoneIndex,
+		int32 NumBones,
+		int32 ClusterNeighborSearchRange)
+	{
+		using namespace UE::Geometry;
+
+		const int32 NumClusters = InMedialSkeleton.ClusterNeighbors.Num();
+		TArray<TArray<int32, TInlineAllocator<16>>> BoneGroups;
+		BoneGroups.SetNum(NumClusters);
+
+		// For each cluster, BFS-expand by ClusterNeighborSearchRange steps to get the cluster group
+		for (int32 ClusterIdx = 0; ClusterIdx < NumClusters; ++ClusterIdx)
+		{
+			TArray<int32> CurrentFrontier;
+			TArray<int32> NextFrontier;
+			TSet<int32> Visited;
+
+			CurrentFrontier.Add(ClusterIdx);
+			Visited.Add(ClusterIdx);
+
+			for (int32 Step = 0; Step < ClusterNeighborSearchRange; ++Step)
+			{
+				NextFrontier.Reset();
+				for (int32 Src : CurrentFrontier)
+				{
+					for (int32 Nbr : InMedialSkeleton.ClusterNeighbors[Src])
+					{
+						bool bFound;
+						Visited.FindOrAdd(Nbr, &bFound);
+						if (!bFound)
+						{
+							NextFrontier.Add(Nbr);
+						}
+					}
+				}
+				Swap(CurrentFrontier, NextFrontier);
+			}
+
+			// Remap visited cluster indices to bone indices
+			for (int32 VisitedCluster : Visited)
+			{
+				check(MedialIndexToBoneIndex.IsValidIndex(VisitedCluster));
+				int32 BoneIdx = MedialIndexToBoneIndex[VisitedCluster];
+				if (BoneIdx >= 0 && BoneIdx < NumBones)
+				{
+					BoneGroups[ClusterIdx].Add(BoneIdx);
+				}
+			}
+		}
+
+		return BoneGroups;
+	}
+}
+
+namespace UE::Geometry::SkinBinding
+{
+
+
+	
+bool CreateSkinWeightsFromMedialSkeleton(
+	const MedialAxis::FMedialSkeleton& InMedialSkeleton,
+	FDynamicMesh3& TargetMesh,
+	bool& bOutMeshWasCompatible,
+	FName WeightProfileName,
+	const FBindMedialSkeletonSettings& Settings,
+	int32 RootIndex)
+{
+	using namespace MedialSkeletonSkinBindingLocals;
+
+	bOutMeshWasCompatible = InMedialSkeleton.IsCompatibleWithMesh(TargetMesh);
+
+	// Step 1: Convert medial skeleton to hierarchical skeleton
+	TArray<int32> Parents;
+	TArray<FTransform> Poses;
+	TArray<int32> BoneIndexToMedialIndex;
+	TArray<int32> MedialIndexToBoneIndex;
+
+	MedialAxis::FMedialSkeletonToTreeSkeletonOptions::ToHierarchy(
+		InMedialSkeleton, Parents, Poses, Settings.AnimationSkeletonOptions, RootIndex,
+		&BoneIndexToMedialIndex, &MedialIndexToBoneIndex);
+
+	const int32 NumBones = Parents.Num();
+	if (NumBones == 0)
+	{
+		return false;
+	}
+
+	// Step 2: Prepare mesh attributes — clear any (likely-incompatible) existing skin weights
+	TargetMesh.EnableAttributes();
+	TargetMesh.Attributes()->RemoveAllSkinWeightsAttributes();
+
+	// Step 3: Build FBonePoseInfo array from hierarchy
+	TArray<FBonePoseInfo> Bones;
+	Bones.SetNum(Parents.Num());
+	for (int32 BoneIdx = 0; BoneIdx < NumBones; ++BoneIdx)
+	{
+		Bones[BoneIdx].Name = FName("Bone", BoneIdx);
+		Bones[BoneIdx].ParentIndex = Parents[BoneIdx];
+		Bones[BoneIdx].LocalTransform = Poses[BoneIdx];
+	}
+
+	if (bOutMeshWasCompatible)
+	{
+		// Step 4: Build bone groups per-cluster from medial skeleton cluster graph
+		TArray<TArray<int32, TInlineAllocator<16>>> BoneGroups = BuildMedialSkeletonBoneGroups(
+			InMedialSkeleton, MedialIndexToBoneIndex, NumBones,
+			Settings.ClusterNeighborSearchRange);
+
+		// Step 5: Apply constrained skin binding
+		CreateConstrainedSkinWeights(
+			TargetMesh, Bones, WeightProfileName, Settings.BindSettings,
+			BoneGroups, InMedialSkeleton.VIDtoClusterIndex);
+	}
+	else
+	{
+		CreateSkinWeights(TargetMesh, Bones, WeightProfileName, Settings.BindSettings);
+	}
+
+	return true;
+}
+
+} // namespace UE::Geometry::SkinBinding
diff --git a/Engine/Plugins/Runtime/GeometryProcessing/Source/DynamicMesh/Public/Operations/MedialSkeletonSkinBinding.h b/Engine/Plugins/Runtime/GeometryProcessing/Source/DynamicMesh/Public/Operations/MedialSkeletonSkinBinding.h
new file mode 100644
--- /dev/null
+++ b/Engine/Plugins/Runtime/GeometryProcessing/Source/DynamicMesh/Public/Operations/MedialSkeletonSkinBinding.h
@@ -0,0 +1,49 @@
+// Copyright Epic Games, Inc. All Rights Reserved.
+
+#pragma once
+
+#include "Operations/SkinWeightBinding.h"
+#include "Skeletonization/MeshMedialAxisSampling.h"
+
+#define UE_API DYNAMICMESH_API
+
+namespace UE::Geometry::SkinBinding
+{
+
+struct FBindMedialSkeletonSettings
+{
+	// Search range in local medial skeleton to use when binding vertices --
+	// i.e. a value of 1 will restrict vertices to bind to clusters within the 1-ring of their best-fit cluster
+	int32 ClusterNeighborSearchRange = 1;
+
+	// General skin weight binding settings
+	FBindSettings BindSettings;
+
+	// Options controlling conversion to an animation skeleton
+	MedialAxis::FMedialSkeletonToTreeSkeletonOptions AnimationSkeletonOptions;
+};
+
+/**
+ * Convert a medial skeleton + mesh to a skinned mesh, using the medial skeleton clusters 
+ * to both define the animation skeleton and to help guide the skinning -- binding vertices
+ * to bones within the local graph neighborhood of their corresponding medial skeleton cluster/medial sphere
+ *
+ * @param InMedialSkeleton Source medial skeleton
+ * @param TargetMesh Mesh to receive skeleton attributes and skin weights
+ * @param bOutMeshWasCompatible Reports whether the mesh was compatible with the medial skeleton (had a valid vertex->cluster mapping). If not, mesh can still be bound but will not respect some options.
+ * @param WeightProfileName Skin weight profile name to use
+ * @param Settings Options controlling the weight binding and skeleton conversion
+ * @param RootIndex Specify the medial skeleton cluster to prefer as the animation skeleton's root bone. If INDEX_NONE, a default root will be selected automatically.
+ * @return true on success, false if could not be bound (e.g. if medial skeleton was empty)
+ */
+UE_API bool CreateSkinWeightsFromMedialSkeleton(
+	const MedialAxis::FMedialSkeleton& InMedialSkeleton,
+	FDynamicMesh3& TargetMesh,
+	bool& bOutMeshWasCompatible,
+	FName WeightProfileName,
+	const FBindMedialSkeletonSettings& Settings,
+	int32 RootIndex = INDEX_NONE);
+
+} // namespace UE::Geometry::SkinBinding
+
+#undef UE_API
