remove dynamic mesh skinning method from MeshMedialAxisSampling; it has been replaced by the better algorithm in CreateSkinWeightsFromMedialSkeleton -- which lives in a different module to be able to use the skin binding algorithms

#Synced-CL 52392718

diff --git a/Engine/Plugins/Runtime/GeometryProcessing/Source/GeometryAlgorithms/Private/Skeletonization/MeshMedialAxisSampling.cpp b/Engine/Plugins/Runtime/GeometryProcessing/Source/GeometryAlgorithms/Private/Skeletonization/MeshMedialAxisSampling.cpp
--- a/Engine/Plugins/Runtime/GeometryProcessing/Source/GeometryAlgorithms/Private/Skeletonization/MeshMedialAxisSampling.cpp
+++ b/Engine/Plugins/Runtime/GeometryProcessing/Source/GeometryAlgorithms/Private/Skeletonization/MeshMedialAxisSampling.cpp
@@ -2105,76 +2105,6 @@ namespace UE::Geometry::MedialAxis
 		}
 	}
 
-	void FMedialSkeletonToTreeSkeletonOptions::ToDynamicMeshSkeleton(const FMedialSkeleton& InMedialSkeleton, FDynamicMesh3& TargetMesh, FName WeightProfileName, FMedialSkeletonToTreeSkeletonOptions Options, int32 RootIndex)
-	{
-		TArray<int32> BoneRemap;
-		TArray<int32> Parents;
-		TArray<FTransform> Poses;
-
-		ToHierarchy(InMedialSkeleton, Parents, Poses, Options, RootIndex, nullptr, &BoneRemap);
-
-		if (Poses.IsEmpty())
-		{
-			// no skeleton to add
-			return;
-		}
-
-		// make sure attributes are enabled so we can add skeleton data
-		TargetMesh.EnableAttributes();
-
-		// transfer hierarchy to mesh bones
-		const int32 NumBones = Parents.Num();
-		FDynamicMeshAttributeSet* AttributeSet = TargetMesh.Attributes();
-		{
-			AttributeSet->EnableBones(NumBones);
-			FDynamicMeshBoneNameAttribute* BoneNameAttrib = AttributeSet->GetBoneNames();
-			FDynamicMeshBoneParentIndexAttribute* BoneParentIndices = AttributeSet->GetBoneParentIndices();
-			FDynamicMeshBonePoseAttribute* BonePoses = AttributeSet->GetBonePoses();
-			FDynamicMeshBoneColorAttribute* BoneColors = AttributeSet->GetBoneColors();
-
-			for (int32 BoneIdx = 0; BoneIdx < NumBones; ++BoneIdx)
-			{
-				FName BoneName("Bone", BoneIdx);
-				BoneNameAttrib->SetValue(BoneIdx, BoneName);
-				BoneParentIndices->SetValue(BoneIdx, Parents[BoneIdx]);
-				BonePoses->SetValue(BoneIdx, Poses[BoneIdx]);
-				BoneColors->SetValue(BoneIdx, FVector4f::One());
-			}
-		}
-
-		// Remove any existing skin weight attributes; they would generally not be compatible with the newly-created bones
-		TargetMesh.Attributes()->RemoveAllSkinWeightsAttributes();
-		// Add our new skin weight attribute
-		FDynamicMeshVertexSkinWeightsAttribute* Attribute = new FDynamicMeshVertexSkinWeightsAttribute(&TargetMesh);
-		TargetMesh.Attributes()->AttachSkinWeightsAttribute(WeightProfileName, Attribute);
-
-		// Transfer cluster assignments as rigid weights
-		using FBoneWeights = UE::AnimationCore::FBoneWeights;
-		using FBoneWeight = UE::AnimationCore::FBoneWeight;
-		ensureMsgf(TargetMesh.MaxVertexID() == InMedialSkeleton.VIDtoClusterIndex.Num(), TEXT("Medial skeleton does not match target mesh max ID: %d vs %d"), InMedialSkeleton.VIDtoClusterIndex.Num(), TargetMesh.MaxVertexID());
-		for (int32 VID = 0; VID < InMedialSkeleton.VIDtoClusterIndex.Num(); ++VID)
-		{
-			if (!TargetMesh.IsVertex(VID))
-			{
-				continue;
-			}
-			int32 SourceClusterIndex = InMedialSkeleton.VIDtoClusterIndex[VID];
-			if (!BoneRemap.IsValidIndex(SourceClusterIndex))
-			{
-				continue;
-			}
-			int32 BoneIndex = BoneRemap[SourceClusterIndex];
-			if (!ensureMsgf(BoneIndex >= 0 && BoneIndex <= TNumericLimits<FBoneIndexType>::Max(), TEXT("Invalid bone index %d"), BoneIndex))
-			{
-				continue;
-			}
-			FBoneWeights Wts;
-			Wts.SetBoneWeight(static_cast<FBoneIndexType>(BoneIndex), 1.f);
-			Attribute->SetValue(VID, Wts);
-		}
-	}
-
-
 
 	// Instantiate templated ComputeSkeleton methods for the supported mesh types
 	template GEOMETRYALGORITHMS_API FMedialSkeleton FSkeletonViaSampling::ComputeSkeleton<FDynamicMesh3>(const TMeshAABBTree3<FDynamicMesh3>& MeshBVH, TFunctionRef<FVector3d(int32)> GetVertexNormal, TFunctionRef<int32(int32)> VIDtoWeldVIDFn, const TFastWindingTree<FDynamicMesh3>* MeshFWTree);
diff --git a/Engine/Plugins/Runtime/GeometryProcessing/Source/GeometryAlgorithms/Public/Skeletonization/MeshMedialAxisSampling.h b/Engine/Plugins/Runtime/GeometryProcessing/Source/GeometryAlgorithms/Public/Skeletonization/MeshMedialAxisSampling.h
--- a/Engine/Plugins/Runtime/GeometryProcessing/Source/GeometryAlgorithms/Public/Skeletonization/MeshMedialAxisSampling.h
+++ b/Engine/Plugins/Runtime/GeometryProcessing/Source/GeometryAlgorithms/Public/Skeletonization/MeshMedialAxisSampling.h
@@ -408,17 +408,6 @@ namespace MedialAxis
 		{
 			return ToHierarchy<FDynamicMesh3>(InMedialSkeleton, nullptr, OutParent, OutTransform, Options, RootIndex, OutBoneIndexToMedialIndex, OutMedialIndexToBoneIndex);
 		}
-
-		/**
-		 * Transfer a medial skeleton to a dynamic mesh as a hierarchical skeleton, so it can be converted to a USkeletalMesh
-		 * 
-		 * @param InMedialSkeleton Input medial skeleton
-		 * @param TargetMesh Mesh to add skeleton a skin weights to
-		 * @param WeightProfile Name of skin weight profile to use. Note: Other profiles will be cleared, since they would likely not be compatible with the new skeleton.
-		 * @param Options Additional options to control the skeleton conversion
-		 * @param RootIndex Specify the medial skeleton cluster to prefer as the root bone. If INDEX_NONE, a default root will be selected automatically.
-		 */
-		GEOMETRYALGORITHMS_API static void ToDynamicMeshSkeleton(const FMedialSkeleton& InMedialSkeleton, FDynamicMesh3& TargetMesh, FName WeightProfile, FMedialSkeletonToTreeSkeletonOptions Options, int32 RootIndex = INDEX_NONE);
 	};
 
 	struct FSkeletonSimplifier
