Add SphereOverlapThreshold to skeleton simplifier

#Synced-CL 52426395

diff --git a/Engine/Plugins/Runtime/GeometryProcessing/Source/GeometryAlgorithms/Private/Skeletonization/MeshMedialAxisSampling.cpp b/Engine/Plugins/Runtime/GeometryProcessing/Source/GeometryAlgorithms/Private/Skeletonization/MeshMedialAxisSampling.cpp
--- a/Engine/Plugins/Runtime/GeometryProcessing/Source/GeometryAlgorithms/Private/Skeletonization/MeshMedialAxisSampling.cpp
+++ b/Engine/Plugins/Runtime/GeometryProcessing/Source/GeometryAlgorithms/Private/Skeletonization/MeshMedialAxisSampling.cpp
@@ -753,6 +753,24 @@ namespace MeshMedialAxisSamplingLocals
 			{
 				return true;
 			}
+			if (Simplifier.SphereOverlapThreshold > 0)
+			{
+				double RadiusA = Skeleton.Spheres[Edge.A].W;
+				double RadiusB = Skeleton.Spheres[Edge.B].W;
+				double Dist = FVector3d::Dist(Skeleton.Spheres[Edge.A].Center, Skeleton.Spheres[Edge.B].Center);
+				double OverlapDepth = RadiusA + RadiusB - Dist;
+				if (OverlapDepth < 0)
+				{
+					return true; // Spheres don't touch at all
+				}
+				double MinRadius = FMath::Min(RadiusA, RadiusB);
+				// Normalize by diameter of smaller sphere, or snap to fully-contained if diameter is 0
+				double OverlapFraction = (MinRadius > 0) ? (OverlapDepth / (2.0 * MinRadius)) : 1.0;
+				if (OverlapFraction < Simplifier.SphereOverlapThreshold)
+				{
+					return true; // Spheres don't overlap enough
+				}
+			}
 			return false;
 		};
 
diff --git a/Engine/Plugins/Runtime/GeometryProcessing/Source/GeometryAlgorithms/Public/Skeletonization/MeshMedialAxisSampling.h b/Engine/Plugins/Runtime/GeometryProcessing/Source/GeometryAlgorithms/Public/Skeletonization/MeshMedialAxisSampling.h
--- a/Engine/Plugins/Runtime/GeometryProcessing/Source/GeometryAlgorithms/Public/Skeletonization/MeshMedialAxisSampling.h
+++ b/Engine/Plugins/Runtime/GeometryProcessing/Source/GeometryAlgorithms/Public/Skeletonization/MeshMedialAxisSampling.h
@@ -427,6 +427,11 @@ namespace MedialAxis
 		// If positive, the maximum medial sphere radius to consider for collapse (will not collapse if both cluster's medial sphere radii on an edge are larger than this)
 		double SphereRadiusThreshold = -1;
 
+		// If positive, will only collapse edges with medial spheres whose overlap depth is at least this fraction of the smaller sphere's diameter --
+		//  so at 1.0 we will only simplify an edge if the smaller sphere is fully contained inside the larger
+		// (A degenerate zero-radius sphere is considered to fully overlap any sphere it touches.)
+		double SphereOverlapThreshold = -1;
+
 		// If true, only simplify edges that have triangles associated
 		bool bOnlySimplifySurfaces = false;
 
