[GeometryProcess / SkeletalMesh Modeling Tools] Added cotan weighting support to TSmoothBoneWeights::SmoothWeightsAtVertex

#Synced-CL 52437958

diff --git a/Engine/Plugins/Runtime/GeometryProcessing/Source/DynamicMesh/Private/Operations/SmoothBoneWeights.cpp b/Engine/Plugins/Runtime/GeometryProcessing/Source/DynamicMesh/Private/Operations/SmoothBoneWeights.cpp
--- a/Engine/Plugins/Runtime/GeometryProcessing/Source/DynamicMesh/Private/Operations/SmoothBoneWeights.cpp
+++ b/Engine/Plugins/Runtime/GeometryProcessing/Source/DynamicMesh/Private/Operations/SmoothBoneWeights.cpp
@@ -6,6 +6,7 @@
 #include "DynamicMesh/DynamicMeshAttributeSet.h"
 #include "DynamicMesh/NonManifoldMappingSupport.h"
 #include "Util/ProgressCancel.h"
+#include "Solvers/PrecomputedMeshWeightData.h"
 #include "Math/UnrealMathUtility.h"
 #include "Parameterization/MeshLocalParam.h"
 #include "DynamicMesh/MeshNormals.h"
@@ -126,53 +127,78 @@ EOperationValidationResult TSmoothBoneWeights<BoneIndexType, BoneWeightType>::Va
 	return EOperationValidationResult::Ok;
 }
 
+template <typename BoneIndexType, typename BoneWeightType>
+void TSmoothBoneWeights<BoneIndexType, BoneWeightType>::ComputeEdgeWeightsIfNeeded()
+{
+	if (EdgeWeightMethod == ESmoothBoneWeightsEdgeWeightMethod::CotanWeights
+		&& CotangentEdgeWeights.Num() != SourceMesh->MaxEdgeID())
+	{
+		UE::MeshDeformation::ConstructEdgeCotanWeightsDataArray(*SourceMesh, CotangentEdgeWeights);
+	}
+}
+
 template <typename BoneIndexType, typename BoneWeightType>
 bool TSmoothBoneWeights<BoneIndexType, BoneWeightType>::SmoothWeightsAtVertex(const int32  VertexID,
 																			  const BoneWeightType VertexFalloff,
 																			  TMap<BoneIndexType, BoneWeightType>& OutFinalWeights)
 {
 	using namespace SmoothBoneWeightsLocals;
-	
-	// const UE::Geometry::FNonManifoldMappingSupport NonManifoldMappingSupport(*SourceMesh);
 
-	// for each vertex in the stamp...
-	constexpr int32 AvgNumNeighbors = 8;
-	using VertexNeighborWeights = TArray<BoneWeightType, TInlineAllocator<AvgNumNeighbors>>;
-	TArray<int32> AllNeighborVertices;
-	TMap<BoneIndexType, VertexNeighborWeights> WeightsOnAllNeighbors;
-	// const int32 SrcVertexID = NonManifoldMappingSupport.GetOriginalNonManifoldVertexID(VertexID);
+	ComputeEdgeWeightsIfNeeded();
 
-	// get list of all neighboring vertices, AND this vertex
-	AllNeighborVertices.Add(VertexID);
-	for (const int32 NeighborVertexID : SourceMesh->VtxVerticesItr(VertexID))
-	{
-		AllNeighborVertices.Add(NeighborVertexID);
-	}
+	// Weighted accumulators per bone: WeightedBoneSum accumulates edgeWeight * boneWeight,
+	// TotalEdgeWeightPerBone accumulates edgeWeight. Final average = WeightedBoneSum / TotalEdgeWeightPerBone.
+	// When Uniform, all edge weights are 1.0, so this reduces to Sum / Count (identical to previous behavior).
+	TMap<BoneIndexType, BoneWeightType> WeightedBoneSum;
+	TMap<BoneIndexType, BoneWeightType> TotalEdgeWeightPerBone;
 
-	// get all weights above a given threshold across ALL neighbors (including self)
-	for (const int32 NeighborVertexID : AllNeighborVertices)
+	// Helper lambda to accumulate bone weights from a vertex with a given edge weight
+	auto AccumulateWeightsFromVertex = [&](const int32 NeighborVertexID, const BoneWeightType EdgeWeight)
 	{
 		for (int32 Index = 0; Index < DataSource->GetBoneNum(NeighborVertexID); ++Index)
-		{	
+		{
 			const BoneWeightType Weight = DataSource->GetBoneWeight(NeighborVertexID, Index);
 			if (Weight > MinimumWeightThreshold)
 			{
-				VertexNeighborWeights& BoneWeights = WeightsOnAllNeighbors.FindOrAdd(DataSource->GetBoneIndex(NeighborVertexID, Index));
-				BoneWeights.Add(Weight);
+				const BoneIndexType BoneIndex = DataSource->GetBoneIndex(NeighborVertexID, Index);
+				WeightedBoneSum.FindOrAdd(BoneIndex) += EdgeWeight * Weight;
+				TotalEdgeWeightPerBone.FindOrAdd(BoneIndex) += EdgeWeight;
 			}
 		}
+	};
+
+	// Self vertex always contributes with edge weight 1.0 (no mesh edge to self)
+	AccumulateWeightsFromVertex(VertexID, (BoneWeightType)1.0);
+
+	// Accumulate neighbor contributions with appropriate edge weights
+	for (const int32 NeighborVertexID : SourceMesh->VtxVerticesItr(VertexID))
+	{
+		BoneWeightType EdgeWeight = (BoneWeightType)1.0;
+		if (EdgeWeightMethod == ESmoothBoneWeightsEdgeWeightMethod::CotanWeights)
+		{
+			const int32 EdgeId = SourceMesh->FindEdge(VertexID, NeighborVertexID);
+			if (EdgeId != FDynamicMesh3::InvalidID)
+			{
+				// Cotangent weights can be negative; clamp to zero (same as FSmoothDynamicMeshAttributes)
+				EdgeWeight = static_cast<BoneWeightType>(FMath::Max(0.0, CotangentEdgeWeights[EdgeId]));
+			}
+			else
+			{
+				EdgeWeight = (BoneWeightType)0.0;
+			}
+		}
+		AccumulateWeightsFromVertex(NeighborVertexID, EdgeWeight);
 	}
 
-	// calculate single average weight of each bone on all the neighbors
+	// Calculate weighted average per bone
 	OutFinalWeights.Reset();
-	for (const TTuple<BoneIndexType, VertexNeighborWeights>& NeighborWeights : WeightsOnAllNeighbors)
+	for (const TTuple<BoneIndexType, BoneWeightType>& Entry : WeightedBoneSum)
 	{
-		BoneWeightType TotalWeightOnThisBone = (BoneWeightType) 0.f;
-		for (const BoneWeightType& Value : NeighborWeights.Value)
+		const BoneWeightType TotalEdgeWeight = TotalEdgeWeightPerBone.FindRef(Entry.Key);
+		if (!FMath::IsNearlyZero(TotalEdgeWeight))
 		{
-			TotalWeightOnThisBone += Value;
+			OutFinalWeights.Add(Entry.Key, Entry.Value / TotalEdgeWeight);
 		}
-		OutFinalWeights.Add(NeighborWeights.Key, TotalWeightOnThisBone / (BoneWeightType)NeighborWeights.Value.Num());
 	}
 
 	// normalize the weights
@@ -206,15 +232,8 @@ bool TSmoothBoneWeights<BoneIndexType, BoneWeightType>::SmoothWeightsAtVertex(
 	TMap<BoneIndexType, BoneWeightType>& OutFinalWeights)
 {
 	using namespace SmoothBoneWeightsLocals;
-	
-	// const UE::Geometry::FNonManifoldMappingSupport NonManifoldMappingSupport(*SourceMesh);
 
-	// for each vertex in the stamp...
-	constexpr int32 AvgNumNeighbors = 8;
-	using VertexNeighborWeights = TArray<BoneWeightType, TInlineAllocator<AvgNumNeighbors>>;
-	TArray<int32> AllNeighborVertices;
-	TMap<BoneIndexType, VertexNeighborWeights> WeightsOnAllNeighbors;
-	// const int32 SrcVertexID = NonManifoldMappingSupport.GetOriginalNonManifoldVertexID(VertexID);
+	ComputeEdgeWeightsIfNeeded();
 
 	TMap<BoneIndexType, BoneWeightType> LockedWeights;
 	BoneWeightType TotalLockedWeight = 0;
@@ -228,46 +247,59 @@ bool TSmoothBoneWeights<BoneIndexType, BoneWeightType>::SmoothWeightsAtVertex(
 			TotalLockedWeight += Weight;
 		}
 	}
-	
-	// get list of all neighboring vertices, AND this vertex
-	AllNeighborVertices.Add(VertexID);
-	for (const int32 NeighborVertexID : SourceMesh->VtxVerticesItr(VertexID))
-	{
-		AllNeighborVertices.Add(NeighborVertexID);
-	}
 
-	// get all weights above a given threshold across ALL neighbors (including self)
-	for (const int32 NeighborVertexID : AllNeighborVertices)
+	// Weighted accumulators per bone (excluding locked bones)
+	TMap<BoneIndexType, BoneWeightType> WeightedBoneSum;
+	TMap<BoneIndexType, BoneWeightType> TotalEdgeWeightPerBone;
+
+	// Helper lambda to accumulate bone weights from a vertex with a given edge weight, skipping locked bones
+	auto AccumulateWeightsFromVertex = [&](const int32 NeighborVertexID, const BoneWeightType EdgeWeight)
 	{
 		for (int32 Index = 0; Index < DataSource->GetBoneNum(NeighborVertexID); ++Index)
 		{
-			BoneIndexType BoneIndex = DataSource->GetBoneIndex(NeighborVertexID, Index);
+			const BoneIndexType BoneIndex = DataSource->GetBoneIndex(NeighborVertexID, Index);
 			if (!LockedBones.Contains(BoneIndex))
 			{
 				const BoneWeightType Weight = DataSource->GetBoneWeight(NeighborVertexID, Index);
 				if (Weight > MinimumWeightThreshold)
 				{
-					VertexNeighborWeights& BoneWeights = WeightsOnAllNeighbors.FindOrAdd(BoneIndex);
-					BoneWeights.Add(Weight);
+					WeightedBoneSum.FindOrAdd(BoneIndex) += EdgeWeight * Weight;
+					TotalEdgeWeightPerBone.FindOrAdd(BoneIndex) += EdgeWeight;
 				}
 			}
-			
 		}
-	}
+	};
 
-	// calculate single average weight of each bone on all the neighbors, excluding locked ones
-	OutFinalWeights.Reset();
-	for (const TTuple<BoneIndexType, VertexNeighborWeights>& NeighborWeights : WeightsOnAllNeighbors)
+	// Self vertex always contributes with edge weight 1.0
+	AccumulateWeightsFromVertex(VertexID, (BoneWeightType)1.0);
+
+	// Accumulate neighbor contributions with appropriate edge weights
+	for (const int32 NeighborVertexID : SourceMesh->VtxVerticesItr(VertexID))
 	{
-		BoneIndexType BoneIndex = NeighborWeights.Key;
-		if (!LockedBones.Contains(BoneIndex))
+		BoneWeightType EdgeWeight = (BoneWeightType)1.0;
+		if (EdgeWeightMethod == ESmoothBoneWeightsEdgeWeightMethod::CotanWeights)
 		{
-			BoneWeightType TotalWeightOnThisBone = (BoneWeightType) 0.f;
-			for (const BoneWeightType& Value : NeighborWeights.Value)
+			const int32 EdgeId = SourceMesh->FindEdge(VertexID, NeighborVertexID);
+			if (EdgeId != FDynamicMesh3::InvalidID)
 			{
-				TotalWeightOnThisBone += Value;
+				EdgeWeight = static_cast<BoneWeightType>(FMath::Max(0.0, CotangentEdgeWeights[EdgeId]));
 			}
-			OutFinalWeights.Add(BoneIndex, TotalWeightOnThisBone / (BoneWeightType)NeighborWeights.Value.Num());
+			else
+			{
+				EdgeWeight = (BoneWeightType)0.0;
+			}
+		}
+		AccumulateWeightsFromVertex(NeighborVertexID, EdgeWeight);
+	}
+
+	// Calculate weighted average per bone, excluding locked ones
+	OutFinalWeights.Reset();
+	for (const TTuple<BoneIndexType, BoneWeightType>& Entry : WeightedBoneSum)
+	{
+		const BoneWeightType TotalEdgeWeight = TotalEdgeWeightPerBone.FindRef(Entry.Key);
+		if (!FMath::IsNearlyZero(TotalEdgeWeight))
+		{
+			OutFinalWeights.Add(Entry.Key, Entry.Value / TotalEdgeWeight);
 		}
 	}
 
diff --git a/Engine/Plugins/Runtime/GeometryProcessing/Source/DynamicMesh/Public/Operations/SmoothBoneWeights.h b/Engine/Plugins/Runtime/GeometryProcessing/Source/DynamicMesh/Public/Operations/SmoothBoneWeights.h
--- a/Engine/Plugins/Runtime/GeometryProcessing/Source/DynamicMesh/Public/Operations/SmoothBoneWeights.h
+++ b/Engine/Plugins/Runtime/GeometryProcessing/Source/DynamicMesh/Public/Operations/SmoothBoneWeights.h
@@ -25,6 +25,14 @@ namespace UE
 namespace Geometry
 {
 
+/** Method used to compute edge weights during bone weight smoothing. */
+enum class ESmoothBoneWeightsEdgeWeightMethod : uint8
+{
+	/** Each neighbor contributes equally (existing behavior). */
+	Uniform = 0,
+	/** Neighbors weighted by cotangent of opposite angles on shared edge. */
+	CotanWeights = 1
+};
 
 /** 
  * Interface for getting the bone weights data. Should be subclassed and implemented if the bone weights data
@@ -97,7 +105,10 @@ public:
 	FProgressCancel* Progress = nullptr;
 
 	/** Only consider the weights above the threshold. */
-	BoneWeightType MinimumWeightThreshold = (BoneWeightType) 0.0f; 
+	BoneWeightType MinimumWeightThreshold = (BoneWeightType) 0.0f;
+
+	/** The edge weighting method used during smoothing. */
+	ESmoothBoneWeightsEdgeWeightMethod EdgeWeightMethod = ESmoothBoneWeightsEdgeWeightMethod::Uniform;
 
 protected:
 		
@@ -105,7 +116,13 @@ protected:
 	
 	/** Data source supplied by the caller. */
 	TBoneWeightsDataSource<BoneIndexType, BoneWeightType>* DataSource = nullptr;
-	
+
+	/** Pre-computed cotangent edge weights indexed by EdgeID. Populated lazily when CotanWeights is selected. */
+	TArray<double> CotangentEdgeWeights;
+
+	/** Pre-compute edge weights if needed based on EdgeWeightMethod. */
+	void ComputeEdgeWeightsIfNeeded();
+
 public:
 
 	/**
diff --git a/Engine/Plugins/Runtime/MeshModelingToolset/Source/MeshModelingToolsEditorOnly/Private/SkeletalMesh/SkinWeightsPaintTool.cpp b/Engine/Plugins/Runtime/MeshModelingToolset/Source/MeshModelingToolsEditorOnly/Private/SkeletalMesh/SkinWeightsPaintTool.cpp
--- a/Engine/Plugins/Runtime/MeshModelingToolset/Source/MeshModelingToolsEditorOnly/Private/SkeletalMesh/SkinWeightsPaintTool.cpp
+++ b/Engine/Plugins/Runtime/MeshModelingToolset/Source/MeshModelingToolsEditorOnly/Private/SkeletalMesh/SkinWeightsPaintTool.cpp
@@ -2621,7 +2621,7 @@ void USkinWeightsPaintTool::UpdateCurrentlyEditedMesh(const FDynamicMesh3& InDyn
 	bVertexColorsNeedUpdated = true;
 	
 	// update smooth operator (this must be done after PreviewMesh & Weights have been updated)
-	InitializeSmoothWeightsOperator();
+	InitializeSmoothWeightsOperator(InDynamicMesh);
 
 	// after any mesh change, the mirror tables will need rebuilt next time mirroring is used
 	MirrorData.SetNeedsReinitialized();
@@ -3360,18 +3360,13 @@ void USkinWeightsPaintTool::InitializeOctrees()
 	}
 }
 
-void USkinWeightsPaintTool::InitializeSmoothWeightsOperator()
+void USkinWeightsPaintTool::InitializeSmoothWeightsOperator(const FDynamicMesh3& RestPoseMesh)
 {
-	if (!ensure(PreviewMesh && PreviewMesh->GetMesh()))
-	{
-		return;
-	}
-
 	// NOTE: this could probably be initialized lazily as it's only used with the relax brush
-	const FDynamicMesh3* DynaMesh = PreviewMesh->GetMesh();
 	SmoothWeightsDataSource = MakeUnique<FPaintToolWeightsDataSource>(&Weights);
-	SmoothWeightsOp = MakeUnique<UE::Geometry::TSmoothBoneWeights<int32, float>>(DynaMesh, SmoothWeightsDataSource.Get());
+	SmoothWeightsOp = MakeUnique<UE::Geometry::TSmoothBoneWeights<int32, float>>(&RestPoseMesh, SmoothWeightsDataSource.Get());
 	SmoothWeightsOp->MinimumWeightThreshold = MinimumWeightThreshold;
+	SmoothWeightsOp->EdgeWeightMethod = UE::Geometry::ESmoothBoneWeightsEdgeWeightMethod::CotanWeights;
 }
 
 void USkinWeightsPaintTool::ApplyWeightEditsWithoutTransaction(const FMultiBoneWeightEdits& WeightEdits)
diff --git a/Engine/Plugins/Runtime/MeshModelingToolset/Source/MeshModelingToolsEditorOnly/Public/SkeletalMesh/SkinWeightsPaintTool.h b/Engine/Plugins/Runtime/MeshModelingToolset/Source/MeshModelingToolsEditorOnly/Public/SkeletalMesh/SkinWeightsPaintTool.h
--- a/Engine/Plugins/Runtime/MeshModelingToolset/Source/MeshModelingToolsEditorOnly/Public/SkeletalMesh/SkinWeightsPaintTool.h
+++ b/Engine/Plugins/Runtime/MeshModelingToolset/Source/MeshModelingToolsEditorOnly/Public/SkeletalMesh/SkinWeightsPaintTool.h
@@ -909,7 +909,7 @@ protected:
 	// Smooth weights data source and operator
 	TUniquePtr<UE::Geometry::TBoneWeightsDataSource<int32, float>> SmoothWeightsDataSource;
 	TUniquePtr<UE::Geometry::TSmoothBoneWeights<int32, float>> SmoothWeightsOp;
-	UE_API void InitializeSmoothWeightsOperator();
+	UE_API void InitializeSmoothWeightsOperator(const FDynamicMesh3& RestPoseMesh);
 
 	// vertex colors updated when switching current bone or initializing whole mesh
 	UE_API void UpdateVertexColorForAllVertices();
