fix crash in constrained skin binding when some vertices have no assigned group

#Synced-CL 52609914

diff --git a/Engine/Plugins/Runtime/GeometryProcessing/Source/DynamicMesh/Private/Operations/SkinWeightBinding.cpp b/Engine/Plugins/Runtime/GeometryProcessing/Source/DynamicMesh/Private/Operations/SkinWeightBinding.cpp
--- a/Engine/Plugins/Runtime/GeometryProcessing/Source/DynamicMesh/Private/Operations/SkinWeightBinding.cpp
+++ b/Engine/Plugins/Runtime/GeometryProcessing/Source/DynamicMesh/Private/Operations/SkinWeightBinding.cpp
@@ -503,6 +503,42 @@ static TConstArrayView<TArray<int32, TInlineAllocator<16>>> FilterBoneGroupEndEf
 }
 
 
+// Enumerate bone group for a vertex, w/ fallback to standard un-constrained bones when vertex has no valid bone group
+static void EnumerateVertexBoneGroup(
+	int32 VertexIndex,
+	TConstArrayView<int32> VIDtoGroupIndex,
+	TConstArrayView<TArray<int32, TInlineAllocator<16>>> UseBoneGroups,
+	const FTransformHierarchyQuery& Skeleton,
+	int32 NumBones,
+	int32 MaxInfluences,
+	TFunctionRef<void(int32)> ProcessBone)
+{
+	int32 GroupIndex = VIDtoGroupIndex[VertexIndex];
+	// for vertices that weren't assigned a (valid) group, fall back to a full search over bones
+	if (!UseBoneGroups.IsValidIndex(GroupIndex))
+	{
+		for (int32 BoneIndex = 0; BoneIndex < NumBones; ++BoneIndex)
+		{
+			// For single influences, avoid end effectors to avoid closest-distance fighting with the tips of their
+			// parent bones.
+			if (MaxInfluences == 1 && Skeleton.IsEndEffector(BoneIndex))
+			{
+				continue;
+			}
+
+			ProcessBone(BoneIndex);
+		}
+	}
+	else
+	{
+		for (int32 BoneIndex : UseBoneGroups[GroupIndex])
+		{
+			ProcessBone(BoneIndex);
+		}
+	}
+}
+
+
 // Direct Distance skin weight implementation where vertices can only bind to bones in limited 'bone groups'
 void CreateSkinWeightsImpl_DirectDistance_Constrained(
 	FDynamicMesh3& InMesh, TConstArrayView<SkinBinding::FBonePoseInfo> TransformHierarchy, FName ProfileName,
@@ -525,13 +561,9 @@ void CreateSkinWeightsImpl_DirectDistance_Constrained(
 		FilterBoneGroupEndEffectors(MaxInfluences, BoneGroups, Skeleton, FilteredBoneGroups);
 
 	CreateSkinWeightsImpl_DirectDistance_General(InMesh, TransformHierarchy, ProfileName, InStiffness, InSettings, Skeleton,
-	[&VIDtoGroupIndex, &UseBoneGroups](int32 VertexIndex, TFunctionRef<void(int32)> ProcessBone)
+	[&VIDtoGroupIndex, &UseBoneGroups, &Skeleton, NumBones=TransformHierarchy.Num(), MaxInfluences](int32 VertexIndex, TFunctionRef<void(int32)> ProcessBone)
 	{
-		int32 GroupIndex = VIDtoGroupIndex[VertexIndex];
-		for (int32 BoneIndex : UseBoneGroups[GroupIndex])
-		{
-			ProcessBone(BoneIndex);
-		}
+		EnumerateVertexBoneGroup(VertexIndex, VIDtoGroupIndex, UseBoneGroups, Skeleton, NumBones, MaxInfluences, ProcessBone);
 	});
 }
 
@@ -781,6 +813,27 @@ void CreateSkinWeightsImpl_GeodesicVoxel_Constrained(
 		}
 	}
 
+	// If we have any vertices w/out a valid group, they will need an unconstrained search; update ActiveBones accordingly
+	bool bHasUnassignedVerts = false;
+	for (int32 GroupIndex : VIDtoGroupIndex)
+	{
+		if (!UseBoneGroups.IsValidIndex(GroupIndex))
+		{
+			bHasUnassignedVerts = true;
+			break;
+		}
+	}
+	if (bHasUnassignedVerts)
+	{
+		for (int32 BoneIndex = 0; BoneIndex < TransformHierarchy.Num(); ++BoneIndex)
+		{
+			if (MaxInfluences > 1 || !Skeleton.IsEndEffector(BoneIndex))
+			{
+				ActiveBones[BoneIndex] = true;
+			}
+		}
+	}
+
 	CreateSkinWeightsImpl_GeodesicVoxel_General(InMesh, TransformHierarchy, ProfileName, InStiffness, InVoxelResolution, InSettings, Skeleton,
 		// ShouldComputeBone: only compute bones that appear in at least one group
 		[&ActiveBones](int32 BoneIndex) -> bool
@@ -788,13 +841,9 @@ void CreateSkinWeightsImpl_GeodesicVoxel_Constrained(
 			return ActiveBones[BoneIndex];
 		},
 		// EnumerateBonesForVertex: only iterate bones in the vertex's assigned group
-		[&VIDtoGroupIndex, &UseBoneGroups](int32 VertexIndex, TFunctionRef<void(int32)> ProcessBone)
+		[&VIDtoGroupIndex, &UseBoneGroups, &Skeleton, MaxInfluences, NumBones=TransformHierarchy.Num()](int32 VertexIndex, TFunctionRef<void(int32)> ProcessBone)
 		{
-			int32 GroupIndex = VIDtoGroupIndex[VertexIndex];
-			for (int32 BoneIndex : UseBoneGroups[GroupIndex])
-			{
-				ProcessBone(BoneIndex);
-			}
+			EnumerateVertexBoneGroup(VertexIndex, VIDtoGroupIndex, UseBoneGroups, Skeleton, NumBones, MaxInfluences, ProcessBone);
 		});
 }
 
diff --git a/Engine/Plugins/Runtime/GeometryProcessing/Source/DynamicMesh/Public/Operations/SkinWeightBinding.h b/Engine/Plugins/Runtime/GeometryProcessing/Source/DynamicMesh/Public/Operations/SkinWeightBinding.h
--- a/Engine/Plugins/Runtime/GeometryProcessing/Source/DynamicMesh/Public/Operations/SkinWeightBinding.h
+++ b/Engine/Plugins/Runtime/GeometryProcessing/Source/DynamicMesh/Public/Operations/SkinWeightBinding.h
@@ -51,7 +51,8 @@ namespace SkinBinding
 	// Create skin weights from the input bone hierarchy, and set them on the dynamic mesh
 	UE_API void CreateSkinWeights(FDynamicMesh3& InMesh, TConstArrayView<FBonePoseInfo> TransformHierarchy, FName ProfileName, const FBindSettings& Settings);
 
-	// Create skin weights from the input bone hierarchy, with vertices constrained to only allow weights on specific subsets of bones
+	// Create skin weights from the input bone hierarchy, with vertices constrained to only allow weights on specific subsets of bones.
+	// @param VIDtoGroupIndex Mapping from vertex ID to BoneGroups index, or INDEX_NONE if the vertex has no assigned group (will use unconstrained weights in this case)
 	UE_API void CreateConstrainedSkinWeights(FDynamicMesh3& InMesh, TConstArrayView<FBonePoseInfo> TransformHierarchy, FName ProfileName, const FBindSettings& Settings,
 		TConstArrayView<TArray<int32, TInlineAllocator<16>>> BoneGroups, TConstArrayView<int32> VIDtoGroupIndex);
 
