minor comment improvement for skinweightbinding

#Synced-CL 52609914

diff --git a/Engine/Plugins/Runtime/GeometryProcessing/Source/DynamicMesh/Private/Operations/SkinWeightBinding.cpp b/Engine/Plugins/Runtime/GeometryProcessing/Source/DynamicMesh/Private/Operations/SkinWeightBinding.cpp
--- a/Engine/Plugins/Runtime/GeometryProcessing/Source/DynamicMesh/Private/Operations/SkinWeightBinding.cpp
+++ b/Engine/Plugins/Runtime/GeometryProcessing/Source/DynamicMesh/Private/Operations/SkinWeightBinding.cpp
@@ -840,7 +840,7 @@ void CreateSkinWeightsImpl_GeodesicVoxel_Constrained(
 		{
 			return ActiveBones[BoneIndex];
 		},
-		// EnumerateBonesForVertex: only iterate bones in the vertex's assigned group
+		// EnumerateBonesForVertex: iterate bones in the vertex's assigned group (or all bones if unassigned)
 		[&VIDtoGroupIndex, &UseBoneGroups, &Skeleton, MaxInfluences, NumBones=TransformHierarchy.Num()](int32 VertexIndex, TFunctionRef<void(int32)> ProcessBone)
 		{
 			EnumerateVertexBoneGroup(VertexIndex, VIDtoGroupIndex, UseBoneGroups, Skeleton, NumBones, MaxInfluences, ProcessBone);
