clean up / generalize edge splits logic and add subdiv method for medial axis skeleton

#Synced-CL 52686699

diff --git a/Engine/Plugins/Runtime/GeometryProcessing/Source/GeometryAlgorithms/Private/Skeletonization/MeshMedialAxisSampling.cpp b/Engine/Plugins/Runtime/GeometryProcessing/Source/GeometryAlgorithms/Private/Skeletonization/MeshMedialAxisSampling.cpp
--- a/Engine/Plugins/Runtime/GeometryProcessing/Source/GeometryAlgorithms/Private/Skeletonization/MeshMedialAxisSampling.cpp
+++ b/Engine/Plugins/Runtime/GeometryProcessing/Source/GeometryAlgorithms/Private/Skeletonization/MeshMedialAxisSampling.cpp
@@ -145,29 +145,33 @@ namespace MeshMedialAxisSamplingLocals
 		);
 	}
 
-	// Info about a split edge: original edge A-B was split into A-I and I-B
+	// Info about a split edge: original edge A-B was split, inserting intermediate nodes along it
 	struct FSplitEdgeInfo
 	{
-		FSplitEdgeInfo(int32 A, int32 B, int32 I) : A(A), B(B), I(I) {}
-		int32 A, B, I;
+		FSplitEdgeInfo(int32 InA, int32 InB) : A(InA), B(InB) {}
+		int32 A, B;
+		TArray<int32, TInlineAllocator<2>> IntermediateNodes; // nodes inserted between A and B, in order
 		TArray<int32, TInlineAllocator<2>> OppositeVerts; // vertices opposite the source AB edge
 	};
 
-	// Helper to split skinny triangles -- to give the simplifier more DOFs to collapse away surfaces
-	template<typename MeshType>
-	void SplitThinSkeletonTriEdges(
+	// Request to split an edge at specified parametric positions
+	struct FEdgeSplitRequest
+	{
+		UE::Geometry::FIndex2i Edge;
+		TArray<double, TInlineAllocator<4>> SplitTs; // parameter values in (0,1), sorted, where intermediate nodes should be added
+	};
+
+	// Split skeleton edges according to provided split requests, updating topology
+	// CollectCandidates takes Skeleton and EdgeToOppVerts map, and returns-by-reference the split candidates
+	void SplitSkeletonEdgesImpl(
 		UE::Geometry::MedialAxis::FMedialSkeleton& Skeleton,
-		TArray<FSplitEdgeInfo>& OutSplitEdges,
-		const MeshType& Mesh,
-		TConstArrayView<FSQEM> VtxSQEMs,
-		TConstArrayView<float> VtxAreas,
-		double AngleThresholdDeg,
-		double PosErrorWt)
+		TFunctionRef<void(const UE::Geometry::MedialAxis::FMedialSkeleton&, const TMap<UE::Geometry::FIndex2i, TArray<int32>>&, TArray<FEdgeSplitRequest>&)> CollectCandidates,
+		int32& OutStartNewClusters,
+		TSet<int32>& OutClusterCanRemapVertices,
+		TArray<FSplitEdgeInfo>* OutSplitEdges = nullptr)
 	{
 		using namespace UE::Geometry;
 
-		const double UseCosAngleThreshold = FMath::Cos(FMath::Clamp(AngleThresholdDeg, 90. + FMathd::ZeroTolerance, 180. - FMathd::ZeroTolerance) * FMathd::DegToRad);
-
 		// Build edge to opposite vertex map for all tris
 		TMap<FIndex2i, TArray<int32>> EdgeToOppVerts;
 		for (const FIndex3i& Tri : Skeleton.ClusterTriangles)
@@ -192,7 +196,7 @@ namespace MeshMedialAxisSamplingLocals
 			}
 		};
 
-		// Helpers to add/remove edges from Skeleton data
+		// Helper to add edges from Skeleton data
 		auto AddEdge = [&Skeleton](int32 A, int32 B)
 		{
 			Skeleton.ClusterNeighbors[A].AddUnique(B);
@@ -203,112 +207,120 @@ namespace MeshMedialAxisSamplingLocals
 		// The skeleton triangle changes to apply after
 		TSet<FIndex3i> TrianglesToRemove;
 		TArray<FIndex3i> TrianglesToAdd;
-		// The clusters that may need local vertex reassignments after
-		TSet<int32> ClusterCanRemapVertices;
 
-		// Helper to split an edge by inserting an (already-created) new cluster, I, and do associated book-keeping
-		auto SplitEdge = [&OutSplitEdges, &EdgeToOppVerts, &TrianglesToRemove, &TrianglesToAdd, &ClusterCanRemapVertices, &AddEdge, &ReplaceEdgeOppVert, &Skeleton](FIndex2i Edge, int32 I)
-		{
-			int32 A = Edge.A, B = Edge.B;
+		// Collect candidates via caller-provided callback
+		TArray<FEdgeSplitRequest> Candidates;
+		CollectCandidates(Skeleton, EdgeToOppVerts, Candidates);
 
-			// Replace every triangle (A,B,Cx) with (A,I,Cx) and (B,I,Cx)
-			const TArray<int32>& OppVerts = EdgeToOppVerts[Edge];
-			FSplitEdgeInfo& SplitInfo = OutSplitEdges.Emplace_GetRef(A, B, I);
-			SplitInfo.OppositeVerts.Append(OppVerts);
-			for (int32 Cx : OppVerts)
-			{
-				TrianglesToRemove.Add(AsTri(A, B, Cx));
-				TrianglesToAdd.Add(AsTri(A, I, Cx));
-				TrianglesToAdd.Add(AsTri(B, I, Cx));
-				AddEdge(Cx, I);
+		OutStartNewClusters = Skeleton.Spheres.Num();
 
-				ReplaceEdgeOppVert(AsEdge(A, Cx), B, I);
-				ReplaceEdgeOppVert(AsEdge(B, Cx), A, I);
+		// Process each candidate: split edge into chain of intermediate nodes
+		for (const FEdgeSplitRequest& Cand : Candidates)
+		{
+			if (!ensure(!Cand.SplitTs.IsEmpty()))
+			{
+				continue;
 			}
+			const int32 A = Cand.Edge.A, B = Cand.Edge.B;
 
-			// Fix Skeleton.ClusterNeighbors w/ new edges
-			AddEdge(A, I);
-			AddEdge(I, B);
-			Skeleton.ClusterNeighbors[A].Remove(B);
-			Skeleton.ClusterNeighbors[B].Remove(A);
-
-			// Mark that clusters may need to re-assign vertices to new cluster(s)
-			ClusterCanRemapVertices.Add(A);
-			ClusterCanRemapVertices.Add(B);
-			ClusterCanRemapVertices.Append(OppVerts);
 
-			EdgeToOppVerts.Remove(Edge); // we only split edges once, can clear the stale ones
-			// (and likewise don't need AI, IB edges in EdgeToOppVerts)
-		};
+			// Build chain: [A, I1, I2, ..., I_k, B]
+			TArray<int32, TInlineAllocator<8>> ChainNodes;
+			ChainNodes.Add(A);
+			for (double T : Cand.SplitTs)
+			{
+				int32 I = Skeleton.Spheres.Add(LerpSphere(Skeleton.Spheres[A], Skeleton.Spheres[B], T));
+				Skeleton.ClusterNeighbors.AddDefaulted();
+				ChainNodes.Add(I);
+			}
+			ChainNodes.Add(B);
 
-		// Get candidates edges w/ attached skinny triangles to split
-		struct FCandidateSplit
-		{
-			FIndex2i Edge;
-			double CosAngle, ProjT;
-		};
-		TArray<FCandidateSplit> Candidates;
-		for (const auto& Entry : EdgeToOppVerts)
-		{
-			const FIndex2i& Edge = Entry.Key;
-			const TArray<int32>& OppVerts = Entry.Value;
-			const FVector3d& PA = Skeleton.Spheres[Edge.A].Center;
-			const FVector3d& PB = Skeleton.Spheres[Edge.B].Center;
-			const FVector3d AB = PB - PA;
-			const double ABLenSq = AB.SquaredLength();
-			if (ABLenSq < FMathd::ZeroTolerance)
+			// Get opposite vertices (copy since we modify EdgeToOppVerts below)
+			TArray<int32> OppVerts;
+			if (const TArray<int32>* OppVertsPtr = EdgeToOppVerts.Find(Cand.Edge))
 			{
-				continue;
+				OppVerts = *OppVertsPtr;
 			}
-			double MinCosAngle = UseCosAngleThreshold;
-			int32 MinCosOppV = INDEX_NONE;
+
+			// Fan triangulation from opposite vertices
 			for (int32 Cx : OppVerts)
 			{
-				const FVector3d& PC = Skeleton.Spheres[Cx].Center;
-				const double CosAngleAtCx = Normalized(PA - PC).Dot(Normalized(PB - PC));
-				if (CosAngleAtCx < MinCosAngle)
+				TrianglesToRemove.Add(AsTri(A, B, Cx));
+
+				for (int32 Seg = 0; Seg < ChainNodes.Num() - 1; ++Seg)
 				{
-					MinCosAngle = CosAngleAtCx;
-					MinCosOppV = Cx;
+					TrianglesToAdd.Add(AsTri(ChainNodes[Seg], ChainNodes[Seg + 1], Cx));
 				}
-			}
-			if (MinCosOppV != INDEX_NONE)
-			{
-				double ProjT = (Skeleton.Spheres[MinCosOppV].Center - PA).Dot(AB) / ABLenSq;
-				if (ProjT > FMathd::ZeroTolerance && ProjT < 1. - FMathd::ZeroTolerance)
+
+				for (int32 Seg = 1; Seg < ChainNodes.Num() - 1; ++Seg)
 				{
-					Candidates.Add({ Edge, MinCosAngle, ProjT });
+					AddEdge(Cx, ChainNodes[Seg]);
 				}
+
+				ReplaceEdgeOppVert(AsEdge(A, Cx), B, ChainNodes[1]);
+				ReplaceEdgeOppVert(AsEdge(B, Cx), A, ChainNodes[ChainNodes.Num() - 2]);
 			}
-		}
 
-		// Process largest-angle (smallest cosine) triangles first, since split order can affect results
-		Candidates.Sort([](const FCandidateSplit& A, const FCandidateSplit& B) { return A.CosAngle < B.CosAngle; });
+			// Update ClusterNeighbors: remove A-B, add chain edges
+			Skeleton.ClusterNeighbors[A].Remove(B);
+			Skeleton.ClusterNeighbors[B].Remove(A);
+			for (int32 Seg = 0; Seg < ChainNodes.Num() - 1; ++Seg)
+			{
+				AddEdge(ChainNodes[Seg], ChainNodes[Seg + 1]);
+			}
 
-		const int32 StartNewClusters = Skeleton.Spheres.Num();
-		// Apply the edges splits
-		for (const FCandidateSplit& Cand : Candidates)
-		{
-			const int32 A = Cand.Edge.A, B = Cand.Edge.B;
+			OutClusterCanRemapVertices.Add(A);
+			OutClusterCanRemapVertices.Add(B);
+			OutClusterCanRemapVertices.Append(OppVerts);
 
-			// New cluster I, projected onto AB
-			const int32 I = Skeleton.Spheres.Add(LerpSphere(Skeleton.Spheres[A], Skeleton.Spheres[B], Cand.ProjT));
-			Skeleton.ClusterNeighbors.AddDefaulted();
+			// Fill optional split edge info
+			if (OutSplitEdges)
+			{
+				FSplitEdgeInfo& Info = OutSplitEdges->Emplace_GetRef(A, B);
+				for (int32 Seg = 1; Seg < ChainNodes.Num() - 1; ++Seg)
+				{
+					Info.IntermediateNodes.Add(ChainNodes[Seg]);
+				}
+				Info.OppositeVerts.Append(OppVerts);
+			}
 
-			SplitEdge(Cand.Edge, I);
+			EdgeToOppVerts.Remove(Cand.Edge);
 		}
 
+		// Apply accumulated triangle modifications
+		Skeleton.ClusterTriangles.SetNum(Algo::RemoveIf(Skeleton.ClusterTriangles,
+			[&TrianglesToRemove](const FIndex3i& Tri) { return TrianglesToRemove.Contains(Tri); }), EAllowShrinking::No);
+		Skeleton.ClusterTriangles.Append(TrianglesToAdd);
+	}
+
+	// Shared re-clustering helper: reassign vertices near split edges to new intermediate clusters
+	template<typename MeshType>
+	void ReclusterVerticesAfterSplit(
+		UE::Geometry::MedialAxis::FMedialSkeleton& Skeleton,
+		const MeshType& Mesh,
+		TConstArrayView<FSQEM> VtxSQEMs,
+		TConstArrayView<float> VtxAreas,
+		double PosErrorWt,
+		int32 StartNewClusters,
+		const TSet<int32>& ClusterCanRemapVertices)
+	{
+		const double UsePosErrorWt = FMath::Clamp(PosErrorWt, UE_DOUBLE_KINDA_SMALL_NUMBER, 1.0);
+
 		// Error metric that produces equivalent rankings to ComputeSkeleton's vertex-cluster assignment method, using PosErrorWt
 		// Note: area-scaled (not normalized), so only valid for comparing across clusters for the same vertex
-		auto ClusterVertexError = [&Mesh, &VtxAreas, &VtxSQEMs, PosErrorWt](int32 VID, const FSphere& S) -> double
+		auto ClusterVertexError = [&Mesh, &VtxAreas, &VtxSQEMs, UsePosErrorWt](int32 VID, const FSphere& S) -> double
 		{
 			const double PosErr = (Mesh.GetVertex(VID) - S.Center).Length() - S.W;
-			return PosErrorWt * VtxAreas[VID] * PosErr * PosErr + (1. - PosErrorWt) * VtxSQEMs[VID].Eval(S);
+			return UsePosErrorWt * VtxAreas[VID] * PosErr * PosErr + (1. - UsePosErrorWt) * VtxSQEMs[VID].Eval(S);
 		};
 		// Pass through vertices and remap cluster assignments to new vertices, if they are on a marked-as-can-remap cluster
 		ParallelFor(Mesh.MaxVertexID(), [&Skeleton, &ClusterCanRemapVertices, &ClusterVertexError, StartNewClusters](int32 VID)
 		{
 			const int32 CurCluster = Skeleton.VIDtoClusterIndex[VID];
+			if (CurCluster == INDEX_NONE)
+			{
+				return;
+			}
 			if (ClusterCanRemapVertices.Contains(CurCluster))
 			{
 				double BestClusterErr = ClusterVertexError(VID, Skeleton.Spheres[CurCluster]);
@@ -329,11 +341,158 @@ namespace MeshMedialAxisSamplingLocals
 				Skeleton.VIDtoClusterIndex[VID] = BestCluster;
 			}
 		});
+	}
 
-		// Apply accumulated triangle modifications
-		Skeleton.ClusterTriangles.SetNum(Algo::RemoveIf(Skeleton.ClusterTriangles,
-			[&TrianglesToRemove](const FIndex3i& Tri) { return TrianglesToRemove.Contains(Tri); }), EAllowShrinking::No);
-		Skeleton.ClusterTriangles.Append(TrianglesToAdd);
+	// Helper to split skinny triangles -- to give the simplifier more DOFs to collapse away surfaces
+	template<typename MeshType>
+	void SplitThinSkeletonTriEdges(
+		UE::Geometry::MedialAxis::FMedialSkeleton& Skeleton,
+		TArray<FSplitEdgeInfo>& OutSplitEdges,
+		const MeshType& Mesh,
+		TConstArrayView<FSQEM> VtxSQEMs,
+		TConstArrayView<float> VtxAreas,
+		double AngleThresholdDeg,
+		double PosErrorWt)
+	{
+		using namespace UE::Geometry;
+
+		const double UseCosAngleThreshold = FMath::Cos(FMath::Clamp(AngleThresholdDeg, 90. + FMathd::ZeroTolerance, 180. - FMathd::ZeroTolerance) * FMathd::DegToRad);
+
+		int32 StartNewClusters = 0;
+		TSet<int32> ClusterCanRemapVertices;
+		SplitSkeletonEdgesImpl(Skeleton,
+			[&UseCosAngleThreshold](const UE::Geometry::MedialAxis::FMedialSkeleton& InSkeleton, const TMap<FIndex2i, TArray<int32>>& EdgeToOppVerts, TArray<FEdgeSplitRequest>& OutRequests)
+			{
+				// Get candidate edges w/ attached skinny triangles to split
+				struct FCandidateSplit
+				{
+					FIndex2i Edge;
+					double CosAngle, ProjT;
+				};
+				TArray<FCandidateSplit> RawCandidates;
+				for (const auto& Entry : EdgeToOppVerts)
+				{
+					const FIndex2i& Edge = Entry.Key;
+					const TArray<int32>& OppVerts = Entry.Value;
+					const FVector3d& PA = InSkeleton.Spheres[Edge.A].Center;
+					const FVector3d& PB = InSkeleton.Spheres[Edge.B].Center;
+					const FVector3d AB = PB - PA;
+					const double ABLenSq = AB.SquaredLength();
+					if (ABLenSq < FMathd::ZeroTolerance)
+					{
+						continue;
+					}
+					double MinCosAngle = UseCosAngleThreshold;
+					int32 MinCosOppV = INDEX_NONE;
+					for (int32 Cx : OppVerts)
+					{
+						const FVector3d& PC = InSkeleton.Spheres[Cx].Center;
+						const double CosAngleAtCx = Normalized(PA - PC).Dot(Normalized(PB - PC));
+						if (CosAngleAtCx < MinCosAngle)
+						{
+							MinCosAngle = CosAngleAtCx;
+							MinCosOppV = Cx;
+						}
+					}
+					if (MinCosOppV != INDEX_NONE)
+					{
+						double ProjT = (InSkeleton.Spheres[MinCosOppV].Center - PA).Dot(AB) / ABLenSq;
+						if (ProjT > FMathd::ZeroTolerance && ProjT < 1. - FMathd::ZeroTolerance)
+						{
+							RawCandidates.Add({ Edge, MinCosAngle, ProjT });
+						}
+					}
+				}
+
+				// Process largest-angle (smallest cosine) triangles first, since split order can affect results
+				RawCandidates.Sort([](const FCandidateSplit& A, const FCandidateSplit& B) { return A.CosAngle < B.CosAngle; });
+
+				OutRequests.Reserve(RawCandidates.Num());
+				for (const FCandidateSplit& Cand : RawCandidates)
+				{
+					FEdgeSplitRequest& Req = OutRequests.AddDefaulted_GetRef();
+					Req.Edge = Cand.Edge;
+					Req.SplitTs.Add(Cand.ProjT);
+				}
+			},
+			StartNewClusters,
+			ClusterCanRemapVertices,
+			&OutSplitEdges);
+
+		ReclusterVerticesAfterSplit(Skeleton, Mesh, VtxSQEMs, VtxAreas, PosErrorWt, StartNewClusters, ClusterCanRemapVertices);
+	}
+
+	// Collect edges that exceed the target length and produce split requests with uniform T values for refinement
+	void CollectSubdivideCandidates(
+		const UE::Geometry::MedialAxis::FMedialSkeleton& InSkeleton,
+		const TMap<UE::Geometry::FIndex2i, TArray<int32>>& EdgeToOppVerts,
+		TArray<FEdgeSplitRequest>& OutRequests,
+		TFunctionRef<double(const FSphere&, const FSphere&)> GetTargetEdgeLength,
+		bool bSubdivideOnSurfaces)
+	{
+		using namespace UE::Geometry;
+
+		const int32 NumClusters = InSkeleton.Spheres.Num();
+		for (int32 ClusterIdx = 0; ClusterIdx < NumClusters; ++ClusterIdx)
+		{
+			for (int32 NbrIdx : InSkeleton.ClusterNeighbors[ClusterIdx])
+			{
+				if (NbrIdx <= ClusterIdx)
+				{
+					continue;
+				}
+				FIndex2i Edge = AsEdge(ClusterIdx, NbrIdx);
+
+				if (!bSubdivideOnSurfaces)
+				{
+					const TArray<int32>* OppVerts = EdgeToOppVerts.Find(Edge);
+					if (OppVerts && OppVerts->Num() > 0)
+					{
+						continue;
+					}
+				}
+
+				const FSphere& SA = InSkeleton.Spheres[Edge.A];
+				const FSphere& SB = InSkeleton.Spheres[Edge.B];
+				const double EdgeLen = FVector3d::Dist(SA.Center, SB.Center);
+				const double TargetLen = GetTargetEdgeLength(SA, SB);
+
+				if (TargetLen <= 0 || EdgeLen <= TargetLen)
+				{
+					continue;
+				}
+
+				const int32 N = FMath::CeilToInt32(EdgeLen / TargetLen);
+				if (N > 1)
+				{
+					FEdgeSplitRequest& Req = OutRequests.AddDefaulted_GetRef();
+					Req.Edge = Edge;
+					Req.SplitTs.Reserve(N - 1);
+					for (int32 Seg = 1; Seg < N; ++Seg)
+					{
+						Req.SplitTs.Add((double)Seg / (double)N);
+					}
+				}
+			}
+		}
+	}
+
+	// Subdivide skeleton edges, splitting those that exceed the target length. Returns whether any edges were split.
+	bool SubdivideSkeletonEdges(
+		const UE::Geometry::MedialAxis::FSkeletonSubdivider& Subdivider,
+		UE::Geometry::MedialAxis::FMedialSkeleton& Skeleton,
+		int32& OutStartNewClusters,
+		TSet<int32>& OutClusterCanRemapVertices)
+	{
+		SplitSkeletonEdgesImpl(Skeleton,
+			[&Subdivider](const UE::Geometry::MedialAxis::FMedialSkeleton& InSkeleton, const TMap<UE::Geometry::FIndex2i, TArray<int32>>& EdgeToOppVerts, TArray<FEdgeSplitRequest>& OutRequests)
+			{
+				CollectSubdivideCandidates(InSkeleton, EdgeToOppVerts, OutRequests, Subdivider.GetTargetEdgeLength, Subdivider.bSubdivideOnSurfaces);
+			},
+			OutStartNewClusters,
+			OutClusterCanRemapVertices);
+
+		return OutStartNewClusters < Skeleton.Spheres.Num();
 	}
 
 	// Build up standard mesh QEM error terms for the skeleton geometry
@@ -780,14 +939,16 @@ namespace MeshMedialAxisSamplingLocals
 		TMap<FIndex2i, float> SplitSubEdgeErrorFloors;
 		for (const FSplitEdgeInfo& Split : SplitEdges)
 		{
+			checkSlow(Split.IntermediateNodes.Num() == 1);
+			const int32 I = Split.IntermediateNodes[0];
 			FSphere TmpSphere;
 			float ErrorFloor = ComputeMergeError(FIndex2i(Split.A, Split.B), TmpSphere);
 			for (int32 C : Split.OppositeVerts)
 			{
-				ErrorFloor = FMath::Min(ErrorFloor, ComputeMergeError(AsEdge(Split.I, C), TmpSphere) + UE_KINDA_SMALL_NUMBER);
+				ErrorFloor = FMath::Min(ErrorFloor, ComputeMergeError(AsEdge(I, C), TmpSphere) + UE_KINDA_SMALL_NUMBER);
 			}
-			SplitSubEdgeErrorFloors.Add(AsEdge(Split.A, Split.I), ErrorFloor);
-			SplitSubEdgeErrorFloors.Add(AsEdge(Split.B, Split.I), ErrorFloor);
+			SplitSubEdgeErrorFloors.Add(AsEdge(Split.A, I), ErrorFloor);
+			SplitSubEdgeErrorFloors.Add(AsEdge(Split.B, I), ErrorFloor);
 		}
 
 		for (int32 ClusterIdx = 0; ClusterIdx < InitialNumClusters; ++ClusterIdx)
@@ -910,9 +1071,10 @@ namespace MeshMedialAxisSamplingLocals
 			bool bAddedEdges = false;
 			for (const FSplitEdgeInfo& Split : SplitEdges)
 			{
+				checkSlow(Split.IntermediateNodes.Num() == 1);
 				for (int32 EndV : {Split.A, Split.B})
 				{
-					int32 InterV = Split.I;
+					int32 InterV = Split.IntermediateNodes[0];
 					UpdateClusterIndexFromMerge(EndV);
 					UpdateClusterIndexFromMerge(InterV);
 					if (EndV == InterV)
@@ -1784,6 +1946,46 @@ namespace UE::Geometry::MedialAxis
 		}
 	}
 
+	template<typename MeshType>
+	void FSkeletonSubdivider::Subdivide(
+		FMedialSkeleton& Skeleton,
+		const TMeshAABBTree3<MeshType>& MeshBVH,
+		TFunctionRef<FVector3d(int32)> GetVertexNormal,
+		TFunctionRef<int32(int32)> VIDtoWeldVIDFn)
+	{
+		using namespace MeshMedialAxisSamplingLocals;
+		const MeshType& Mesh = *MeshBVH.GetMesh();
+		if (!ensure(Skeleton.IsCompatibleWithMesh(Mesh)))
+		{
+			return;
+		}
+		check(GetTargetEdgeLength);
+		TArray<FVector3f> TriNormals;
+		TArray<float> TriAreas;
+		TArray<FSQEM> VtxSQEMs;
+		TArray<float> VtxAreas;
+		ComputeVertexAreasAndSQEMs(Mesh, VIDtoWeldVIDFn, TriNormals, TriAreas, VtxSQEMs, VtxAreas);
+
+		int32 StartNewClusters = 0;
+		TSet<int32> ClusterCanRemapVertices;
+		if (SubdivideSkeletonEdges(*this, Skeleton, StartNewClusters, ClusterCanRemapVertices))
+		{
+			ReclusterVerticesAfterSplit(Skeleton, Mesh, VtxSQEMs, VtxAreas, ReassignClusterPosErrorWt, StartNewClusters, ClusterCanRemapVertices);
+		}
+	}
+
+	void FSkeletonSubdivider::Subdivide(FMedialSkeleton& Skeleton)
+	{
+		using namespace MeshMedialAxisSamplingLocals;
+		check(GetTargetEdgeLength);
+
+		int32 StartNewClusters = 0;
+		TSet<int32> ClusterCanRemapVertices;
+		SubdivideSkeletonEdges(*this, Skeleton, StartNewClusters, ClusterCanRemapVertices);
+
+		Skeleton.VIDtoClusterIndex.Empty();
+	}
+
 	void FSkeletonViaSampling::RemoveSmallClusters(FMedialSkeleton& Result, TArray<int32>& InOutClusterSize, 
 		TArray<int32>& OutRemapClusters, 
 		TArray<std::atomic<int8>>& OutUpdatedClusters,
@@ -2135,6 +2337,9 @@ namespace UE::Geometry::MedialAxis
 	template GEOMETRYALGORITHMS_API void FSkeletonSimplifier::Simplify<FDynamicMesh3>(FMedialSkeleton& Skeleton, const TMeshAABBTree3<FDynamicMesh3>& MeshBVH, TFunctionRef<bool(int32)> OptionalBoundaryVertexFilter, double ImplicitVertexNormalWeldThreshold);
 	template GEOMETRYALGORITHMS_API void FSkeletonSimplifier::Simplify<FTriangleMeshAdapterd>(FMedialSkeleton& Skeleton, const TMeshAABBTree3<FTriangleMeshAdapterd>& MeshBVH, TFunctionRef<bool(int32)> OptionalBoundaryVertexFilter, double ImplicitVertexNormalWeldThreshold);
 
+	template GEOMETRYALGORITHMS_API void FSkeletonSubdivider::Subdivide<FDynamicMesh3>(FMedialSkeleton& Skeleton, const TMeshAABBTree3<FDynamicMesh3>& MeshBVH, TFunctionRef<FVector3d(int32)> GetVertexNormal, TFunctionRef<int32(int32)> VIDtoWeldVIDFn);
+	template GEOMETRYALGORITHMS_API void FSkeletonSubdivider::Subdivide<FTriangleMeshAdapterd>(FMedialSkeleton& Skeleton, const TMeshAABBTree3<FTriangleMeshAdapterd>& MeshBVH, TFunctionRef<FVector3d(int32)> GetVertexNormal, TFunctionRef<int32(int32)> VIDtoWeldVIDFn);
+
 	template GEOMETRYALGORITHMS_API void FMedialSkeletonToTreeSkeletonOptions::ToHierarchy<FDynamicMesh3>(const FMedialSkeleton& InMedialSkeleton, const TMeshAABBTree3<FDynamicMesh3>* AvoidIntersectingMeshBVH, TArray<int32>& OutParent, TArray<FTransform>& OutTransform, FMedialSkeletonToTreeSkeletonOptions Options, int32 RootIndex, TArray<int32>* OutBoneIndexToMedialIndex, TArray<int32>* OutMedialIndexToBoneIndex);
 	template GEOMETRYALGORITHMS_API void FMedialSkeletonToTreeSkeletonOptions::ToHierarchy<FTriangleMeshAdapterd>(const FMedialSkeleton& InMedialSkeleton, const TMeshAABBTree3<FTriangleMeshAdapterd>* AvoidIntersectingMeshBVH, TArray<int32>& OutParent, TArray<FTransform>& OutTransform, FMedialSkeletonToTreeSkeletonOptions Options, int32 RootIndex, TArray<int32>* OutBoneIndexToMedialIndex, TArray<int32>* OutMedialIndexToBoneIndex);
 
diff --git a/Engine/Plugins/Runtime/GeometryProcessing/Source/GeometryAlgorithms/Public/Skeletonization/MeshMedialAxisSampling.h b/Engine/Plugins/Runtime/GeometryProcessing/Source/GeometryAlgorithms/Public/Skeletonization/MeshMedialAxisSampling.h
--- a/Engine/Plugins/Runtime/GeometryProcessing/Source/GeometryAlgorithms/Public/Skeletonization/MeshMedialAxisSampling.h
+++ b/Engine/Plugins/Runtime/GeometryProcessing/Source/GeometryAlgorithms/Public/Skeletonization/MeshMedialAxisSampling.h
@@ -410,6 +410,26 @@ namespace MedialAxis
 		}
 	};
 
+	struct FSkeletonSubdivider
+	{
+		// Function to compute a new target edge length from the medial spheres on the current edge
+		TFunction<double(const FSphere& A, const FSphere& B)> GetTargetEdgeLength;
+
+		// If true, allow subdividing edges that have associated triangles (also splitting these triangles)
+		bool bSubdivideOnSurfaces = false;
+
+		// When re-clustering for subdivided edges, weight to use for the position error term. Equivalent in meaning to PosErrorWt in FSkeletonViaSampling.
+		double ReassignClusterPosErrorWt = .2;
+
+
+
+		template<typename MeshType>
+		GEOMETRYALGORITHMS_API void Subdivide(FMedialSkeleton& Skeleton, const TMeshAABBTree3<MeshType>& MeshBVH, TFunctionRef<FVector3d(int32)> GetVertexNormal, TFunctionRef<int32(int32)> VIDtoWeldVIDFn);
+
+		// Subdivide a skeleton without referencing an associated mesh. Will clear the medial skeleton's VIDtoClusterIndex map.
+		GEOMETRYALGORITHMS_API void Subdivide(FMedialSkeleton& Skeleton);
+	};
+
 	struct FSkeletonSimplifier
 	{
 
