add a version of skeleton subdivide that computes vertex normals for the caller (ala the simplifier method)
+ clean up simplify methods as well, making the methods consistently match ComputeSkeleton in terms of normals / calling variants

#Synced-CL 52686699

diff --git a/Engine/Plugins/Runtime/GeometryProcessing/Source/GeometryAlgorithms/Private/Skeletonization/MeshMedialAxisSampling.cpp b/Engine/Plugins/Runtime/GeometryProcessing/Source/GeometryAlgorithms/Private/Skeletonization/MeshMedialAxisSampling.cpp
--- a/Engine/Plugins/Runtime/GeometryProcessing/Source/GeometryAlgorithms/Private/Skeletonization/MeshMedialAxisSampling.cpp
+++ b/Engine/Plugins/Runtime/GeometryProcessing/Source/GeometryAlgorithms/Private/Skeletonization/MeshMedialAxisSampling.cpp
@@ -1934,15 +1934,19 @@ namespace UE::Geometry::MedialAxis
 		{
 			return;
 		}
+		TArray<FVector3f> Normals;
+		Normals.SetNumUninitialized(Mesh.MaxVertexID());
+		auto NormalFn = [&Normals](int32 VID) -> FVector3d { return (FVector3d)Normals[VID]; };
 		TArray<int32> WeldIdx = ComputeWeldIdx(Mesh, OptionalBoundaryVertexFilter, ImplicitVertexNormalWeldThreshold);
-		auto NopNormals = [](int32) -> FVector3d { return FVector3d::ZeroVector; };
 		if (WeldIdx.IsEmpty())
 		{
-			Simplify(Skeleton, MeshBVH, NopNormals, [](int32 VID) -> int32 { return VID; });
+			MeshNormals::ComputeVertexNormals(Mesh, Normals);
+			Simplify(Skeleton, MeshBVH, NormalFn, [](int32 VID) -> int32 { return VID; });
 		}
 		else
 		{
-			Simplify(Skeleton, MeshBVH, NopNormals, [&WeldIdx](int32 VID) -> int32 { return WeldIdx[VID]; });
+			MeshNormals::ComputeWeldedVertexNormals(Mesh, [&WeldIdx](int32 VID) -> int32 { return WeldIdx[VID]; }, Normals);
+			Simplify(Skeleton, MeshBVH, NormalFn, [&WeldIdx](int32 VID) -> int32 { return WeldIdx[VID]; });
 		}
 	}
 
@@ -1974,6 +1978,35 @@ namespace UE::Geometry::MedialAxis
 		}
 	}
 
+	template<typename MeshType>
+	void FSkeletonSubdivider::Subdivide(
+		FMedialSkeleton& Skeleton,
+		const TMeshAABBTree3<MeshType>& MeshBVH,
+		TFunctionRef<bool(int32)> OptionalBoundaryVertexFilter,
+		double ImplicitVertexNormalWeldThreshold)
+	{
+		using namespace MeshMedialAxisSamplingLocals;
+		const MeshType& Mesh = *MeshBVH.GetMesh();
+		if (!ensure(Skeleton.IsCompatibleWithMesh(Mesh)))
+		{
+			return;
+		}
+		TArray<FVector3f> Normals;
+		Normals.SetNumUninitialized(Mesh.MaxVertexID());
+		auto NormalFn = [&Normals](int32 VID) -> FVector3d { return (FVector3d)Normals[VID]; };
+		TArray<int32> WeldIdx = ComputeWeldIdx(Mesh, OptionalBoundaryVertexFilter, ImplicitVertexNormalWeldThreshold);
+		if (WeldIdx.IsEmpty())
+		{
+			MeshNormals::ComputeVertexNormals(Mesh, Normals);
+			Subdivide(Skeleton, MeshBVH, NormalFn, [](int32 VID) -> int32 { return VID; });
+		}
+		else
+		{
+			MeshNormals::ComputeWeldedVertexNormals(Mesh, [&WeldIdx](int32 VID) -> int32 { return WeldIdx[VID]; }, Normals);
+			Subdivide(Skeleton, MeshBVH, NormalFn, [&WeldIdx](int32 VID) -> int32 { return WeldIdx[VID]; });
+		}
+	}
+
 	void FSkeletonSubdivider::Subdivide(FMedialSkeleton& Skeleton)
 	{
 		using namespace MeshMedialAxisSamplingLocals;
@@ -2339,6 +2372,8 @@ namespace UE::Geometry::MedialAxis
 
 	template GEOMETRYALGORITHMS_API void FSkeletonSubdivider::Subdivide<FDynamicMesh3>(FMedialSkeleton& Skeleton, const TMeshAABBTree3<FDynamicMesh3>& MeshBVH, TFunctionRef<FVector3d(int32)> GetVertexNormal, TFunctionRef<int32(int32)> VIDtoWeldVIDFn);
 	template GEOMETRYALGORITHMS_API void FSkeletonSubdivider::Subdivide<FTriangleMeshAdapterd>(FMedialSkeleton& Skeleton, const TMeshAABBTree3<FTriangleMeshAdapterd>& MeshBVH, TFunctionRef<FVector3d(int32)> GetVertexNormal, TFunctionRef<int32(int32)> VIDtoWeldVIDFn);
+	template GEOMETRYALGORITHMS_API void FSkeletonSubdivider::Subdivide<FDynamicMesh3>(FMedialSkeleton& Skeleton, const TMeshAABBTree3<FDynamicMesh3>& MeshBVH, TFunctionRef<bool(int32)> OptionalBoundaryVertexFilter, double ImplicitVertexNormalWeldThreshold);
+	template GEOMETRYALGORITHMS_API void FSkeletonSubdivider::Subdivide<FTriangleMeshAdapterd>(FMedialSkeleton& Skeleton, const TMeshAABBTree3<FTriangleMeshAdapterd>& MeshBVH, TFunctionRef<bool(int32)> OptionalBoundaryVertexFilter, double ImplicitVertexNormalWeldThreshold);
 
 	template GEOMETRYALGORITHMS_API void FMedialSkeletonToTreeSkeletonOptions::ToHierarchy<FDynamicMesh3>(const FMedialSkeleton& InMedialSkeleton, const TMeshAABBTree3<FDynamicMesh3>* AvoidIntersectingMeshBVH, TArray<int32>& OutParent, TArray<FTransform>& OutTransform, FMedialSkeletonToTreeSkeletonOptions Options, int32 RootIndex, TArray<int32>* OutBoneIndexToMedialIndex, TArray<int32>* OutMedialIndexToBoneIndex);
 	template GEOMETRYALGORITHMS_API void FMedialSkeletonToTreeSkeletonOptions::ToHierarchy<FTriangleMeshAdapterd>(const FMedialSkeleton& InMedialSkeleton, const TMeshAABBTree3<FTriangleMeshAdapterd>* AvoidIntersectingMeshBVH, TArray<int32>& OutParent, TArray<FTransform>& OutTransform, FMedialSkeletonToTreeSkeletonOptions Options, int32 RootIndex, TArray<int32>* OutBoneIndexToMedialIndex, TArray<int32>* OutMedialIndexToBoneIndex);
diff --git a/Engine/Plugins/Runtime/GeometryProcessing/Source/GeometryAlgorithms/Public/Skeletonization/MeshMedialAxisSampling.h b/Engine/Plugins/Runtime/GeometryProcessing/Source/GeometryAlgorithms/Public/Skeletonization/MeshMedialAxisSampling.h
--- a/Engine/Plugins/Runtime/GeometryProcessing/Source/GeometryAlgorithms/Public/Skeletonization/MeshMedialAxisSampling.h
+++ b/Engine/Plugins/Runtime/GeometryProcessing/Source/GeometryAlgorithms/Public/Skeletonization/MeshMedialAxisSampling.h
@@ -426,6 +426,15 @@ namespace MedialAxis
 		template<typename MeshType>
 		GEOMETRYALGORITHMS_API void Subdivide(FMedialSkeleton& Skeleton, const TMeshAABBTree3<MeshType>& MeshBVH, TFunctionRef<FVector3d(int32)> GetVertexNormal, TFunctionRef<int32(int32)> VIDtoWeldVIDFn);
 
+		template<typename MeshType>
+		GEOMETRYALGORITHMS_API void Subdivide(FMedialSkeleton& Skeleton, const TMeshAABBTree3<MeshType>& MeshBVH, TFunctionRef<bool(int32)> OptionalBoundaryVertexFilter = [](int32)->bool { return false; }, double ImplicitVertexNormalWeldThreshold = UE_DOUBLE_KINDA_SMALL_NUMBER);
+
+		void Subdivide(FMedialSkeleton& Skeleton, const TMeshAABBTree3<FDynamicMesh3>& MeshBVH, double ImplicitVertexNormalWeldThreshold = UE_DOUBLE_KINDA_SMALL_NUMBER)
+		{
+			const FDynamicMesh3& Mesh = *MeshBVH.GetMesh();
+			Subdivide(Skeleton, MeshBVH, [&Mesh](int32 VID)->bool {return Mesh.IsBoundaryVertex(VID);}, ImplicitVertexNormalWeldThreshold);
+		}
+
 		// Subdivide a skeleton without referencing an associated mesh. Will clear the medial skeleton's VIDtoClusterIndex map.
 		GEOMETRYALGORITHMS_API void Subdivide(FMedialSkeleton& Skeleton);
 	};
@@ -482,6 +491,12 @@ namespace MedialAxis
 
 		template<typename MeshType>
 		GEOMETRYALGORITHMS_API void Simplify(FMedialSkeleton& Skeleton, const TMeshAABBTree3<MeshType>& MeshBVH, TFunctionRef<bool(int32)> OptionalBoundaryVertexFilter = [](int32)->bool { return false; }, double ImplicitVertexNormalWeldThreshold = UE_DOUBLE_KINDA_SMALL_NUMBER);
+
+		void Simplify(FMedialSkeleton& Skeleton, const TMeshAABBTree3<FDynamicMesh3>& MeshBVH, double ImplicitVertexNormalWeldThreshold = UE_DOUBLE_KINDA_SMALL_NUMBER)
+		{
+			const FDynamicMesh3& Mesh = *MeshBVH.GetMesh();
+			Simplify(Skeleton, MeshBVH, [&Mesh](int32 VID)->bool {return Mesh.IsBoundaryVertex(VID);}, ImplicitVertexNormalWeldThreshold);
+		}
 	};
 
 
