General attribute support for FMeshSimplification

#Synced-CL 52724062

diff --git a/Engine/Plugins/Runtime/GeometryProcessing/Source/DynamicMesh/DynamicMesh.Build.cs b/Engine/Plugins/Runtime/GeometryProcessing/Source/DynamicMesh/DynamicMesh.Build.cs
--- a/Engine/Plugins/Runtime/GeometryProcessing/Source/DynamicMesh/DynamicMesh.Build.cs
+++ b/Engine/Plugins/Runtime/GeometryProcessing/Source/DynamicMesh/DynamicMesh.Build.cs
@@ -25,7 +25,7 @@ public class DynamicMesh : ModuleRules
 		PrivateDependencyModuleNames.AddRange(
 			new string[] {
 				"Eigen",
-				"MeshConversion",
+				"MeshConversion"
 			}
 		);
 		CppCompileWarningSettings.UnsafeTypeCastWarningLevel = WarningLevel.Error;
diff --git a/Engine/Plugins/Runtime/GeometryProcessing/Source/DynamicMesh/Private/MeshSimplification.cpp b/Engine/Plugins/Runtime/GeometryProcessing/Source/DynamicMesh/Private/MeshSimplification.cpp
--- a/Engine/Plugins/Runtime/GeometryProcessing/Source/DynamicMesh/Private/MeshSimplification.cpp
+++ b/Engine/Plugins/Runtime/GeometryProcessing/Source/DynamicMesh/Private/MeshSimplification.cpp
@@ -7,9 +7,25 @@
 #include "Util/IndexUtil.h"
 #include "Async/ParallelFor.h"
 #include "Templates/UnrealTypeTraits.h"
+#include "MeshQueries.h"
+#include "MeshSimplificationQuadrics.h"
+#include "VectorUtil.h"
 
 using namespace UE::Geometry;
 
+template <typename QuadricErrorType>
+inline bool CanAddVertexQuadric()
+{
+	return true;
+}
+
+template <>
+inline bool CanAddVertexQuadric<FAttrBasedQuadricErrorV2d>()
+{
+	// to support wedges, we need more information of the topology when combining vertex quadrics into an edge quadric
+	return false;
+}
+
 template <typename QuadricErrorType>
 QuadricErrorType TMeshSimplification<QuadricErrorType>::ComputeFaceQuadric(const int tid, FVector3d& nface, FVector3d& c, double& Area) const
 {
@@ -61,6 +77,152 @@ FAttrBasedQuadricErrord TMeshSimplification<FAttrBasedQuadricErrord>::ComputeFac
 	return FQuadricErrorType(p0, p1, p2, n0d, n1d, n2d, nface, c, attrweight);
 }
 
+template<>
+FAttrBasedQuadricErrorV2d TMeshSimplification<FAttrBasedQuadricErrorV2d>::ComputeFaceQuadric(const int TID, FVector3d& Nface, FVector3d& C, double& Area) const
+{
+	// compute the new quadric for this tri.
+	Mesh->GetTriInfo(TID, Nface, Area, C);
+
+	FVector3d P0, P1, P2;
+	Mesh->GetTriVertices(TID, P0, P1, P2);
+
+	FAttrBasedQuadricErrorV2d Q(P0, P1, P2, QuadricOptions);
+
+	FAttrBasedQuadricErrorV2d::FScopedAttributeDataBuilder ScopedAttributeDataBuilder(Q, P0, P1, P2, Nface);
+
+	FDynamicMeshAttributeSet* Attributes = Mesh->Attributes();
+
+	if (QuadricOptions.NormalAttributeWeight > 0.)
+	{
+		// collect normals
+		FIndex3i Tri { FDynamicMesh3::InvalidID, FDynamicMesh3::InvalidID, FDynamicMesh3::InvalidID };
+
+		FVector3f N0, N1, N2;
+		if (NormalOverlay != nullptr)
+		{
+			Tri = NormalOverlay->GetTriangle(TID);
+			if (NormalOverlay->IsSetTriangle(TID))
+			{
+				N0 = NormalOverlay->GetElement(Tri[0]);
+				N1 = NormalOverlay->GetElement(Tri[1]);
+				N2 = NormalOverlay->GetElement(Tri[2]);
+			}
+			else
+			{
+				FVector3f FaceNormal = (FVector3f)Mesh->GetTriNormal(TID);
+				N0 = N1 = N2 = FaceNormal;
+			}
+		}
+		else
+		{
+			Tri = Mesh->GetTriangle(TID);
+			N0 = Mesh->GetVertexNormal(Tri[0]);
+			N1 = Mesh->GetVertexNormal(Tri[1]);
+			N2 = Mesh->GetVertexNormal(Tri[2]);
+		}
+		
+		N0.Normalize();
+		N1.Normalize();
+		N2.Normalize();
+
+		ScopedAttributeDataBuilder.AddWedgeAttribute( Tri, N0, N1, N2, QuadricOptions.NormalAttributeWeight );
+	}
+	
+	FDynamicMeshNormalOverlay* TangentOverlay = Attributes ? Attributes->PrimaryTangents() : nullptr;
+	if (TangentOverlay && QuadricOptions.TangentAttributeWeight > 0.)
+	{
+		if (TangentOverlay->IsSetTriangle(TID))
+		{
+			FIndex3i TangentTri = TangentOverlay->GetTriangle(TID);
+		
+			FVector3f t0 = TangentOverlay->GetElement(TangentTri[0]);
+			FVector3f t1 = TangentOverlay->GetElement(TangentTri[1]);
+			FVector3f t2 = TangentOverlay->GetElement(TangentTri[2]);
+
+			ScopedAttributeDataBuilder.AddWedgeAttribute( TangentTri, t0, t1, t2, QuadricOptions.TangentAttributeWeight );
+		}
+		else
+		{
+			ScopedAttributeDataBuilder.AddMissingWedgeAttribute();
+		}
+	}
+
+	FDynamicMeshNormalOverlay* BiTangentOverlay = Attributes ? Attributes->PrimaryBiTangents() : nullptr;
+	if (BiTangentOverlay && QuadricOptions.TangentAttributeWeight > 0.)
+	{
+		if (BiTangentOverlay->IsSetTriangle(TID))
+		{
+			FIndex3i BiTangentTri = BiTangentOverlay->GetTriangle(TID);
+			FVector3f b0 = BiTangentOverlay->GetElement(BiTangentTri[0]);
+			FVector3f b1 = BiTangentOverlay->GetElement(BiTangentTri[1]);
+			FVector3f b2 = BiTangentOverlay->GetElement(BiTangentTri[2]);
+
+			ScopedAttributeDataBuilder.AddWedgeAttribute( BiTangentTri, b0, b1, b2, QuadricOptions.TangentAttributeWeight );
+		}
+		else
+		{
+			ScopedAttributeDataBuilder.AddMissingWedgeAttribute();
+		}
+	}
+
+	FDynamicMeshColorOverlay* ColorOverlay = Attributes ? Attributes->PrimaryColors() : nullptr;
+	if (ColorOverlay && QuadricOptions.ColorAttributeWeight > 0.)
+	{
+		if (ColorOverlay->IsSetTriangle(TID))
+		{
+			FIndex3i ColorTri = ColorOverlay->GetTriangle(TID);
+			FVector4f c0 = ColorOverlay->GetElement(ColorTri[0]);
+			FVector4f c1 = ColorOverlay->GetElement(ColorTri[1]);
+			FVector4f c2 = ColorOverlay->GetElement(ColorTri[2]);
+			
+			ScopedAttributeDataBuilder.AddWedgeAttribute( ColorTri, FVector3f(c0.X, c0.Y, c0.Z), FVector3f(c1.X, c1.Y, c1.Z), FVector3f(c2.X, c2.Y, c2.Z), QuadricOptions.ColorAttributeWeight );
+		}
+		else
+		{
+			ScopedAttributeDataBuilder.AddMissingWedgeAttribute();
+		}
+	}
+
+	const int NumWeightLayers = Attributes ? Attributes->NumWeightLayers() : 0;
+	if (NumWeightLayers > 0 && QuadricOptions.WeightLayerWeight > 0.)
+	{
+		const FIndex3i VIDs = Mesh->GetTriangle(TID);
+		FVector3f WeightChannels[3] = { FVector3f(0.f), FVector3f(0.f), FVector3f(0.f) }; 
+		for (int WeightLayerIndex = 0; WeightLayerIndex < FMath::Min(NumWeightLayers, 3); ++WeightLayerIndex)
+		{
+			const FDynamicMeshWeightAttribute* WeightAttribute = Attributes->GetWeightLayer(WeightLayerIndex);
+			WeightAttribute->GetValue<float*const>(VIDs[0], &WeightChannels[0][WeightLayerIndex]);
+			WeightAttribute->GetValue<float*const>(VIDs[1], &WeightChannels[1][WeightLayerIndex]);
+			WeightAttribute->GetValue<float*const>(VIDs[2], &WeightChannels[2][WeightLayerIndex]);
+		}
+		ScopedAttributeDataBuilder.AddWedgeAttribute( VIDs, WeightChannels[0], WeightChannels[1], WeightChannels[2], QuadricOptions.WeightLayerWeight );
+	}
+
+	const int NumUVOverlays = Attributes ? Attributes->NumUVLayers() : 0;
+	for (int UVLayerIndex = 0; UVLayerIndex < NumUVOverlays; ++UVLayerIndex)
+	{
+		FDynamicMeshUVOverlay* UVOverlay = Attributes->GetUVLayer(UVLayerIndex);
+		if (UVOverlay && QuadricOptions.TexCoordAttributeWeight > 0.)
+		{
+			if (UVOverlay->IsSetTriangle(TID))
+			{
+				// Collect UVs
+				FIndex3i UVTri = UVOverlay->GetTriangle(TID);
+				FVector2f uv0 = UVOverlay->GetElement(UVTri[0]);
+				FVector2f uv1 = UVOverlay->GetElement(UVTri[1]);
+				FVector2f uv2 = UVOverlay->GetElement(UVTri[2]);
+				ScopedAttributeDataBuilder.AddWedgeAttribute( UVTri, FVector3f(uv0.X, uv0.Y, 0.f), FVector3f(uv1.X, uv1.Y, 0.f), FVector3f(uv2.X, uv2.Y, 0.f), QuadricOptions.TexCoordAttributeWeight  );
+			}
+			else
+			{
+				ScopedAttributeDataBuilder.AddMissingWedgeAttribute();
+			}
+		}
+	}
+
+	return Q;
+}
+
 template <typename QuadricErrorType>
 void TMeshSimplification<QuadricErrorType>::InitializeTriQuadrics()
 {
@@ -174,10 +336,10 @@ void TMeshSimplification<QuadricErrorType>::InitializeVertexQuadrics()
 	{
 		FQuadricErrorType Q = FQuadricErrorType::Zero();
 		
-		
 		double AreaSum = 0;
 		FVector3d Normal(0);
 		bool bNeedsNormals = RegularizeUsesNormals();
+
 		for (int tid : Mesh->VtxTrianglesItr(vid))
 		{
 			AreaSum += triAreas[tid];
@@ -185,11 +347,10 @@ void TMeshSimplification<QuadricErrorType>::InitializeVertexQuadrics()
 			{
 				Normal += triAreas[tid] * triNormals[tid];
 			}
-			Q.Add(triAreas[tid], triQuadrics[tid]);
+			Q.AccumulateFaceQuadric(IndexUtil::FindTriIndex(vid, Mesh->GetTriangle(tid)), triAreas[tid], 1., triQuadrics[tid]);
 		}
-
 		ApplyScalingAndRegularizationToVertexQuadric(Q, vid, AreaSum, Normal);
-
+		
 		vertQuadrics[vid] = Q;
 	}
 }
@@ -233,10 +394,15 @@ template <typename QuadricErrorType>
 QuadricErrorType TMeshSimplification<QuadricErrorType>::AssembleEdgeQuadric(const FDynamicMesh3::FEdge& edge) const
 {
 	//  form standard edge quadric as sum of the vertex quadrics for the edge endpoints
-	QuadricErrorType EdgeQuadric(vertQuadrics[edge.Vert.A], vertQuadrics[edge.Vert.B]);
+	QuadricErrorType EdgeQuadric = MergeVertexQuadricsToEdgeQuadric(edge);
 	
-	if (!bRetainQuadricMemory)
-	{ 
+	if (!EdgeQuadric.CanCollapse())
+	{
+		return EdgeQuadric;
+	}
+
+	if (!bRetainQuadricMemory || !CanAddVertexQuadric<QuadricErrorType>())
+	{
 		// the edge.Tri faces are double counted. Remove one.
 		// if using scaled vertex quadrics, removed the smaller of the two
 		double RemoveScale = 1.;
@@ -247,14 +413,17 @@ QuadricErrorType TMeshSimplification<QuadricErrorType>::AssembleEdgeQuadric(cons
 		const FIndex2i& Tris = edge.Tri;
 		if (Tris.A != FDynamicMesh3::InvalidID)
 		{
-			EdgeQuadric.Add(-triAreas[Tris.A] * RemoveScale, triQuadrics[Tris.A]);
+			const int CornerIdx = IndexUtil::FindEdgeIndexInTri(edge.Vert.A, edge.Vert.B, Mesh->GetTriangle(Tris.A));
+			EdgeQuadric.RemoveFaceQuadric(CornerIdx, triAreas[Tris.A], RemoveScale, triQuadrics[Tris.A]);
 		}
 
 		if (Tris.B != FDynamicMesh3::InvalidID)
 		{
-			EdgeQuadric.Add(-triAreas[Tris.B] * RemoveScale, triQuadrics[Tris.B]);
+			const int CornerIdx = IndexUtil::FindEdgeIndexInTri(edge.Vert.A, edge.Vert.B, Mesh->GetTriangle(Tris.B));
+			EdgeQuadric.RemoveFaceQuadric(CornerIdx, triAreas[Tris.B], RemoveScale, triQuadrics[Tris.B]);
 		}
 	}
+	
 	if (bAllowSeamCollapse)
 	{ 
 		// lambda that adds any adjacent seam quadrics to the edge quadric
@@ -278,6 +447,130 @@ QuadricErrorType TMeshSimplification<QuadricErrorType>::AssembleEdgeQuadric(cons
 }
 
 
+template <typename QuadricErrorType>
+QuadricErrorType TMeshSimplification<QuadricErrorType>::MergeVertexQuadricsToEdgeQuadric(const FDynamicMesh3::FEdge& edge) const
+{
+	// form standard edge quadric as sum of the vertex quadrics for the edge endpoints
+	QuadricErrorType EdgeQuadric(vertQuadrics[edge.Vert.A], vertQuadrics[edge.Vert.B]);
+	
+	return EdgeQuadric;
+}
+
+template <>
+FAttrBasedQuadricErrorV2d TMeshSimplification<FAttrBasedQuadricErrorV2d>::MergeVertexQuadricsToEdgeQuadric(
+	const FDynamicMesh3::FEdge& edge) const
+{
+	
+	// the edge.Tri faces are double counted. Remove one.
+	// if using scaled vertex quadrics, remove the smaller of the two
+	double RemoveScale = 1.;
+	if (!AppliedVertexQuadricScales.IsEmpty())
+	{
+		RemoveScale = FMath::Min(AppliedVertexQuadricScales[edge.Vert.A], AppliedVertexQuadricScales[edge.Vert.B]);
+	}
+
+	const FAttrBasedQuadricErrorV2d* TriangleQuadricA = nullptr;
+	const FAttrBasedQuadricErrorV2d* TriangleQuadricB = nullptr;
+	int LocalEdgeIndexA = IndexConstants::InvalidID;
+	int LocalEdgeIndexB = IndexConstants::InvalidID;
+
+	const FIndex2i Tris = edge.Tri;
+
+	if (Tris.A != FDynamicMesh3::InvalidID)
+	{
+		LocalEdgeIndexA = IndexUtil::FindEdgeIndexInTri(edge.Vert.A, edge.Vert.B, Mesh->GetTriangle(Tris.A));
+		check(LocalEdgeIndexA != IndexConstants::InvalidID);
+		TriangleQuadricA = &triQuadrics[Tris.A];
+	}
+	if (Tris.B != FDynamicMesh3::InvalidID)
+	{
+		LocalEdgeIndexB = IndexUtil::FindEdgeIndexInTri(edge.Vert.A, edge.Vert.B, Mesh->GetTriangle(Tris.B));
+		check(LocalEdgeIndexB != IndexConstants::InvalidID);
+		TriangleQuadricB = &triQuadrics[Tris.B];
+	}
+
+	if (LocalEdgeIndexA != IndexConstants::InvalidID && LocalEdgeIndexB != IndexConstants::InvalidID)
+	{
+		const FIndex3i TriA = Mesh->GetTriangle(Tris.A);
+		const FIndex3i TriB = Mesh->GetTriangle(Tris.B);
+
+		if ( (TriA[LocalEdgeIndexA] != TriB[(LocalEdgeIndexB+1)%3]) || 
+		     (TriB[LocalEdgeIndexB] != TriA[(LocalEdgeIndexA+1)%3]) )
+		{
+			// bowtie
+			return FAttrBasedQuadricErrorV2d(); // not collapsable
+		}
+	}
+	
+	auto CheckOverlayCanCollapse = [&]<typename OverlayType>(const OverlayType& Overlay) -> bool
+	{
+		if (Overlay.IsSetTriangle(Tris.A) && Overlay.IsSetTriangle(Tris.B))
+		{
+			const FIndex3i& OverlayTriA = Overlay.GetTriangle(Tris.A);
+			const FIndex3i& OverlayTriB = Overlay.GetTriangle(Tris.B);
+
+			const bool MatchV0 = OverlayTriA[LocalEdgeIndexA] == OverlayTriB[(LocalEdgeIndexB+1)%3];
+			const bool MatchV1 = OverlayTriB[LocalEdgeIndexB] == OverlayTriA[(LocalEdgeIndexA+1)%3];
+		
+			return MatchV0 == MatchV1; // either match EIDs at both vertices (no seam, or non of them matches, but not just one)
+
+			// TODO check the link condition, to ensure there are no duplicate EIDs in the two sets of attribute one-rings.
+		}
+		return true;
+	};
+
+	if (Tris.B != IndexConstants::InvalidID)
+	{
+		// check for valid seams on each attribute
+
+		if (NormalOverlay != nullptr && !CheckOverlayCanCollapse(*NormalOverlay))
+		{
+			return FAttrBasedQuadricErrorV2d(); // not collapsable
+		}
+
+		FDynamicMeshAttributeSet* Attributes = Mesh->Attributes();
+		if (Attributes)
+		{
+			FDynamicMeshNormalOverlay* TangentOverlay = Attributes->PrimaryTangents();
+			if (TangentOverlay && !CheckOverlayCanCollapse(*TangentOverlay))
+			{
+				return FAttrBasedQuadricErrorV2d();
+			}
+
+			FDynamicMeshNormalOverlay* BiTangentOverlay = Attributes->PrimaryBiTangents();
+			if (BiTangentOverlay && !CheckOverlayCanCollapse(*BiTangentOverlay))
+			{
+				return FAttrBasedQuadricErrorV2d();
+			}
+
+			FDynamicMeshColorOverlay* ColorOverlay = Attributes->PrimaryColors();
+			if (ColorOverlay && !CheckOverlayCanCollapse(*ColorOverlay))
+			{
+				return FAttrBasedQuadricErrorV2d();
+			}
+
+			const int NumUVOverlays = Attributes->NumUVLayers();
+			for (int UVLayerIndex = 0; UVLayerIndex < NumUVOverlays; ++UVLayerIndex)
+			{
+				FDynamicMeshUVOverlay* UVOverlay = Attributes->GetUVLayer(UVLayerIndex);
+				if (UVOverlay && !CheckOverlayCanCollapse(*UVOverlay))
+				{
+					return FAttrBasedQuadricErrorV2d();
+				}
+			}
+		}
+	}
+
+	return FAttrBasedQuadricErrorV2d(vertQuadrics[edge.Vert.A], 
+		                            vertQuadrics[edge.Vert.B],
+									TriangleQuadricA,
+									TriangleQuadricB,
+									LocalEdgeIndexA,
+									LocalEdgeIndexB,
+									RemoveScale);
+	
+}
+
 template <typename QuadricErrorType>
 void TMeshSimplification<QuadricErrorType>::InitializeQueue()
 {
@@ -295,10 +588,14 @@ void TMeshSimplification<QuadricErrorType>::InitializeQueue()
 	for (int eid : Mesh->EdgeIndicesItr())
 	{
 		FDynamicMesh3::FEdge edge = Mesh->GetEdge(eid);
+
 		FQuadricErrorType Q = AssembleEdgeQuadric(edge);
-		FVector3d opt = OptimalPoint(eid, Q, edge.Vert.A, edge.Vert.B);
-		EdgeErrors[eid] = { (float)Q.Evaluate(opt), eid };
-		EdgeQuadrics[eid] = QEdge(eid, Q, opt);
+		if (Q.CanCollapse()) 
+		{
+			FVector3d opt = OptimalPoint(eid, Q, edge.Vert.A, edge.Vert.B);
+			EdgeErrors[eid] = { (float)Q.Evaluate(opt), eid };
+			EdgeQuadrics[eid] = QEdge(eid, Q, opt);
+		}
 	}
 
 	// sorted pq insert is faster, so sort edge errors array and index map
@@ -467,7 +764,6 @@ void TMeshSimplification<QuadricErrorType>::UpdateNeighborhood(const FDynamicMes
 	// --- Update the seam quadrics
 	if (bAllowSeamCollapse)
 	{
-		
 		FIndex2i removedEdges = collapseInfo.RemovedEdges;
 		FIndex2i keptEdges = collapseInfo.KeptEdges;
 
@@ -554,7 +850,7 @@ void TMeshSimplification<QuadricErrorType>::UpdateNeighborhood(const FDynamicMes
 	}
 
 	// --- Update the vertex quadrics
-	if (bRetainQuadricMemory)
+	if (bRetainQuadricMemory && CanAddVertexQuadric<QuadricErrorType>())
 	{ 
 		// Quadric "memory"  the retained vertex quadric is the sum of the two vert quadrics
 		vertQuadrics[kvid] = QuadricErrorType(vertQuadrics[kvid], vertQuadrics[rvid]);
@@ -562,24 +858,22 @@ void TMeshSimplification<QuadricErrorType>::UpdateNeighborhood(const FDynamicMes
 	else
 	{
 		// compute the change in affected face quadrics, and then propagate 
-			// that change to the face adjacent verts.
+		// that change to the face adjacent verts.
 		FVector3d n, c;
-		double NewtriArea;
+		double NewTriArea;
 
 		// Update the triangle areas and quadrics that will have changed
 		for (int tid : Mesh->VtxTrianglesItr(kvid))
 		{
-
-			const double OldtriArea = triAreas[tid];
-			const FQuadricErrorType OldtriQuadric = triQuadrics[tid];
-
+			const double OldTriArea = triAreas[tid];
+			const FQuadricErrorType OldTriQuadric = triQuadrics[tid];
 
 			// compute the new quadric for this tri.
-			FQuadricErrorType NewtriQuadric = ComputeFaceQuadric(tid, n, c, NewtriArea);
+			FQuadricErrorType NewTriQuadric = ComputeFaceQuadric(tid, n, c, NewTriArea);
 
 			// update the arrays that hold the current face area & quadric
-			triAreas[tid] = NewtriArea;
-			triQuadrics[tid] = NewtriQuadric;
+			triAreas[tid] = NewTriArea;
+			triQuadrics[tid] = NewTriQuadric;
 			if (!triNormals.IsEmpty())
 			{
 				triNormals[tid] = n;
@@ -594,8 +888,9 @@ void TMeshSimplification<QuadricErrorType>::UpdateNeighborhood(const FDynamicMes
 
 				// correct the adjacent vertQuadrics
 				double VertexScale = AppliedVertexQuadricScales.IsEmpty() ? 1. : AppliedVertexQuadricScales[tri_vids[i]];
-				vertQuadrics[tri_vids[i]].Add(-OldtriArea * VertexScale, OldtriQuadric); // subtract old quadric
-				vertQuadrics[tri_vids[i]].Add(NewtriArea * VertexScale, NewtriQuadric); // add new quadric
+
+				vertQuadrics[tri_vids[i]].RemoveFaceQuadric(i, OldTriArea, VertexScale, OldTriQuadric); // subtract old quadric
+				vertQuadrics[tri_vids[i]].AccumulateFaceQuadric(i, NewTriArea, VertexScale, NewTriQuadric); // subtract old quadric
 			}
 		}
 
@@ -606,11 +901,13 @@ void TMeshSimplification<QuadricErrorType>::UpdateNeighborhood(const FDynamicMes
 				if (removedTris[i] != FDynamicMesh3::InvalidID)
 				{
 					const double   oldArea = triAreas[removedTris[i]];
-					FQuadricErrorType oldQuadric = triQuadrics[removedTris[i]];
+					const FQuadricErrorType& oldQuadric = triQuadrics[removedTris[i]];
+
+					const int CornerIndex = IndexUtil::FindTriIndex(opposingVerts[i], Mesh->GetTriangle(removedTris[i]));
 
 					// subtract the quadric from the opposing vert
 					double VertexScale = AppliedVertexQuadricScales.IsEmpty() ? 1. : AppliedVertexQuadricScales[opposingVerts[i]];
-					vertQuadrics[opposingVerts[i]].Add(-oldArea * VertexScale, oldQuadric);
+					vertQuadrics[opposingVerts[i]].RemoveFaceQuadric(CornerIndex, oldArea, VertexScale, oldQuadric);
 
 					// zero out the quadric & area for the removed tris.
 					triQuadrics[removedTris[i]] = FQuadricErrorType::Zero();
@@ -618,6 +915,7 @@ void TMeshSimplification<QuadricErrorType>::UpdateNeighborhood(const FDynamicMes
 				}
 			}
 		}
+
 		// Rebuild the quadric for the vert that was retained during the collapse.
 		// NB: in the version with memory this quadric took the value of the edge quadric that collapsed.
 		{
@@ -625,14 +923,16 @@ void TMeshSimplification<QuadricErrorType>::UpdateNeighborhood(const FDynamicMes
 			double AreaSum = 0;
 			FVector3d Normal(0);
 			bool bNeedsNormals = RegularizeUsesNormals();
+
 			for (int tid : Mesh->VtxTrianglesItr(kvid))
 			{
 				AreaSum += triAreas[tid];
+				
 				if (bNeedsNormals)
 				{
 					Normal += triAreas[tid] * triNormals[tid];
 				}
-				vertQuadric.Add(triAreas[tid], triQuadrics[tid]);
+				vertQuadric.AccumulateFaceQuadric(IndexUtil::FindTriIndex(kvid, Mesh->GetTriangle(tid)), triAreas[tid], 1., triQuadrics[tid]);
 			}
 			ApplyScalingAndRegularizationToVertexQuadric(vertQuadric, kvid, AreaSum, Normal);
 			vertQuadrics[kvid] = vertQuadric;
@@ -643,23 +943,32 @@ void TMeshSimplification<QuadricErrorType>::UpdateNeighborhood(const FDynamicMes
 	// NB: this has to follow updating all potential seam quadrics adjacent to kvid 
 	// because an edge quadric gathers seam quadrics adjacent the ends 
 
-	if (bRetainQuadricMemory)
+	if (bRetainQuadricMemory && CanAddVertexQuadric<QuadricErrorType>())
 	{ 
 		for (int eid : Mesh->VtxEdgesItr(kvid))
 		{
 			FDynamicMesh3::FEdge ne = Mesh->GetEdge(eid);
 
 			QuadricErrorType Q = AssembleEdgeQuadric(ne);
-			FVector3d opt = OptimalPoint(eid, Q, ne.Vert.A, ne.Vert.B);
-			float err = (float)Q.Evaluate(opt);
-			EdgeQuadrics[eid] = QEdge(eid, Q, opt);
-			if (EdgeQueue.Contains(eid))
+
+			if (Q.CanCollapse())
 			{
-				EdgeQueue.Update(eid, err);
+				FVector3d opt = OptimalPoint(eid, Q, ne.Vert.A, ne.Vert.B);
+				float err = (float)Q.Evaluate(opt);
+				EdgeQuadrics[eid] = QEdge(eid, Q, opt);
+
+				if (EdgeQueue.Contains(eid))
+				{
+					EdgeQueue.Update(eid, err);
+				}
+				else
+				{
+					EdgeQueue.Insert(eid, err);
+				}
 			}
 			else
 			{
-				EdgeQueue.Insert(eid, err);
+				EdgeQueue.Update(eid, MAX_FLT);
 			}
 		}
 	}
@@ -766,6 +1075,7 @@ void TMeshSimplification<QuadricErrorType>::DoSimplify()
 	{
 		return;
 	}
+
 	InitializeQueue();
 	if (Cancelled())
 	{
@@ -774,8 +1084,8 @@ void TMeshSimplification<QuadricErrorType>::DoSimplify()
 	ProfileEndSetup();
 
 	ProfileBeginOps();
-
 	ProfileBeginCollapse();
+
 	while (EdgeQueue.GetCount() > 0)
 	{
 		// termination criteria
@@ -813,9 +1123,18 @@ void TMeshSimplification<QuadricErrorType>::DoSimplify()
 			return;
 		}
 
-
 		FDynamicMesh3::FEdgeCollapseInfo collapseInfo;
+
+		// CollapseEdge checks whether the edge can be collapsed later, but some conditions need to be met
+		// already during edge quadric construction time.
+		// 
+		if (!EdgeQuadrics[eid].q.CanCollapse())
+		{
+			continue;
+		}
+
 		ESimplificationResult result = CollapseEdge(eid, EdgeQuadrics[eid].collapse_pt, collapseInfo);
+
 		if (result == ESimplificationResult::Ok_Collapsed)
 		{
 			// update the quadrics
@@ -1495,7 +1814,6 @@ ESimplificationResult TMeshSimplification<QuadricErrorType>::CollapseEdge(int ed
 		return ESimplificationResult::Ignored_EdgeIsFullyConstrained;
 	}
 
-
 	// look up verts and tris for this edge
 	if (Mesh->IsEdge(edgeID) == false)
 	{
@@ -1749,6 +2067,7 @@ ESimplificationResult TMeshSimplification<QuadricErrorType>::CollapseEdge(int ed
 			Constraints->ClearVertexConstraint(iCollapse);
 		}
 		OnEdgeCollapse(edgeID, iKeep, iCollapse, collapseInfo);
+		UpdateVertexAttributes(edgeID, iKeep, iCollapse);
 		DoDebugChecks();
 
 		retVal = ESimplificationResult::Ok_Collapsed;
@@ -1926,6 +2245,8 @@ bool TMeshSimplification<QuadricErrorType>::RemoveIsolatedTriangle(int tID)
 
 }
 
+
+
 template <typename QuadricErrorType>
 void TMeshSimplification<QuadricErrorType>::OnEdgeCollapse(int edgeID, int va, int vb, const FDynamicMesh3::FEdgeCollapseInfo& collapseInfo)
 {
@@ -1938,6 +2259,11 @@ void TMeshSimplification<QuadricErrorType>::OnRemoveIsolatedTriangle(int tId)
 	// this is for subclasses
 }
 
+template <>
+void TMeshSimplification<FAttrBasedQuadricErrorV2d>::OnRemoveIsolatedTriangle(int tId)
+{
+}
+
 // Project vertices onto projection target. 
 // We can do projection in parallel if we have .net 
 template <typename QuadricErrorType>
@@ -2007,12 +2333,15 @@ FVector3d TMeshSimplification<QuadricErrorType>::GetProjectedCollapsePosition(in
 	return vNewPos;
 }
 
+template<typename QuadricErrorType>
+void TMeshSimplification<QuadricErrorType>::UpdateVertexAttributes(int edgeID, int va, int vb)
+{
+}
 
 // Custom behavior for FAttrBasedQuadric simplifier.
 template<>
-void TMeshSimplification<FAttrBasedQuadricErrord>::OnEdgeCollapse(int edgeID, int va, int vb, const FDynamicMesh3::FEdgeCollapseInfo& collapseInfo)
+void TMeshSimplification<FAttrBasedQuadricErrord>::UpdateVertexAttributes(int edgeID, int va, int vb)
 {
-
 	// Update the normal
 	FAttrBasedQuadricErrord& Quadric = EdgeQuadrics[edgeID].q;
 	FVector3d collapse_pt = EdgeQuadrics[edgeID].collapse_pt;
@@ -2046,7 +2375,216 @@ void TMeshSimplification<FAttrBasedQuadricErrord>::OnEdgeCollapse(int edgeID, in
 	{
 		Mesh->SetVertexNormal(va, UpdatedNormal);
 	}
+}
+
+// Custom behavior for FAttrBasedQuadric simplifier.
+template<>
+void TMeshSimplification<FAttrBasedQuadricErrorV2d>::UpdateVertexAttributes(int edgeID, int va, int vb)
+{
+	// Update the normal
+	const FAttrBasedQuadricErrorV2d& Q = EdgeQuadrics[edgeID].q;
+	FVector3d collapse_pt = EdgeQuadrics[edgeID].collapse_pt;
+
+	const FAttrBasedQuadricErrorV2d::FAttrArray WedgeAttributes = Q.ComputeAttributes(collapse_pt);
+	const int AttributeCount = Q.GetAttributeCount();
+
+	TArray<int> ElementIdArray;
+	
+	auto UpdateAttribute = [&ElementIdArray, &va, &Q, &WedgeAttributes, this]<typename OverlayT, typename TransformT>(
+		const int WedgeBegin, const int WedgeCount, OverlayT* Overlay, TransformT Transform)
+	{
+		Overlay->GetVertexElements(va, ElementIdArray);
+
+		// walk over wedges
+		for (int j=0; j<WedgeCount; ++j)
+		{
+			if (!Q.IsWedgeActive(WedgeBegin + j))
+			{
+				continue;
+			}
+
+			const FIndex3i& AttrTri = Q.GetWedgeAttributeIndex(WedgeBegin + j);
+			int EID = IndexConstants::InvalidID;
+
+			if (ElementIdArray.Find(AttrTri.B) != INDEX_NONE)
+			{
+				EID = AttrTri.B;
+			}
+			else if (ElementIdArray.Find(AttrTri.A) != INDEX_NONE)
+			{
+				EID = AttrTri.A;
+			}
+
+			if (EID != IndexConstants::InvalidID) 
+			{
+			 	const auto Attribute = Transform(WedgeAttributes[WedgeBegin + j], EID);
+				if (VectorUtil::IsFinite(Attribute))
+				{
+					Overlay->SetElement(EID, Attribute);
+				}
+			}
+		}
+	};
+
+	int WedgeBegin = 0;
+
+	int AttributeIndex = 0;
+
+	if (QuadricOptions.NormalAttributeWeight > 0.)
+	{
+		if (!ensure(AttributeIndex < AttributeCount))
+		{
+			return;
+		}
+
+		const int N = Q.GetWedgeCount(AttributeIndex);
+
+		if (NormalOverlay != nullptr)
+		{
+			UpdateAttribute(WedgeBegin, N, NormalOverlay, [&Q](const FVector3d& Normal, const int EID) -> FVector3f
+			{
+				FVector3f UpdatedNormal(Normal);
+				return UpdatedNormal;
+			} );
+		}
+		else
+		{
+			FVector3f UpdatedNormal(WedgeAttributes[WedgeBegin]);
+			Normalize(UpdatedNormal);
+			Mesh->SetVertexNormal(va, UpdatedNormal);
+		}
+		WedgeBegin += N;
+		AttributeIndex++;
+	}
+
+	FDynamicMeshAttributeSet* Attributes = Mesh->Attributes();
+
+	FDynamicMeshNormalOverlay* TangentOverlay = Attributes ? Attributes->PrimaryTangents() : nullptr;
+	if (TangentOverlay && QuadricOptions.TangentAttributeWeight > 0.)
+	{
+		if (!ensure(AttributeIndex < AttributeCount))
+		{
+			return;
+		}
+
+		// is it normal, that tangents are normal? 
+		const int N = Q.GetWedgeCount(AttributeIndex);
+		
+		if (TangentOverlay)
+		{
+			UpdateAttribute(WedgeBegin, N, TangentOverlay, [](const FVector3d& Tangent, const int EID) -> FVector3f 
+			{
+				FVector3f UpdatedTangent(Tangent);
+				Normalize(UpdatedTangent);
+				return UpdatedTangent;
+			} );
+		}
+
+		WedgeBegin += N;
+		AttributeIndex++;
+	}
+
+	FDynamicMeshNormalOverlay* BiTangentOverlay = Attributes ? Attributes->PrimaryBiTangents() : nullptr;
+	if (BiTangentOverlay && QuadricOptions.TangentAttributeWeight > 0.)
+	{
+		if (!ensure(AttributeIndex < AttributeCount))
+		{
+			return;
+		}
 
+		const int N = Q.GetWedgeCount(AttributeIndex);
+
+		if (BiTangentOverlay)
+		{
+			UpdateAttribute(WedgeBegin, N, BiTangentOverlay, [](const FVector3d& BiTangent, const int EID) -> FVector3f 
+			{
+				FVector3f UpdatedBiTangent(BiTangent);
+				Normalize(UpdatedBiTangent);
+				return UpdatedBiTangent;
+			} );
+		}
+
+		WedgeBegin += N;
+		AttributeIndex++;
+	}
+
+	FDynamicMeshColorOverlay* ColorOverlay = Attributes ? Attributes->PrimaryColors() : nullptr;
+	if (ColorOverlay && QuadricOptions.ColorAttributeWeight > 0.)
+	{
+		if (!ensure(AttributeIndex < AttributeCount))
+		{
+			return;
+		}
+
+		const int N = Q.GetWedgeCount(AttributeIndex);
+
+		if (ColorOverlay)
+		{
+			UpdateAttribute(WedgeBegin, N, ColorOverlay, [&ColorOverlay](const FVector3d& Color, const int EID) -> FVector4f 
+			{
+				FVector4f UpdatedColor(ColorOverlay->GetElement(EID));
+				UpdatedColor.X = static_cast<float>(Color.X);
+				UpdatedColor.Y = static_cast<float>(Color.Y);
+				UpdatedColor.Z = static_cast<float>(Color.Z);
+				return UpdatedColor;
+			} );
+		}
+		WedgeBegin += N;
+		AttributeIndex++;
+	}
+
+	const int NumWeightLayers = Attributes ? Attributes->NumWeightLayers() : 0;
+	if (NumWeightLayers > 0 && QuadricOptions.WeightLayerWeight > 0.)
+	{
+		if (!ensure(AttributeIndex < AttributeCount))
+		{
+			return;
+		}
+
+		const int N = Q.GetWedgeCount(AttributeIndex);
+		ensure(N == 1);
+
+		// walk over wedges
+		for (int j=0; j<N; ++j)
+		{
+			if (!Q.IsWedgeActive(WedgeBegin + j))
+			{
+				continue;
+			}
+
+			FVector3f WeightChannels(WedgeAttributes[WedgeBegin + j]);
+
+			for (int WeightLayerIndex = 0; WeightLayerIndex < FMath::Min(NumWeightLayers, 3); ++WeightLayerIndex)
+			{
+				FDynamicMeshWeightAttribute* WeightAttribute = Attributes->GetWeightLayer(WeightLayerIndex);
+				WeightAttribute->SetValue(va, &WeightChannels[WeightLayerIndex]);
+			}
+		}
+
+		WedgeBegin += N;
+		AttributeIndex++;
+	}
+
+	int NumUVOverlays = Attributes ? Attributes->NumUVLayers() : 0;
+	for (int UVLayerIndex = 0; UVLayerIndex < NumUVOverlays; ++UVLayerIndex)
+	{
+		FDynamicMeshUVOverlay* UVOverlay = Attributes->GetUVLayer(UVLayerIndex);
+		if (UVOverlay && QuadricOptions.TexCoordAttributeWeight > 0.)
+		{
+			if (!ensure(AttributeIndex < AttributeCount))
+			{
+				return;
+			}
+			const int N = Q.GetWedgeCount(AttributeIndex);
+
+			UpdateAttribute(WedgeBegin, N, UVOverlay, [](const FVector3d& UV, const int EID) -> FVector2f 
+			{
+				return FVector2f(static_cast<float>(UV[0]), static_cast<float>(UV[1]));
+			} );
+			WedgeBegin += N;
+			AttributeIndex++;
+		}
+	}
 }
 
 namespace UE
@@ -2059,13 +2597,16 @@ namespace Geometry
 // This is necessary because we have placed most of the templated functions in this .cpp file, instead of the header.
 #if PLATFORM_COMPILER_CLANG && !PLATFORM_MICROSOFT
 template class DYNAMICMESH_API TMeshSimplification< FAttrBasedQuadricErrord >;
+template class DYNAMICMESH_API TMeshSimplification< FAttrBasedQuadricErrorV2d >;
 template class DYNAMICMESH_API TMeshSimplification< FVolPresQuadricErrord >;
 template class DYNAMICMESH_API TMeshSimplification< FQuadricErrord >;
 #else
 template class TMeshSimplification< FAttrBasedQuadricErrord >;
+template class TMeshSimplification< FAttrBasedQuadricErrorV2d >;
 template class TMeshSimplification< FVolPresQuadricErrord >;
 template class TMeshSimplification< FQuadricErrord >;
 #endif
 
 } // end namespace UE::Geometry
-} // end namespace UE
\ No newline at end of file
+} // end namespace UE
+
diff --git a/Engine/Plugins/Runtime/GeometryProcessing/Source/DynamicMesh/Private/MeshSimplificationQuadrics.cpp b/Engine/Plugins/Runtime/GeometryProcessing/Source/DynamicMesh/Private/MeshSimplificationQuadrics.cpp
new file mode 100644
--- /dev/null
+++ b/Engine/Plugins/Runtime/GeometryProcessing/Source/DynamicMesh/Private/MeshSimplificationQuadrics.cpp
@@ -0,0 +1,271 @@
+// Copyright Epic Games, Inc. All Rights Reserved.
+
+#include "MeshSimplificationQuadrics.h"
+
+#if defined(_MSC_VER) && USING_CODE_ANALYSIS
+#pragma warning(push)
+#pragma warning(disable : 6011)
+#pragma warning(disable : 6387)
+#pragma warning(disable : 6313)
+#pragma warning(disable : 6294)
+#endif
+PRAGMA_DEFAULT_VISIBILITY_START
+THIRD_PARTY_INCLUDES_START
+#include <Eigen/Dense>
+THIRD_PARTY_INCLUDES_END
+PRAGMA_DEFAULT_VISIBILITY_END
+#if defined(_MSC_VER) && USING_CODE_ANALYSIS
+#pragma warning(pop)
+#endif
+
+namespace UE
+{
+
+namespace Geometry
+{
+
+template <typename RealType>
+void UpdateResidual(double& MaxRes, const RealType SM[6], const TVector<double>& X, const TVector<double>& B)
+{
+	const TVector<double> Ax = TQuadricError<RealType>::MultiplySymmetricMatrix(SM, X);
+	const double ResSqr = (Ax - B) | (Ax - B);
+
+	MaxRes = FMath::Max(ResSqr, MaxRes);
+}
+
+template <typename RealType>
+struct SolverBuiltin
+{
+	SolverBuiltin(const RealType SM[6])
+	{
+		bIsValid = TQuadricError<RealType>::InvertSymmetricMatrix(SM, InvSM);
+	}
+
+	bool IsValid() const 
+	{
+		return bIsValid;
+	}
+
+	bool Solve(TVector<RealType>& X, const TVector<RealType>& Rhs) const
+	{
+		if (bIsValid)
+		{
+			X = TQuadricError<RealType>::MultiplySymmetricMatrix(InvSM, Rhs);
+		}
+		return bIsValid;
+	}
+
+	RealType InvSM[6];
+	bool bIsValid { false };
+};
+
+
+template <typename RealType>
+struct SolverEigen
+{
+	SolverEigen(const RealType SM[6])
+	{
+		M << SM[0], SM[1], SM[2],
+		     SM[1], SM[3], SM[4],
+		     SM[2], SM[4], SM[5];
+
+		ldlt.compute(M);
+
+		bIsValid = true;
+
+		if (ldlt.info() != Eigen::Success)
+		{
+			bIsValid = false;
+		}
+
+		if(ldlt.vectorD().cwiseAbs().minCoeff() < 1e-12)
+		{
+			bIsValid = false;
+		}
+	}
+
+	bool IsValid() const 
+	{
+		return bIsValid;
+	}
+
+	bool Solve(TVector<RealType>& X, const TVector<RealType>& Rhs) const
+	{
+		if (bIsValid)
+		{
+			Eigen::Vector3d x = ldlt.solve(Eigen::Vector3d(Rhs[0], Rhs[1], Rhs[2]));
+			X.X = x(0);
+			X.Y = x(1);
+			X.Z = x(2);
+		}
+		return bIsValid;
+	}
+
+	Eigen::Matrix3d M;
+	Eigen::LDLT<Eigen::Matrix3d> ldlt;
+	bool bIsValid { false };
+};
+
+#ifdef UE_USE_INTERNAL_MATRIX_SOLVER
+// requires dependencies to Source/Developer, so currently disabled. It would require making the solvers accessible
+// to runtime plugins
+
+template <typename RealType>
+struct SolverLUP
+{
+	SolverLUP(const RealType SM[6])
+	{
+		M[0] = SM[0];
+		M[1] = SM[1];
+		M[2] = SM[2];
+
+		M[3] = SM[1];
+		M[4] = SM[3];
+		M[5] = SM[4];
+
+		M[6] = SM[2];
+		M[7] = SM[4];
+		M[8] = SM[5];
+
+		FMemory::Memcpy( LU, M );
+		bIsValid = LUPFactorize( LU, Pivot, 3, 1.e-12 );
+	}
+
+	bool IsValid() const 
+	{
+		return bIsValid;
+	}
+
+	bool Solve(TVector<RealType>& X, const TVector<RealType>& Rhs) const
+	{
+		RealType x[3];
+		RealType b[3] = { Rhs[0], Rhs[1], Rhs[2] };
+
+		if (!bIsValid)
+		{
+			return false;
+		}
+
+		if (LUPSolveIterate( M, LU, Pivot, 3, b, x ))
+		{
+			X[0] = x[0];
+			X[1] = x[1];
+			X[2] = x[2];
+			return true;
+		}
+
+		return false;
+	}
+
+	RealType M[9];
+	RealType LU[9];
+	uint32 Pivot[3];
+	bool bIsValid { false };
+};
+
+#endif
+
+template <typename RealType>
+template <typename SolverType>
+bool TAttrBasedQuadricErrorV2<RealType>::OptimalPointImpl(
+	UE::Math::TVector<RealType>& OutResult) const
+{
+	const TAttrBasedQuadricErrorV2<RealType>& Q = *this;
+
+	RealType SM[6] = { Q.Axx, Q.Axy, Q.Axz, Q.Ayy, Q.Ayz, Q.Azz };
+	TVector<RealType> Rhs(-Q.bx, -Q.by, -Q.bz);
+
+	for (const auto& AttrWedgeQuadricData : Q.AttrWedgeQuadricDatas)
+	{
+		if (AttrWedgeQuadricData.TotalAreaTerm < UE_DOUBLE_SMALL_NUMBER)
+		{
+			continue;
+		}
+
+		const RealType InvAlpha = 1.0 / AttrWedgeQuadricData.TotalAreaTerm;
+		const TVector<RealType> (&G)[3] = AttrWedgeQuadricData.G;
+		const TVector<RealType>& D = AttrWedgeQuadricData.D;
+
+		// C - alpha^{-1} B B^T
+		// rhs = -b1 + alpha^{-1} B b2 
+		for (int i=0; i<3; ++i)
+		{
+			SM[0] -= InvAlpha * G[i].X * G[i].X;
+			SM[1] -= InvAlpha * G[i].X * G[i].Y;
+			SM[2] -= InvAlpha * G[i].X * G[i].Z;
+
+			SM[3] -= InvAlpha * G[i].Y * G[i].Y;
+			SM[4] -= InvAlpha * G[i].Y * G[i].Z;
+
+			SM[5] -= InvAlpha * G[i].Z * G[i].Z;
+
+			Rhs.X += InvAlpha * G[i].X * D[i];
+			Rhs.Y += InvAlpha * G[i].Y * D[i];
+			Rhs.Z += InvAlpha * G[i].Z * D[i];
+		}
+	}
+
+	SolverType Solver(SM);
+	bool SolveValid = Solver.Solve(OutResult, Rhs);
+
+	if (!SolveValid)
+	{
+		return false;
+	}
+	
+	// Adjust with the volume constraint.
+	// See: http://hhoppe.com/newqem.pdf or https://www.cc.gatech.edu/~turk/my_papers/memless_vis98.pdf
+
+	// Compute the effect of the volumetric constraint.
+	TVector<RealType> Ainv_g;
+	SolveValid = Solver.Solve(Ainv_g, Q.GVol);
+
+	if (SolveValid)
+	{
+		const RealType gt_Ainv_g = Q.GVol.Dot(Ainv_g);
+
+		const RealType gt_unopt = Q.GVol.Dot(OutResult);
+		const RealType lambda = (Q.DVol - gt_unopt) / gt_Ainv_g;
+
+		// squared magnitude of the update
+		const RealType magnitude_sqr = lambda * lambda * Ainv_g.SizeSquared();
+
+		constexpr double VolThreshold = 0.5;
+
+		if (magnitude_sqr < VolThreshold * Q.Area) 
+		{
+			// also, guard against outliers in which the new point is too far away relative to the local neighborhood size
+			// in case of division by zero, lambda is NaN and this will evaluate to false
+
+			OutResult += lambda * Ainv_g;
+		}
+		else
+		{
+			// solve for position succeeded, but volume constraint rejected. still accept the solve
+			return true;
+		}
+	}
+	
+	return SolveValid;
+}
+
+
+template <typename RealType>
+bool TAttrBasedQuadricErrorV2<RealType>::OptimalPoint(
+	TVector<RealType>& OutResult, RealType MinThresh) const
+{
+	// return OptimalPointImpl<SolverBuiltin<RealType>>(OutResult);
+	return OptimalPointImpl<SolverEigen<RealType>>(OutResult);
+
+	// MatrixUtils.hh LU solver seems to be same accuracy as Eigen ldlt
+	// return OptimalPointImpl<SolverLUP<RealType>>(OutResult);	
+}
+
+#if PLATFORM_COMPILER_CLANG && !PLATFORM_MICROSOFT
+template class DYNAMICMESH_API TAttrBasedQuadricErrorV2<double>;
+#else
+template class TAttrBasedQuadricErrorV2<double>;
+#endif
+
+} // namespace Geometry
+} // namespace UE
diff --git a/Engine/Plugins/Runtime/GeometryProcessing/Source/DynamicMesh/Public/MeshSimplification.h b/Engine/Plugins/Runtime/GeometryProcessing/Source/DynamicMesh/Public/MeshSimplification.h
--- a/Engine/Plugins/Runtime/GeometryProcessing/Source/DynamicMesh/Public/MeshSimplification.h
+++ b/Engine/Plugins/Runtime/GeometryProcessing/Source/DynamicMesh/Public/MeshSimplification.h
@@ -10,7 +10,6 @@
 
 #include "DynamicMesh/DynamicMeshAttributeSet.h"
 
-
 namespace UE
 {
 namespace Geometry
@@ -42,6 +41,20 @@ enum class ERegularizeSimplificationMethod
 	NormalLine
 };
 
+
+template <typename T, typename = void>
+struct TQuadricOptionsHelper
+{ 
+	 struct FOptions {}; 
+};
+
+// Partial specialization (kicks in when T::FOptions exists)
+template <typename T>
+struct TQuadricOptionsHelper<T, std::void_t<typename T::FOptions>>
+{
+    using FOptions = typename T::FOptions;
+};
+
 /**
  * Implementation of Garland & Heckbert Quadric Error Metric (QEM) Triangle Mesh Simplification
  */
@@ -54,6 +67,7 @@ public:
 	typedef typename FQuadricErrorType::ScalarType RealType;
 	typedef TQuadricError<RealType>                FSeamQuadricType;
 
+	using FQuadricOptions = TQuadricOptionsHelper<QuadricErrorType>::FOptions;
 	
 	enum class ESimplificationCollapseModes
 	{
@@ -126,6 +140,8 @@ public:
 	/** Tolerance to use in geometric error checking */
 	double GeometricErrorTolerance = 0.0;
 
+	FQuadricOptions QuadricOptions;
+
 	// Custom scales to apply to *measured* errors and edge lengths. Note this is in effect the inverse of scaling the tolerances / target lengths.
 	TFunction<double(const FDynamicMesh3&, int VertexA, int VertexB)> CustomGeometricErrorScaleF;
 	TFunction<double(const FDynamicMesh3&, int VertexA, int VertexB)> CustomEdgeLengthScaleF;
@@ -267,6 +283,9 @@ protected:
 	// uses pre-computed vertex, face and seam quadrics to construct the edge quadric.
 	FQuadricErrorType AssembleEdgeQuadric(const FDynamicMesh3::FEdge& edge) const;
 
+	FQuadricErrorType MergeVertexQuadricsToEdgeQuadric(const FDynamicMesh3::FEdge& edge) const;
+		
+	void UpdateVertexAttributes(int edgeID, int va, int vb );
 	
 	// internal class for priority queue
 	struct QEdge 
@@ -494,11 +513,11 @@ private:
 
 };
 
-
 // The simplifier
 typedef TMeshSimplification< FAttrBasedQuadricErrord >  FAttrMeshSimplification;
-typedef TMeshSimplification < FVolPresQuadricErrord >   FVolPresMeshSimplification;
+typedef TMeshSimplification< FVolPresQuadricErrord >    FVolPresMeshSimplification;
 typedef TMeshSimplification< FQuadricErrord >           FQEMSimplification;
 
+
 } // end namespace UE::Geometry
 } // end namespace UE
diff --git a/Engine/Plugins/Runtime/GeometryProcessing/Source/DynamicMesh/Public/MeshSimplificationQuadrics.h b/Engine/Plugins/Runtime/GeometryProcessing/Source/DynamicMesh/Public/MeshSimplificationQuadrics.h
new file mode 100644
--- /dev/null
+++ b/Engine/Plugins/Runtime/GeometryProcessing/Source/DynamicMesh/Public/MeshSimplificationQuadrics.h
@@ -0,0 +1,962 @@
+// Copyright Epic Games, Inc. All Rights Reserved.
+
+#pragma once
+
+#include "MathUtil.h"
+#include "IndexTypes.h"
+#include "MatrixTypes.h"
+#include "VectorTypes.h"
+#include "Math/UnrealMathUtility.h"
+#include "QuadricError.h"
+#include "VectorUtil.h"
+
+namespace UE
+{
+namespace Geometry
+{
+
+using namespace UE::Math;
+
+
+// Configurable attributes with separate topology (wedges)
+// Volume preservation
+// Plane (Garland/Heckbert) and triangle quadrics (Lindstrom/Turk), including probabilistic versions (Trettner/Kobbelt)
+//
+template<typename RealType>
+class TAttrBasedQuadricErrorV2 : public TQuadricError<RealType>
+{
+protected:
+	using BaseStruct = TQuadricError<RealType>;
+
+	struct FAttrWedgeQuadricData
+	{
+		TVector<RealType> G[3]; 
+		TVector<RealType> D;
+		FIndex3i          AttrTri;
+		RealType          TotalAreaTerm { 0. }; // sum of triangles contributing to wedge, or sum squared
+		RealType          AttrWeight    { 0. };
+
+		void Scale(const RealType Wt)
+		{
+			G[0] *= Wt;
+			G[1] *= Wt;
+			G[2] *= Wt;
+			D    *= Wt;
+		}
+		
+		void AccumulateQuadricData(const FAttrWedgeQuadricData& FaceQuadricData)
+		{
+			G[0] += FaceQuadricData.G[0];
+			G[1] += FaceQuadricData.G[1];
+			G[2] += FaceQuadricData.G[2];
+			D    += FaceQuadricData.D;
+			TotalAreaTerm += FaceQuadricData.TotalAreaTerm;
+			AttrWeight = FaceQuadricData.AttrWeight;
+
+			ensure(TotalAreaTerm >= -UE_DOUBLE_KINDA_SMALL_NUMBER);
+		}
+
+		void RemoveQuadricData(const FAttrWedgeQuadricData& FaceQuadricData)
+		{
+			G[0] -= FaceQuadricData.G[0];
+			G[1] -= FaceQuadricData.G[1];
+			G[2] -= FaceQuadricData.G[2];
+			D    -= FaceQuadricData.D;
+			TotalAreaTerm -= FaceQuadricData.TotalAreaTerm;
+			AttrWeight = FaceQuadricData.AttrWeight;
+
+			ensure(TotalAreaTerm >= -UE_DOUBLE_KINDA_SMALL_NUMBER);
+			TotalAreaTerm = FMath::Max(0., TotalAreaTerm);
+		}
+
+		void AccumulateToPositionQuadric(TAttrBasedQuadricErrorV2& FaceQuadric) const
+		{
+			const RealType AreaTerm = FaceQuadric.TotalAreaTerm;
+
+			ensure(AreaTerm >= -UE_DOUBLE_KINDA_SMALL_NUMBER);
+
+			for (int i=0; i<3; ++i)
+			{
+				FaceQuadric.Axx += AreaTerm * G[i].X * G[i].X;
+				FaceQuadric.Axy += AreaTerm * G[i].X * G[i].Y;
+				FaceQuadric.Axz += AreaTerm * G[i].X * G[i].Z;
+
+				FaceQuadric.Ayy += AreaTerm * G[i].Y * G[i].Y;
+				FaceQuadric.Ayz += AreaTerm * G[i].Y * G[i].Z;
+
+				FaceQuadric.Azz += AreaTerm * G[i].Z * G[i].Z;
+
+				FaceQuadric.bx  += AreaTerm * G[i].X * D[i];
+				FaceQuadric.by  += AreaTerm * G[i].Y * D[i];
+				FaceQuadric.bz  += AreaTerm * G[i].Z * D[i];
+
+				FaceQuadric.c   += AreaTerm * D[i] * D[i];
+			}
+		}
+	};
+
+	constexpr static uint32 InlineAttrElements = 6; 
+	
+	// 6 enough to hold Normals, Tangents (2x), VertexColor, 2 x UV
+	TArray<FAttrWedgeQuadricData, TInlineAllocator<InlineAttrElements>> AttrWedgeQuadricDatas;
+	TArray<int, TInlineAllocator<InlineAttrElements>>                   WedgeCount;
+
+public:
+
+	using FAttrArray = TArray<TVector<RealType>, TInlineAllocator<InlineAttrElements>>;
+
+	// Helper class to register attributes to the quadric on construction
+	class FScopedAttributeDataBuilder
+	{
+	private:
+		void CalcGandD(const FVector3f& W0, const FVector3f& W1, const FVector3f& W2, TVector<RealType> G[3], TVector<RealType>& D) const
+		{
+			if (IsValid) 
+			{
+				for (int i = 0; i < 3; ++i)
+				{
+					const FVector3d Attr = FVector3d(W0[i], W1[i], W2[i]);
+					const FVector3d DataVec(Attr[2] - Attr[0], Attr[1] - Attr[0], 0.);
+					G[i] = CoBasisMatrix * DataVec;
+					D[i] = Attr[0] - P0.Dot(G[i]);
+				}
+			}
+			else
+			{
+				for (int i = 0; i < 3; ++i)
+				{
+					G[i] = TVector<RealType>::Zero();
+					D[i] = (W0[i] + W1[i] + W2[i]) / 3.; 
+				}
+			}
+		}
+
+	public:
+		FScopedAttributeDataBuilder(
+			TAttrBasedQuadricErrorV2& InQuadric,
+			const TVector<RealType>& InP0, 
+		    const TVector<RealType>& InP1, 
+			const TVector<RealType>& InP2,
+			const TVector<RealType>& InN)
+			: Quadric(InQuadric), P0(InP0)
+		{
+			const TMatrix3<RealType> BasisMatrix(InP2-InP0, InP1-InP0, InN, true);
+			const RealType Det = BasisMatrix.Determinant();
+			if (FMath::Abs(Det) > 1.e-8)
+			{
+				CoBasisMatrix = BasisMatrix.Inverse();
+				IsValid = true;
+			}
+		}
+
+		void AddWedgeAttribute(const FIndex3i& AttrTri, const FVector3f& W0, const FVector3f& W1, const FVector3f& W2, const RealType AttrWeight)
+		{
+			FAttrWedgeQuadricData AttrWedgeQuadricData;
+			CalcGandD(W0, W1, W2, AttrWedgeQuadricData.G, AttrWedgeQuadricData.D);
+
+			AttrWedgeQuadricData.TotalAreaTerm = Quadric.TotalAreaTerm;
+			AttrWedgeQuadricData.AttrWeight = AttrWeight;
+
+			AttrWedgeQuadricData.G[0] *= AttrWeight;
+			AttrWedgeQuadricData.G[1] *= AttrWeight;
+			AttrWedgeQuadricData.G[2] *= AttrWeight;
+			AttrWedgeQuadricData.D    *= AttrWeight;
+
+			AttrWedgeQuadricData.AttrWeight = AttrWeight;
+
+			AttrWedgeQuadricData.AccumulateToPositionQuadric(Quadric);
+
+			AttrWedgeQuadricData.G[0] *= Quadric.TotalAreaTerm;
+			AttrWedgeQuadricData.G[1] *= Quadric.TotalAreaTerm;
+			AttrWedgeQuadricData.G[2] *= Quadric.TotalAreaTerm;
+			AttrWedgeQuadricData.D    *= Quadric.TotalAreaTerm;
+
+			AttrWedgeQuadricData.AttrTri = AttrTri;
+			Quadric.AttrWedgeQuadricDatas.Add(AttrWedgeQuadricData);
+			Quadric.WedgeCount.Add(1);
+		}
+
+		void AddMissingWedgeAttribute()
+		{
+			Quadric.WedgeCount.Add(0);
+		}
+
+	private:
+		TAttrBasedQuadricErrorV2& Quadric; // face quadric
+		const TVector<RealType>& P0; 
+		FMatrix3d CoBasisMatrix;
+		bool IsValid { false };
+	};
+
+	enum class EQuadricVariant
+	{
+		PlaneQuadric,      // classic Garland-Heckbert squared distance to plane
+		TriangleQuadric    // Lindstrom-Turk volumetric error 
+	};
+
+	struct FOptions
+	{
+		EQuadricVariant QuadricVariant { EQuadricVariant::PlaneQuadric };
+		RealType SigmaN                { 0. }; 
+
+		// currently mimicking FQuadricSimplifierMeshReduction, see ReduceMeshDescription.
+		// consider scaling behaviour (geometry error scales with AreaTerm * ObjectScale^2),
+		// attribute term scales with AreaTerm only.
+		//
+		double NormalAttributeWeight   { 16.0 };
+		double TangentAttributeWeight  { 0.1 };
+		double ColorAttributeWeight    { 0.1 };
+		double TexCoordAttributeWeight { 0.5 };
+		double WeightLayerWeight       { 0.0 };
+	};
+
+	// empty quadric
+	TAttrBasedQuadricErrorV2() 
+		: BaseStruct()
+	{
+	}
+
+	// initialize a face quadric
+	TAttrBasedQuadricErrorV2(const TVector<RealType>& P0, 
+		                     const TVector<RealType>& P1, 
+	 						 const TVector<RealType>& P2,
+							 const FOptions& Options)
+							 : BaseStruct()
+	{
+		switch (Options.QuadricVariant)
+		{
+			case EQuadricVariant::PlaneQuadric:
+				MakeProbabilisticPlaneQuadric(P0, P1, P2, Options.SigmaN );
+				break;
+			case EQuadricVariant::TriangleQuadric:
+				MakeProbabilisticTriangleQuadric(P0, P1, P2, Options.SigmaN );
+				break;
+			default:
+				ensureMsgf(false, TEXT("Invalid quadric variant. Must be either plane or triangle quadric."));
+		}
+
+		const TVector<RealType> Normal = VectorUtil::Normal(P0, P1, P2);
+		GVol = Normal;
+		DVol = ( Normal | P0 );
+	}
+
+	// this should never be called. addition of vertex quadrics is not supported because it doesn't work for attribute wedges
+	TAttrBasedQuadricErrorV2(const TAttrBasedQuadricErrorV2& Q0, const TAttrBasedQuadricErrorV2& Q1)
+		: BaseStruct(Q0, Q1)
+	{
+		check(false);
+		Area = Q0.Area + Q1.Area;
+		GVol = Q0.GVol + Q1.GVol;
+		DVol = Q0.DVol + Q1.DVol;
+		TotalAreaTerm = Q0.TotalAreaTerm + Q1.TotalAreaTerm;
+	}
+
+	// initialize to an edge quadric given the two vertex quadrics at the end points of the edge
+	// and (up to two) face quadrics corresponding to the overlap of both
+	// 
+	// the local edge indices correspond to the local edge index [0..3) of the shared edge within A and B respectively.
+	// 
+	TAttrBasedQuadricErrorV2(const TAttrBasedQuadricErrorV2& VertexQuadric0, 
+		                     const TAttrBasedQuadricErrorV2& VertexQuadric1,
+							 const TAttrBasedQuadricErrorV2* TriangleQuadricA,
+							 const TAttrBasedQuadricErrorV2* TriangleQuadricB,
+							 const int LocalEdgeIndexA,
+							 const int LocalEdgeIndexB,
+							 const RealType RemoveScale)
+
+		: BaseStruct(VertexQuadric0, VertexQuadric1)
+	{
+		Area = VertexQuadric0.Area + VertexQuadric1.Area;
+		GVol = VertexQuadric0.GVol + VertexQuadric1.GVol;
+		DVol = VertexQuadric0.DVol + VertexQuadric1.DVol;
+		TotalAreaTerm = VertexQuadric0.TotalAreaTerm + VertexQuadric1.TotalAreaTerm;
+ 
+		checkSlow(VertexQuadric0.WedgeCount.Num() == VertexQuadric1.WedgeCount.Num());
+
+		bCanCollapse = MergeVertexQuadricsToEdgeQuadric(VertexQuadric0, VertexQuadric1, TriangleQuadricA, TriangleQuadricB, LocalEdgeIndexA, LocalEdgeIndexB, RemoveScale);
+		checkSlow(!bCanCollapse || (WedgeCount.Num() == VertexQuadric0.WedgeCount.Num()));
+	}
+
+	inline static TAttrBasedQuadricErrorV2 Zero() { return TAttrBasedQuadricErrorV2(); };
+
+	// If this is a vertex, accumulate a face quadric on to it.
+	// \param CornerIdx local vertex index [0..3) of the center vertex within the triangle of the face quadric 
+	// 
+	void AccumulateFaceQuadric(const int CornerIdx, RealType TriArea, RealType Weight, const TAttrBasedQuadricErrorV2& FaceQuadric);
+	
+	// If this is a vertex, remove a face quadric from it (undoing accumulate face quadric)
+	// \param CornerIdx local vertex index [0..3) of the center vertex within the triangle of the face quadric 
+	//
+	void RemoveFaceQuadric(const int CornerIdx, RealType TriArea, RealType RemoveScale, const TAttrBasedQuadricErrorV2& FaceQuadric);
+	
+	void Scale(RealType Wt)
+	{
+		BaseStruct::Scale(Wt);
+
+		for (FAttrWedgeQuadricData& AttrWedgeQuadricData : AttrWedgeQuadricDatas)
+		{
+			AttrWedgeQuadricData.Scale(Wt);
+		}
+	}
+	
+	/**
+	* The optimal point minimizing the quadric error with respect to a volume conserving constraint
+	*/
+	bool OptimalPoint(UE::Math::TVector<RealType>& OutResult, RealType MinThresh = 1000.0*TMathUtil<RealType>::Epsilon) const;
+	
+	// Scale-correct attributes 
+	FAttrArray ComputeAttributes(const TVector<RealType>& P) const;
+
+	RealType Evaluate(const TVector<RealType>& P) const
+	{
+		FAttrArray WedgeAttr = ComputeScaledAttributes(P);
+		return Evaluate(P, WedgeAttr);
+	}
+
+	// For each attribute, checks whether the area of all wedges adds up to the total area of the quadric.
+	bool ValidateArea() const;
+
+	// actual area, not the scaling term
+	RealType GetArea() const
+	{
+		return Area;
+	}
+
+	// is this a valid edge-collapse quadric?
+	bool CanCollapse() const { return bCanCollapse; }
+	
+	int GetAttributeCount() const 
+	{ 
+		return WedgeCount.Num(); 
+	}
+
+	int GetWedgeCount(int AttributeIndex) const
+	{
+		return WedgeCount[AttributeIndex];
+	}
+ 
+	bool IsWedgeActive(int WedgeIndex) const
+	{
+		return AttrWedgeQuadricDatas[WedgeIndex].TotalAreaTerm > 0.;
+	}
+
+	const FIndex3i& GetWedgeAttributeIndex(int WedgeIndex) const
+	{
+		return AttrWedgeQuadricDatas[WedgeIndex].AttrTri;
+	}
+
+protected:
+
+	template <typename SolverType>
+	bool OptimalPointImpl(TVector<RealType>& OutResult) const;
+	
+	FAttrArray ComputeScaledAttributes(const TVector<RealType>& P) const;
+
+	RealType Evaluate(const TVector<RealType>& P, const FAttrArray& ScaledWedgeAttrs) const;
+
+	bool MergeVertexQuadricsToEdgeQuadric(const TAttrBasedQuadricErrorV2& VertexQuadric0, 
+		const TAttrBasedQuadricErrorV2& VertexQuadric1,
+		const TAttrBasedQuadricErrorV2* TriangleQuadricA,
+		const TAttrBasedQuadricErrorV2* TriangleQuadricB,
+		const int LocalEdgeIndexA,
+		const int LocalEdgeIndexB,
+		const RealType RemoveScale);
+
+		// Set A to rank-1 matrix V \otimes V = V V^T 
+	void AAddOuterProduct(const TVector<RealType>& V)
+	{
+		BaseStruct::Axx += V.X * V.X;
+		BaseStruct::Axy += V.X * V.Y;
+		BaseStruct::Axz += V.X * V.Z;
+		BaseStruct::Ayy += V.Y * V.Y;
+		BaseStruct::Ayz += V.Y * V.Z;
+		BaseStruct::Azz += V.Z * V.Z;
+	}
+
+    void AAddCrossProductSquaredTranspose(const TVector<RealType>& V)
+    {
+		const RealType V0Sqr = V.X * V.X;
+		const RealType V1Sqr = V.Y * V.Y;
+		const RealType V2Sqr = V.Z * V.Z;
+
+		BaseStruct::Axx +=  V1Sqr + V2Sqr;
+        BaseStruct::Ayy +=  V0Sqr + V2Sqr;
+        BaseStruct::Azz +=  V0Sqr + V1Sqr;
+        
+		BaseStruct::Axy -= V.X * V.Y;
+		BaseStruct::Axz -= V.X * V.Z;
+		BaseStruct::Ayz -= V.Y * V.Z;
+    }
+
+    void AScale(const RealType Alpha)
+    {
+		BaseStruct::Axx *= Alpha;
+		BaseStruct::Ayy *= Alpha;
+		BaseStruct::Azz *= Alpha;
+        
+		BaseStruct::Axy *= Alpha;
+		BaseStruct::Axz *= Alpha;
+		BaseStruct::Ayz *= Alpha;
+    }
+
+	void MakeProbabilisticPlaneQuadric(
+		const TVector<RealType>& P0, 
+		const TVector<RealType>& P1, 
+		const TVector<RealType>& P2,
+		const RealType SigmaN);
+
+	void MakeTriangleQuadric(
+		const TVector<RealType>& P0, 
+		const TVector<RealType>& P1, 
+		const TVector<RealType>& P2);
+    	
+    void MakeProbabilisticTriangleQuadric(
+		const TVector<RealType>& P0, 
+		const TVector<RealType>& P1,
+		const TVector<RealType>& P2,
+		const RealType SigmaN);
+
+	// See: http://hhoppe.com/newqem.pdf or https://www.cc.gatech.edu/~turk/my_papers/memless_vis98.pdf
+	// volume constraint: dot(g_vol,x) = DVol
+	// 
+	// The volume of the tetrahedron spanned by the base triangle (v0, v1, v2) and x is 1/6 * det(v1-v0 v2-v0 x-v0) = 1/6 dot((v1-v0) ^ (v2-v0), x-v0)
+	// Summed over all triangles and setting to 0: dot(sum_t N, x) = sum(dot(N, v0)) with N the scaled normal (v1-v0) ^ (v2-v0)
+	// 
+	TVector<RealType> GVol { 0. };
+	RealType          DVol { 0. };
+	RealType          Area  { 0. };         // accumulated area, this is always used for scaling the volume constraint.
+	RealType          TotalAreaTerm { 0. }; // accumulated area or area^2, depending on which variant is used. This is used for scaling the quadric.
+	bool              bCanCollapse { false };
+};
+
+template <typename RealType>
+inline bool TAttrBasedQuadricErrorV2<RealType>::MergeVertexQuadricsToEdgeQuadric(
+	const TAttrBasedQuadricErrorV2& VertexQuadric0,
+	const TAttrBasedQuadricErrorV2& VertexQuadric1,
+	const TAttrBasedQuadricErrorV2* TriangleQuadricA,
+	const TAttrBasedQuadricErrorV2* TriangleQuadricB,
+	const int LocalEdgeIndexA,
+	const int LocalEdgeIndexB,
+	const RealType RemoveScale)
+{	
+	WedgeCount.Init(0, VertexQuadric0.WedgeCount.Num());
+	checkSlow(VertexQuadric0.WedgeCount.Num() == VertexQuadric1.WedgeCount.Num());
+
+	// Walk through WedgeCount
+
+	int Src0WedgeBegin = 0;
+	int Src1WedgeBegin = 0;
+	int TriAWedgeBegin = 0;
+	int TriBWedgeBegin = 0;
+
+	int MergedWedgeBegin = 0;
+
+	const int NextLocalEdgeIndexA = (LocalEdgeIndexA + 1) % 3; 
+	const int NextLocalEdgeIndexB = (LocalEdgeIndexB + 1) % 3; 
+
+	TArray<FAttrWedgeQuadricData, TInlineAllocator<InlineAttrElements>> MergedAttrWedgeQuadricDatas;
+	TArray<int, TInlineAllocator<InlineAttrElements>> MergedWedgeCount;
+
+	// loop over all attributes
+
+	// check whether a wedge on one vertex would be merged with two distinct wedges on the other vertex.
+	// this could happen on bowtie-like topology of the attributes. we can't handle this and don't 
+	// allow the collapse.
+	bool CheckFailed = false;
+
+	for (int i=0; i<VertexQuadric0.WedgeCount.Num(); ++i)
+	{
+		const int NSrc0 = VertexQuadric0.WedgeCount[i]; 
+		const int NSrc1 = VertexQuadric1.WedgeCount[i]; 
+
+		const int NTriA = TriangleQuadricA ? TriangleQuadricA->WedgeCount[i] : 0; 
+		const int NTriB = TriangleQuadricB ? TriangleQuadricB->WedgeCount[i] : 0; 
+		
+		// copy NDst over
+		for (int j=0; j<NSrc0; ++j)
+		{
+			MergedAttrWedgeQuadricDatas.Add(VertexQuadric0.AttrWedgeQuadricDatas[Src0WedgeBegin + j]);
+		}
+		int NMerged = NSrc0;
+
+		// for all source wedges, search the ones in the merged list
+		for (int j=0; j<NSrc1; ++j)
+		{
+			// the element ID we associate the wedge with
+			const FIndex3i& AttrTri1 = VertexQuadric1.AttrWedgeQuadricDatas[Src1WedgeBegin + j].AttrTri;
+			checkSlow(AttrTri1.B == IndexConstants::InvalidID && AttrTri1.C == IndexConstants::InvalidID);
+			const int EID1 = AttrTri1.A;
+
+			int MergeIdx = IndexConstants::InvalidID;
+			for (int k=0; k<NMerged; ++k)
+			{
+				const FIndex3i& AttrTri0 = MergedAttrWedgeQuadricDatas[MergedWedgeBegin + k].AttrTri;
+				checkSlow(AttrTri0.C == IndexConstants::InvalidID);
+				const int EID0 = AttrTri0.A;
+
+				if (NTriA > 0)
+				{
+					checkSlow(NTriA == 1);
+					const FIndex3i AttrTriA = TriangleQuadricA->AttrWedgeQuadricDatas[TriAWedgeBegin].AttrTri;
+
+					const int EIDA_0 = AttrTriA[LocalEdgeIndexA];
+					const int EIDA_1 = AttrTriA[NextLocalEdgeIndexA];
+
+					if ((EIDA_0 == EID0 && EIDA_1 == EID1) || (EIDA_0 == EID1 && EIDA_1 == EID0))
+					{
+						if (!(MergeIdx == IndexConstants::InvalidID || MergeIdx == MergedWedgeBegin + k)) {
+							CheckFailed = true;
+						}
+						MergeIdx = MergedWedgeBegin + k;
+					}
+				}
+
+				if (NTriB > 0)
+				{
+					checkSlow(NTriB == 1);
+					const FIndex3i AttrTriB = TriangleQuadricB->AttrWedgeQuadricDatas[TriBWedgeBegin].AttrTri;
+
+					const int EIDB_0 = AttrTriB[LocalEdgeIndexB];
+					const int EIDB_1 = AttrTriB[NextLocalEdgeIndexB];
+
+					if ((EIDB_0 == EID0 && EIDB_1 == EID1) || (EIDB_0 == EID1 && EIDB_1 == EID0))
+					{
+						if (!(MergeIdx == IndexConstants::InvalidID || MergeIdx == MergedWedgeBegin + k)) {
+							CheckFailed = true;
+						}
+						MergeIdx = MergedWedgeBegin + k;
+					}
+				}
+			}
+			
+			const int EID = VertexQuadric1.AttrWedgeQuadricDatas[Src1WedgeBegin + j].AttrTri[0];
+			if (MergeIdx == IndexConstants::InvalidID)
+			{
+				// add the entry
+				MergeIdx = MergedAttrWedgeQuadricDatas.AddZeroed();
+				MergedAttrWedgeQuadricDatas[MergeIdx].AttrTri = FIndex3i( EID, IndexConstants::InvalidID, IndexConstants::InvalidID );
+				NMerged++;
+			} 
+			else
+			{
+				MergedAttrWedgeQuadricDatas[MergeIdx].AttrTri[1] = EID;
+			} 
+
+			// accumulate the whole wedge (these are vertex attribute quadrics, so already weighted by area)
+			MergedAttrWedgeQuadricDatas[MergeIdx].AccumulateQuadricData(VertexQuadric1.AttrWedgeQuadricDatas[Src1WedgeBegin + j]);
+
+		} // end walking over all attributes wedges from vertex1 for this attribute
+
+		MergedWedgeBegin += NMerged;
+		MergedWedgeCount.Add(NMerged);
+
+		Src0WedgeBegin += NSrc0;
+		Src1WedgeBegin += NSrc1;
+		TriAWedgeBegin += NTriA;
+		TriBWedgeBegin += NTriB;
+		
+	} // end walking over all attribute slots
+
+	AttrWedgeQuadricDatas = MoveTemp(MergedAttrWedgeQuadricDatas);
+	WedgeCount = MoveTemp(MergedWedgeCount);
+
+	return !CheckFailed;
+}
+
+template <typename RealType>
+inline void TAttrBasedQuadricErrorV2<RealType>::AccumulateFaceQuadric(
+	const int CornerIdx, RealType TriArea, RealType Weight, const TAttrBasedQuadricErrorV2& FaceQuadric)
+{
+	ensure(TriArea > -UE_DOUBLE_KINDA_SMALL_NUMBER);
+
+	GVol += TriArea * FaceQuadric.GVol;
+	DVol += TriArea * FaceQuadric.DVol;
+	Area += TriArea;
+	TotalAreaTerm += FaceQuadric.TotalAreaTerm;
+
+	BaseStruct::Add(Weight, FaceQuadric);
+
+	if (WedgeCount.IsEmpty())
+	{
+		WedgeCount.Init(0, FaceQuadric.WedgeCount.Num());
+	}
+
+	int SrcWedgeBegin = 0;
+	int DstWedgeBegin = 0;
+	int MergedWedgeBegin = 0;
+
+	TArray<FAttrWedgeQuadricData, TInlineAllocator<InlineAttrElements>> MergedAttrWedgeQuadricDatas;
+	TArray<int, TInlineAllocator<InlineAttrElements>> MergedWedgeCount;
+
+	// loop over all attributes
+	for (int i=0; i<FaceQuadric.WedgeCount.Num(); ++i)
+	{
+		const int NSrc = FaceQuadric.WedgeCount[i]; // we have N wedges for this attribute in Other
+		const int NDst = WedgeCount[i];     
+		
+		check(NSrc < 2); // Other should be a tri quadric and not have multiple wedges per attribute
+		
+		// copy NDst over
+		for (int j=0; j<NDst; ++j)
+		{
+			MergedAttrWedgeQuadricDatas.Add(AttrWedgeQuadricDatas[DstWedgeBegin + j]);
+		}
+		int NMerged = NDst;
+
+		// for all source wedges, search the ones in the merged list
+		for (int j=0; j<NSrc; ++j)
+		{
+			// the element ID we associate the wedge with
+			int EID = FaceQuadric.AttrWedgeQuadricDatas[SrcWedgeBegin + j].AttrTri[CornerIdx];
+
+			int MergeIdx = IndexConstants::InvalidID;
+			for (int k=0; k<NMerged; ++k)
+			{
+				const FIndex3i AttrTri = MergedAttrWedgeQuadricDatas[MergedWedgeBegin + k].AttrTri;
+				check(AttrTri.C == IndexConstants::InvalidID);
+				if (AttrTri.A == EID || AttrTri.B == EID)
+				{
+					// found the wedge -> accumulate
+					MergeIdx = MergedWedgeBegin + k;
+					break;
+				}
+			}
+			if (MergeIdx == IndexConstants::InvalidID)
+			{
+				// add the entry
+				MergeIdx = MergedAttrWedgeQuadricDatas.AddZeroed();
+				NMerged++;
+				MergedAttrWedgeQuadricDatas[MergeIdx].AttrTri = FIndex3i(EID, IndexConstants::InvalidID, IndexConstants::InvalidID);
+			}
+			
+			MergedAttrWedgeQuadricDatas[MergeIdx].AccumulateQuadricData(FaceQuadric.AttrWedgeQuadricDatas[SrcWedgeBegin + j]);
+		}
+
+		MergedWedgeBegin += NMerged;
+		MergedWedgeCount.Add(NMerged);
+		check(MergedAttrWedgeQuadricDatas.Num() == MergedWedgeBegin);
+
+		DstWedgeBegin += NDst;
+		SrcWedgeBegin += NSrc;
+	}
+
+#if !UE_BUILD_SHIPPING
+	{
+		// Sanity check
+		check(MergedWedgeCount.Num() == FaceQuadric.WedgeCount.Num());
+
+		int MergedWedgeSum = 0;
+		for (int N : MergedWedgeCount)
+		{
+			MergedWedgeSum += N;
+		}
+		check(MergedWedgeSum == MergedAttrWedgeQuadricDatas.Num());
+	}
+#endif
+
+	AttrWedgeQuadricDatas = MoveTemp(MergedAttrWedgeQuadricDatas);
+	WedgeCount = MoveTemp(MergedWedgeCount);
+}
+
+template <typename RealType>
+inline void TAttrBasedQuadricErrorV2<RealType>::RemoveFaceQuadric(
+	const int CornerIdx, RealType TriArea, RealType RemoveScale, const TAttrBasedQuadricErrorV2& FaceQuadric)
+{
+	ensure(TriArea >= -UE_DOUBLE_KINDA_SMALL_NUMBER);
+	GVol -= TriArea * FaceQuadric.GVol;
+	DVol -= TriArea * FaceQuadric.DVol;
+	Area -= TriArea;
+	
+	TotalAreaTerm -= FaceQuadric.TotalAreaTerm;
+
+	// the face quadric contains the position and the attribute terms in A, b, c. These correspond to the 
+	// accumulation of all the wedges. The B and d terms are per-wedge and aggregated below.
+	BaseStruct::Add(-RemoveScale, FaceQuadric);
+
+	int SrcWedgeBegin = 0;
+	int DstWedgeBegin = 0;
+	
+	check(FaceQuadric.WedgeCount.Num() == WedgeCount.Num());
+
+	// loop over all attributes, each attribute has a separate list of wedges
+	for (int i=0; i<FaceQuadric.WedgeCount.Num(); ++i)
+	{
+		const int NSrc = FaceQuadric.WedgeCount[i]; // we have N wedges for this attribute in Other
+		const int NDst = WedgeCount[i];     
+		
+		checkSlow(NSrc < 2); // should be a face quadric and not have multiple wedges per attribute
+		
+		bool Found = false;
+
+		// for all source wedges, search the ones in the merged list
+		for (int j=0; j<NSrc; ++j)
+		{
+			// the element ID we associate the wedge with
+			const int EID = FaceQuadric.AttrWedgeQuadricDatas[SrcWedgeBegin + j].AttrTri[CornerIdx];
+
+			for (int k=0; k<NDst; ++k)
+			{
+				const FIndex3i AttrTri = AttrWedgeQuadricDatas[DstWedgeBegin + k].AttrTri;
+				checkSlow(AttrTri.C == IndexConstants::InvalidID);
+				
+				if (AttrTri.A == EID || AttrTri.B == EID)
+				{
+					// found the wedge -> remove
+					Found = true;
+
+					AttrWedgeQuadricDatas[DstWedgeBegin + k].RemoveQuadricData(FaceQuadric.AttrWedgeQuadricDatas[SrcWedgeBegin + j]);
+					break;
+				}
+			}
+			checkSlow(Found);
+		}
+
+		DstWedgeBegin += NDst;
+		SrcWedgeBegin += NSrc;
+	}
+}
+
+template <typename RealType>
+inline void TAttrBasedQuadricErrorV2<RealType>::MakeProbabilisticPlaneQuadric(
+	const TVector<RealType>& P0, 
+	const TVector<RealType>& P1, 
+	const TVector<RealType>& P2,
+	const RealType SigmaN)
+{
+	const TVector<RealType> Normal { VectorUtil::NormalArea(P0, P1, P2, Area) };
+	const TVector<RealType> P = (P0 + P1 + P2) * (1./3.);
+
+	TotalAreaTerm = Area;
+	
+	BaseStruct::Axx = TotalAreaTerm * Normal.X * Normal.X;
+	BaseStruct::Axy = TotalAreaTerm * Normal.X * Normal.Y;
+	BaseStruct::Axz = TotalAreaTerm * Normal.X * Normal.Z;
+	BaseStruct::Ayy = TotalAreaTerm * Normal.Y * Normal.Y;
+	BaseStruct::Ayz = TotalAreaTerm * Normal.Y * Normal.Z;
+	BaseStruct::Azz = TotalAreaTerm * Normal.Z * Normal.Z;
+	
+	RealType dotNP = Normal | P;
+
+	BaseStruct::bx = TotalAreaTerm * -dotNP * Normal.X; 
+	BaseStruct::by = TotalAreaTerm * -dotNP * Normal.Y; 
+	BaseStruct::bz = TotalAreaTerm * -dotNP * Normal.Z; 
+
+	BaseStruct::c  = TotalAreaTerm * dotNP * dotNP;
+
+	const double SigmaN2 = SigmaN * SigmaN;
+	if (SigmaN > 0.)
+	{
+		BaseStruct::Axx += TotalAreaTerm * SigmaN2;
+		BaseStruct::Ayy += TotalAreaTerm * SigmaN2;
+		BaseStruct::Azz += TotalAreaTerm * SigmaN2;
+
+		BaseStruct::bx -= TotalAreaTerm * SigmaN2 * P.X;
+		BaseStruct::by -= TotalAreaTerm * SigmaN2 * P.Y;
+		BaseStruct::bz -= TotalAreaTerm * SigmaN2 * P.Z;
+
+		BaseStruct::c += TotalAreaTerm * SigmaN2 * P.Dot(P);
+	}
+}
+
+template <typename RealType>
+inline void TAttrBasedQuadricErrorV2<RealType>::MakeTriangleQuadric(
+	const TVector<RealType>& P0, 
+	const TVector<RealType>& P1, 
+	const TVector<RealType>& P2)
+{
+	Area = VectorUtil::Area(P0, P1, P2);
+	TotalAreaTerm = Area * Area;
+
+	const TVector<RealType> P0xP1 = P0 ^ P1;
+	const TVector<RealType> P1xP2 = P1 ^ P2;
+	const TVector<RealType> P2xP0 = P2 ^ P0;
+
+	const TVector<RealType> S = P0xP1 + P1xP2 + P2xP0;
+	const RealType Det = P0xP1.Dot(P2); // scalar triple product
+
+	AAddOuterProduct(S);
+	BaseStruct::bx = -S.X * Det;
+	BaseStruct::by = -S.Y * Det;
+	BaseStruct::bz = -S.Z * Det;
+
+	BaseStruct::c = Det * Det;
+}
+
+template <typename RealType>
+inline void TAttrBasedQuadricErrorV2<RealType>::MakeProbabilisticTriangleQuadric(
+	const TVector<RealType>& P0,
+	const TVector<RealType>& P1,
+	const TVector<RealType>& P2,
+	const RealType Sigma)
+{
+	Area = VectorUtil::Area(P0, P1, P2);
+	TotalAreaTerm = Area * Area;
+
+	const RealType SigmaSqr = Sigma * Sigma;
+
+	const TVector<RealType> P0xP1 = P0 ^ P1;
+	const TVector<RealType> P1xP2 = P1 ^ P2;
+	const TVector<RealType> P2xP0 = P2 ^ P0;
+
+	const RealType Det = P0xP1.Dot(P2); // scalar triple product
+
+	const TVector<RealType> S = P0xP1 + P1xP2 + P2xP0;
+
+	const TVector<RealType> P1P0 = P0 - P1;
+	const TVector<RealType> P2P1 = P1 - P2;
+	const TVector<RealType> P0P2 = P2 - P0;
+
+	AAddCrossProductSquaredTranspose(P1P0);
+	AAddCrossProductSquaredTranspose(P2P1);
+	AAddCrossProductSquaredTranspose(P0P2);
+	AScale(SigmaSqr);
+
+	AAddOuterProduct(S);
+
+	const RealType SigmaQ = SigmaSqr * SigmaSqr;
+
+	BaseStruct::Axx += 6. * SigmaQ;
+	BaseStruct::Ayy += 6. * SigmaQ;
+	BaseStruct::Azz += 6. * SigmaQ;
+	
+	TVector<RealType> B = S * Det
+		- ((P1P0 ^ P0xP1) + (P2P1 ^ P1xP2) + (P0P2 ^ P2xP0)) * SigmaSqr
+		+ (P0 + P1 + P2) * (2.0 * SigmaQ);
+
+	BaseStruct::bx = -B.X;
+	BaseStruct::by = -B.Y;
+	BaseStruct::bz = -B.Z;
+
+	BaseStruct::c = Det * Det
+		+ SigmaSqr * (P0xP1.Dot(P0xP1) + P1xP2.Dot(P1xP2) + P2xP0.Dot(P2xP0))
+		+ 2.0 * SigmaQ * (P0.Dot(P0) + P1.Dot(P1) + P2.Dot(P2))
+		+ 6.0 * SigmaQ * SigmaSqr;
+}
+
+template <typename RealType>
+inline TAttrBasedQuadricErrorV2<RealType>::FAttrArray TAttrBasedQuadricErrorV2<RealType>::ComputeScaledAttributes(
+	const TVector<RealType>& P) const
+{
+	FAttrArray WedgeAttrs;
+	WedgeAttrs.SetNum(AttrWedgeQuadricDatas.Num());
+
+	for (int i=0; i<AttrWedgeQuadricDatas.Num(); ++i)
+	{
+		TVector<RealType>& Attr = WedgeAttrs[i];
+		if (AttrWedgeQuadricDatas[i].TotalAreaTerm > UE_DOUBLE_KINDA_SMALL_NUMBER)
+		{
+			Attr.X = AttrWedgeQuadricDatas[i].G[0].Dot(P) + AttrWedgeQuadricDatas[i].D[0];
+			Attr.Y = AttrWedgeQuadricDatas[i].G[1].Dot(P) + AttrWedgeQuadricDatas[i].D[1];
+			Attr.Z = AttrWedgeQuadricDatas[i].G[2].Dot(P) + AttrWedgeQuadricDatas[i].D[2];
+			Attr /= AttrWedgeQuadricDatas[i].TotalAreaTerm;
+		}
+		else
+		{
+			Attr = TVector<RealType>(0.);
+		}
+	} 
+
+	return WedgeAttrs;
+}
+
+template <typename RealType>
+inline TAttrBasedQuadricErrorV2<RealType>::FAttrArray TAttrBasedQuadricErrorV2<RealType>::ComputeAttributes(const TVector<RealType>& P) const
+{
+	FAttrArray ScaledAttributes = ComputeScaledAttributes(P);
+
+	for (int i=0; i<AttrWedgeQuadricDatas.Num(); ++i)
+	{
+		ScaledAttributes[i] /= AttrWedgeQuadricDatas[i].AttrWeight;
+	}
+	return ScaledAttributes;
+}
+
+template <typename RealType>
+inline RealType TAttrBasedQuadricErrorV2<RealType>::Evaluate(
+	const TVector<RealType>& P, const FAttrArray& ScaledWedgeAttrs) const
+{
+	// Extended quadric system is
+	//
+	//         ( A + \sum g_i g_i^T  = C | B                 )
+	// A_ext = ( - - - - - - - - - - - - - - - - - - - - - - )
+	//         ( B^T                     | aI                )
+	// 
+	// 
+	// b_ext = ( b + \sum d_i g_i    )
+	//         ( - - - - - - - - - - )
+	//         ( - (d_1, ... d_m)^t  )
+	// 
+	// c_ext = c + \sum_j d_i^2 
+	//
+	// (P^T, S^T) A_ext (P, S) = P^T C P + 2 <P, B S> + a<S, S>
+	
+	const RealType ptCp = P.Dot(BaseStruct::MultiplyA(P));
+
+	RealType attrDotAttr = 0.;
+	TVector<RealType> BS { 0. };
+
+	ensure(ScaledWedgeAttrs.Num() == AttrWedgeQuadricDatas.Num());
+
+	int AttrIdx = 0;
+	for (const FAttrWedgeQuadricData& AttrWedgeQuadricData : AttrWedgeQuadricDatas)
+	{
+		const TVector<RealType> (&G)[3] = AttrWedgeQuadricData.G;
+		const TVector<RealType>& V = ScaledWedgeAttrs[AttrIdx];
+		attrDotAttr += V.Dot(V) * AttrWedgeQuadricData.TotalAreaTerm;
+
+		BS -= G[0] * V[0];
+		BS -= G[1] * V[1];
+		BS -= G[2] * V[2];
+
+		++AttrIdx;
+	}
+	RealType PBS = P.Dot(BS);
+	
+	const RealType QuadraticTerm = ptCp + 2.0 * PBS + attrDotAttr;
+	RealType LinearTerm = 2.0 * (BaseStruct::bx * P.X + BaseStruct::by * P.Y + BaseStruct::bz * P.Z);
+	RealType ConstantTerm = BaseStruct::c;
+
+	AttrIdx = 0;
+	for (const FAttrWedgeQuadricData& AttrWedgeQuadricData : AttrWedgeQuadricDatas)
+	{
+		const TVector<RealType>& D = AttrWedgeQuadricData.D;
+		LinearTerm -= 2. * D.Dot(ScaledWedgeAttrs[AttrIdx]);
+		++AttrIdx;
+	}
+
+	return QuadraticTerm + LinearTerm + ConstantTerm;
+}
+
+template <typename RealType>
+inline bool TAttrBasedQuadricErrorV2<RealType>::ValidateArea() const
+{
+	int WedgeBegin = 0;
+	for (int i=0; i<WedgeCount.Num(); ++i)
+	{
+		RealType AttrTotalAreaTerm = 0.;
+		const int N = WedgeCount[i]; 
+
+		if (N == 0)
+		{
+			continue;
+		}
+
+		for (int j=0; j<N; ++j)
+		{
+			AttrTotalAreaTerm += AttrWedgeQuadricDatas[WedgeBegin + j].TotalAreaTerm;
+		}
+		WedgeBegin += N;
+
+		const RealType RelError = FMath::Abs(AttrTotalAreaTerm - TotalAreaTerm) / TotalAreaTerm;
+		if (RelError > 1.e-12)
+		{
+			return false;
+		}
+	}
+	return true;
+}
+
+using FAttrBasedQuadricErrorV2d = TAttrBasedQuadricErrorV2<double>;
+
+} // namespace Geometry
+} // namespace UE
\ No newline at end of file
diff --git a/Engine/Plugins/Runtime/GeometryScripting/Source/GeometryScriptingCore/Private/MeshSimplifyFunctions.cpp b/Engine/Plugins/Runtime/GeometryScripting/Source/GeometryScriptingCore/Private/MeshSimplifyFunctions.cpp
--- a/Engine/Plugins/Runtime/GeometryScripting/Source/GeometryScriptingCore/Private/MeshSimplifyFunctions.cpp
+++ b/Engine/Plugins/Runtime/GeometryScripting/Source/GeometryScriptingCore/Private/MeshSimplifyFunctions.cpp
@@ -14,6 +14,7 @@
 #include "Operations/PolygroupRemesh.h"
 #include "Operations/MeshClusterSimplifier.h"
 #include "ConstrainedDelaunay2.h"
+#include "MeshSimplificationQuadrics.h"
 
 #if WITH_EDITOR
 // for UE Standard simplifier (editor only)
@@ -32,7 +33,6 @@ using namespace UE::Geometry;
 
 #define LOCTEXT_NAMESPACE "UGeometryScriptLibrary_MeshSimplifyFunctions"
 
-
 UDynamicMesh* UGeometryScriptLibrary_MeshSimplifyFunctions::ApplySimplifyToPlanar(
 	UDynamicMesh* TargetMesh,
 	FGeometryScriptPlanarSimplifyOptions Options,
@@ -152,6 +152,30 @@ namespace UE::PrivateSimplifyHelper
 		EdgeLength
 	};
 
+	template <typename SimplificationType>
+	static void InitQuadricOptions(const FGeometryScriptSimplifyMeshOptions& Options, 
+		typename SimplificationType::FQuadricOptions& QuadricOptions)
+	{
+	}
+
+	template <>
+	void InitQuadricOptions<TMeshSimplification<FAttrBasedQuadricErrorV2d>>(const FGeometryScriptSimplifyMeshOptions& Options, 
+		typename TMeshSimplification<FAttrBasedQuadricErrorV2d>::FQuadricOptions& QuadricOptions )
+	{
+		if (Options.QuadricVariant == EGeometryScriptMeshSimplificationQuadricVariant::PlaneQuadric)
+		{
+			QuadricOptions.QuadricVariant = FAttrBasedQuadricErrorV2d::EQuadricVariant::PlaneQuadric;
+		}
+		else
+		{
+			QuadricOptions.QuadricVariant = FAttrBasedQuadricErrorV2d::EQuadricVariant::TriangleQuadric;
+		}
+		QuadricOptions.NormalAttributeWeight   = Options.NormalAttributeWeight;
+		QuadricOptions.TangentAttributeWeight  = Options.TangentAttributeWeight;
+		QuadricOptions.ColorAttributeWeight    = Options.ColorAttributeWeight;
+		QuadricOptions.TexCoordAttributeWeight = Options.TexCoordAttributeWeight;
+	}
+	
 	template<typename SimplificationType>
 	static void DoSimplifyMesh(
 		FDynamicMesh3& EditMesh,
@@ -165,12 +189,14 @@ namespace UE::PrivateSimplifyHelper
 	{
 		SimplificationType Simplifier(&EditMesh);
 
+		InitQuadricOptions<SimplificationType>(Options, Simplifier.QuadricOptions);
+
 		Simplifier.ProjectionMode = SimplificationType::ETargetProjectionMode::NoProjection;
 		if (ProjectionTarget != nullptr)
 		{
 			Simplifier.SetProjectionTarget(ProjectionTarget);
 		}
-
+		
 		Simplifier.DEBUG_CHECK_LEVEL = 0;
 		Simplifier.bRetainQuadricMemory = Options.bRetainQuadricMemory;
 		Simplifier.bAllowSeamCollapse = Options.bAllowSeamCollapse;
@@ -294,6 +320,10 @@ namespace UE::PrivateSimplifyHelper
 		{
 			DoSimplifyMesh<FVolPresMeshSimplification>(EditMesh, Options, SimplifyTargetMode, TargetCount, TargetMinEdgeLength, ProjectionTarget, GeometricTolerance);
 		}
+		else if (Options.Method == EGeometryScriptRemoveMeshSimplificationType::AttributeAwareV2)
+		{
+			DoSimplifyMesh<TMeshSimplification<FAttrBasedQuadricErrorV2d>>(EditMesh, Options, SimplifyTargetMode, TargetCount, TargetMinEdgeLength, ProjectionTarget, GeometricTolerance);
+		}
 		else // EGeometryScriptRemoveMeshSimplificationType::StandardQEM
 		{
 			DoSimplifyMesh<FQEMSimplification>(EditMesh, Options, SimplifyTargetMode, TargetCount, TargetMinEdgeLength, ProjectionTarget, GeometricTolerance);
@@ -366,9 +396,6 @@ UDynamicMesh* UGeometryScriptLibrary_MeshSimplifyFunctions::ApplySimplifyToTrian
 	return TargetMesh;
 }
 
-
-
-
 UDynamicMesh* UGeometryScriptLibrary_MeshSimplifyFunctions::ApplySimplifyToVertexCount(
 	UDynamicMesh* TargetMesh,
 	int32 VertexCount,
diff --git a/Engine/Plugins/Runtime/GeometryScripting/Source/GeometryScriptingCore/Public/GeometryScript/MeshSimplifyFunctions.h b/Engine/Plugins/Runtime/GeometryScripting/Source/GeometryScriptingCore/Public/GeometryScript/MeshSimplifyFunctions.h
--- a/Engine/Plugins/Runtime/GeometryScripting/Source/GeometryScriptingCore/Public/GeometryScript/MeshSimplifyFunctions.h
+++ b/Engine/Plugins/Runtime/GeometryScripting/Source/GeometryScriptingCore/Public/GeometryScript/MeshSimplifyFunctions.h
@@ -43,9 +43,24 @@ public:
 UENUM(BlueprintType)
 enum class EGeometryScriptRemoveMeshSimplificationType : uint8
 {
-	StandardQEM = 0,
-	VolumePreserving = 1,
-	AttributeAware = 2
+	/** classic quadric error metric without volume preservation */
+	StandardQEM = 0 UMETA(DisplayName = "Standard QEM"),      
+
+	/** classic quadric error metric with volume preservation */
+	VolumePreserving = 1 UMETA(DisplayName = "Volume Preserving"), 
+
+	/** QEM with volume preservation, account for vertex normals in optimization. This used to be called Attribute Aware. */
+	AttributeAware = 2 UMETA(DisplayName = "Normals Aware, Volume Preserving"),   
+
+	/** QEM with volume preservation, accounting for normals/tangents/bitangents/color/texture coordinates and weight channels with attribute seams */
+	AttributeAwareV2 = 3 UMETA(DisplayName = "Attribute Aware, Volume Preserving") 
+};
+
+UENUM(BlueprintType)
+enum class EGeometryScriptMeshSimplificationQuadricVariant : uint8
+{
+	PlaneQuadric = 0,
+	TriangleQuadric = 1
 };
 
 // Specify a weight map for controlling relative density of a mesh
@@ -110,6 +125,22 @@ public:
 	/** If enabled, the simplified mesh is automatically compacted to remove gaps in the index space. This is expensive and can be disabled by advanced users. */
 	UPROPERTY(EditAnywhere, BlueprintReadWrite, Category = Options)
 	bool bAutoCompact = true;
+
+	UPROPERTY(EditAnywhere, BlueprintReadWrite, Category = Options)
+	EGeometryScriptMeshSimplificationQuadricVariant QuadricVariant = EGeometryScriptMeshSimplificationQuadricVariant::PlaneQuadric;
+
+	UPROPERTY(EditAnywhere, BlueprintReadWrite, Category = Options)
+	float NormalAttributeWeight = 16.f;
+
+	UPROPERTY(EditAnywhere, BlueprintReadWrite, Category = Options)
+	float TangentAttributeWeight = 0.1f;
+
+	UPROPERTY(EditAnywhere, BlueprintReadWrite, Category = Options)
+	float ColorAttributeWeight = 0.1f;
+
+	UPROPERTY(EditAnywhere, BlueprintReadWrite, Category = Options)
+	float TexCoordAttributeWeight = 0.5f;
+
 };
 
 UENUM(BlueprintType)
diff --git a/Engine/Source/Runtime/GeometryCore/Public/QuadricError.h b/Engine/Source/Runtime/GeometryCore/Public/QuadricError.h
--- a/Engine/Source/Runtime/GeometryCore/Public/QuadricError.h
+++ b/Engine/Source/Runtime/GeometryCore/Public/QuadricError.h
@@ -73,18 +73,28 @@ struct TQuadricError
 	/**
 	 * Add scalar multiple of a TQuadricError to this TQuadricError
 	 */
-	void Add(RealType w, const TQuadricError& b) 
+	void Add(RealType Weight, const TQuadricError& b) 
 	{
-		Axx += w * b.Axx;
-		Axy += w * b.Axy;
-		Axz += w * b.Axz;
-		Ayy += w * b.Ayy;
-		Ayz += w * b.Ayz;
-		Azz += w * b.Azz;
-		bx += w * b.bx;
-		by += w * b.by;
-		bz += w * b.bz;
-		c += w * b.c;
+		Axx += Weight * b.Axx;
+		Axy += Weight * b.Axy;
+		Axz += Weight * b.Axz;
+		Ayy += Weight * b.Ayy;
+		Ayz += Weight * b.Ayz;
+		Azz += Weight * b.Azz;
+		bx += Weight * b.bx;
+		by += Weight * b.by;
+		bz += Weight * b.bz;
+		c += Weight * b.c;
+	}
+
+	void RemoveFaceQuadric(const int CornerIdx, RealType TriArea, RealType RemoveScale, const TQuadricError& b)
+	{
+		Add(-TriArea * RemoveScale, b);
+	}
+
+	void AccumulateFaceQuadric(const int CornerIdx, RealType TriArea, RealType Weight, const TQuadricError& b)
+	{
+		Add(TriArea * Weight, b);
 	}
 
 	void Add(const TQuadricError& b)
@@ -266,10 +276,18 @@ struct TQuadricError
 		}
 		TVector<RealType> N1, N2;
 		VectorUtil::MakePerpVectors(NormalizedLineDir, N1, N2);
-		Add(Weight * .5, TQuadricError<RealType>(N1, Pos));
-		Add(Weight * .5, TQuadricError<RealType>(N2, Pos));
+		Add(Weight * RealType(0.5), TQuadricError<RealType>(N1, Pos));
+		Add(Weight * RealType(0.5), TQuadricError<RealType>(N2, Pos));
 	}
 
+	// TODO remove the need to track the area in the simplifier and keep it quadric-internal
+	RealType GetArea() const
+	{
+		return -1.f;
+	}
+
+	bool CanCollapse() const { return true; }
+
 };
 
 
@@ -378,11 +396,22 @@ public:
 	/**
 	* Area Weighted Add
 	*/
-	void Add(RealType w, const TVolPresQuadricError& b)
+	void Add(RealType Weight, const TVolPresQuadricError& b)
 	{
-		BaseStruct::Add(w, b);
+		BaseStruct::Add(Weight, b);
 
-		PlaneData.Add(w, b.PlaneData);
+		// this should probably only be weighted by TriArea
+		PlaneData.Add(Weight, b.PlaneData);
+	}
+
+	void RemoveFaceQuadric(const int CornerIdx, RealType TriArea, RealType RemoveScale, const TVolPresQuadricError& b)
+	{
+		Add(-TriArea * RemoveScale, b);
+	}
+
+	void AccumulateFaceQuadric(const int CornerIdx, RealType TriArea, RealType Weight, const TVolPresQuadricError& b)
+	{
+		Add(TriArea * Weight, b);
 	}
 
 	void Scale(RealType Wt)
@@ -509,11 +538,11 @@ public:
 		// and solving for a1, a2
 		const TMatrix3<RealType> BasisMatrix(P2-P0, P1-P0, NFace, true); // row-based matrix constructor.
 
-		const double det = BasisMatrix.Determinant(); // geometrically the det is related to the triangle area squared
+		const RealType det = BasisMatrix.Determinant(); // geometrically the det is related to the triangle area squared
 		if (FMath::Abs(det) > 1.e-8)
 		{
 			// invert the matrix
-			const FMatrix3d CoBasisMatrix = BasisMatrix.Inverse();
+			const TMatrix3<RealType> CoBasisMatrix = BasisMatrix.Inverse();
 
 			// for each component of the attribute vector, 
 			for (int i = 0; i < 3; ++i)
@@ -586,20 +615,30 @@ public:
 	/**
 	* Area Weighted Add
 	*/
-	void Add(RealType w, const TAttrBasedQuadricError& other)
+	void Add(RealType Weight, const TAttrBasedQuadricError& other)
 	{
-		BaseClass::Add(w, other);
+		BaseClass::Add(Weight, other);
 
-		a += w; // accumulate weight
+		a += Weight; // accumulate weight
 		attrweight = other.attrweight;
 
-		grad[0] += w * other.grad[0];
-		grad[1] += w * other.grad[1];
-		grad[2] += w * other.grad[2];
+		grad[0] += Weight * other.grad[0];
+		grad[1] += Weight * other.grad[1];
+		grad[2] += Weight * other.grad[2];
 
-		d[0] += w * other.d[0];
-		d[1] += w * other.d[1];
-		d[2] += w * other.d[2];
+		d[0] += Weight * other.d[0];
+		d[1] += Weight * other.d[1];
+		d[2] += Weight * other.d[2];
+	}
+
+	void RemoveFaceQuadric(const int CornerIdx, RealType TriArea, RealType RemoveScale, const TAttrBasedQuadricError& b)
+	{
+		Add(-TriArea * RemoveScale, b);
+	}
+
+	void AccumulateFaceQuadric(const int CornerIdx, RealType TriArea, RealType Weight, const TAttrBasedQuadricError& b)
+	{
+		Add(TriArea * Weight, b);
 	}
 
 	void Scale(RealType Wt)
@@ -644,6 +683,7 @@ public:
 																SM[5] += BaseStruct::Azz;
 
 
+
 		RealType InvSM[6];
 		bool bValid = BaseStruct::InvertSymmetricMatrix(SM, InvSM, minThresh);
 
diff --git a/Engine/Source/Runtime/GeometryCore/Public/VectorUtil.h b/Engine/Source/Runtime/GeometryCore/Public/VectorUtil.h
--- a/Engine/Source/Runtime/GeometryCore/Public/VectorUtil.h
+++ b/Engine/Source/Runtime/GeometryCore/Public/VectorUtil.h
@@ -54,6 +54,16 @@ namespace VectorUtil
 	}
 
 
+	/**
+	 * @return true if all components of V are finite
+	 */
+	template <typename RealType>
+	inline bool IsFinite(const TVector4<RealType>& V)
+	{
+		return TMathUtil<RealType>::IsFinite(V.X) && TMathUtil<RealType>::IsFinite(V.Y) && 
+			TMathUtil<RealType>::IsFinite(V.Z) && TMathUtil<RealType>::IsFinite(V.W);
+	}
+
 	/**
 	 * @return input Value clamped to range [MinValue, MaxValue]
 	 */
