add option to project subdivided skeleton spheres to the medial axis

#Synced-CL 52703645

diff --git a/Engine/Plugins/Experimental/GeometryCollectionPlugin/Source/GeometryCollectionNodes/Private/Dataflow/GeometryCollectionSamplingNodes.cpp b/Engine/Plugins/Experimental/GeometryCollectionPlugin/Source/GeometryCollectionNodes/Private/Dataflow/GeometryCollectionSamplingNodes.cpp
--- a/Engine/Plugins/Experimental/GeometryCollectionPlugin/Source/GeometryCollectionNodes/Private/Dataflow/GeometryCollectionSamplingNodes.cpp
+++ b/Engine/Plugins/Experimental/GeometryCollectionPlugin/Source/GeometryCollectionNodes/Private/Dataflow/GeometryCollectionSamplingNodes.cpp
@@ -566,6 +566,7 @@ FSubdivideMedialSkeletonDataflowNode::FSubdivideMedialSkeletonDataflowNode(const
 	RegisterInputConnection(&MaxTargetEdgeLength).SetCanHidePin(true).SetPinIsHidden(true);
 	RegisterInputConnection(&bSubdivideOnSurfaces).SetCanHidePin(true).SetPinIsHidden(true);
 	RegisterInputConnection(&ReassignClusterPosErrorWt).SetCanHidePin(true).SetPinIsHidden(true);
+	RegisterInputConnection(&bProjectNewMedialSpheres).SetCanHidePin(true).SetPinIsHidden(true);
 
 	RegisterOutputConnection(&MedialSkeleton);
 	RegisterOutputConnection(&MedialSpheres);
@@ -603,6 +604,7 @@ void FSubdivideMedialSkeletonDataflowNode::Evaluate(UE::Dataflow::FContext& Cont
 					UE::Geometry::MedialAxis::FSkeletonSubdivider Subdivider;
 					Subdivider.bSubdivideOnSurfaces = GetValue(Context, &bSubdivideOnSurfaces);
 					Subdivider.ReassignClusterPosErrorWt = GetValue(Context, &ReassignClusterPosErrorWt);
+					Subdivider.bReprojectMedialSpheres = GetValue(Context, &bProjectNewMedialSpheres);
 					Subdivider.GetTargetEdgeLength = [RadiusFraction, RadRef, MinLen, MaxLen](const FSphere& A, const FSphere& B) -> double
 					{
 						double Radius;
@@ -632,7 +634,14 @@ void FSubdivideMedialSkeletonDataflowNode::Evaluate(UE::Dataflow::FContext& Cont
 						return TargetLen;
 					};
 
-					Subdivider.Subdivide(Skeleton, Spatial);
+					const UE::Geometry::TFastWindingTree<FDynamicMesh3>* FWTreePtr = nullptr;
+					TOptional<UE::Geometry::TFastWindingTree<FDynamicMesh3>> FWTree;
+					if (Subdivider.bReprojectMedialSpheres)
+					{
+						FWTree.Emplace(&Spatial);
+						FWTreePtr = &FWTree.GetValue();
+					}
+					Subdivider.Subdivide(Skeleton, Spatial, UE_DOUBLE_KINDA_SMALL_NUMBER, FWTreePtr);
 				}
 			}
 			else
diff --git a/Engine/Plugins/Experimental/GeometryCollectionPlugin/Source/GeometryCollectionNodes/Public/Dataflow/GeometryCollectionSamplingNodes.h b/Engine/Plugins/Experimental/GeometryCollectionPlugin/Source/GeometryCollectionNodes/Public/Dataflow/GeometryCollectionSamplingNodes.h
--- a/Engine/Plugins/Experimental/GeometryCollectionPlugin/Source/GeometryCollectionNodes/Public/Dataflow/GeometryCollectionSamplingNodes.h
+++ b/Engine/Plugins/Experimental/GeometryCollectionPlugin/Source/GeometryCollectionNodes/Public/Dataflow/GeometryCollectionSamplingNodes.h
@@ -544,6 +544,10 @@ private:
 	UPROPERTY(EditAnywhere, Category = Options, AdvancedDisplay, meta = (DataflowInput, ClampMin = 0, ClampMax = 1))
 	float ReassignClusterPosErrorWt = .2f;
 
+	// If true, project newly-added spheres to the medial axis. Otherwise, new spheres will linearly interpolate their source edge spheres.
+	UPROPERTY(EditAnywhere, Category = Options, meta = (DataflowInput))
+	bool bProjectNewMedialSpheres = false;
+
 	virtual void Evaluate(UE::Dataflow::FContext& Context, const FDataflowOutput* Out) const override;
 
 #if WITH_EDITOR
diff --git a/Engine/Plugins/Runtime/GeometryProcessing/Source/GeometryAlgorithms/Private/Skeletonization/MeshMedialAxisSampling.cpp b/Engine/Plugins/Runtime/GeometryProcessing/Source/GeometryAlgorithms/Private/Skeletonization/MeshMedialAxisSampling.cpp
--- a/Engine/Plugins/Runtime/GeometryProcessing/Source/GeometryAlgorithms/Private/Skeletonization/MeshMedialAxisSampling.cpp
+++ b/Engine/Plugins/Runtime/GeometryProcessing/Source/GeometryAlgorithms/Private/Skeletonization/MeshMedialAxisSampling.cpp
@@ -1955,7 +1955,8 @@ namespace UE::Geometry::MedialAxis
 		FMedialSkeleton& Skeleton,
 		const TMeshAABBTree3<MeshType>& MeshBVH,
 		TFunctionRef<FVector3d(int32)> GetVertexNormal,
-		TFunctionRef<int32(int32)> VIDtoWeldVIDFn)
+		TFunctionRef<int32(int32)> VIDtoWeldVIDFn,
+		const TFastWindingTree<MeshType>* MeshFWTree)
 	{
 		using namespace MeshMedialAxisSamplingLocals;
 		const MeshType& Mesh = *MeshBVH.GetMesh();
@@ -1974,6 +1975,14 @@ namespace UE::Geometry::MedialAxis
 		TSet<int32> ClusterCanRemapVertices;
 		if (SubdivideSkeletonEdges(*this, Skeleton, StartNewClusters, ClusterCanRemapVertices))
 		{
+			if (bReprojectMedialSpheres)
+			{
+				FShrinkMedialBall ShrinkMedial;
+				for (int32 Idx = StartNewClusters; Idx < Skeleton.Spheres.Num(); ++Idx)
+				{
+					ShrinkMedial.Project(MeshBVH, Skeleton.Spheres[Idx], -1.0, MeshFWTree, &VIDtoWeldVIDFn);
+				}
+			}
 			ReclusterVerticesAfterSplit(Skeleton, Mesh, VtxSQEMs, VtxAreas, ReassignClusterPosErrorWt, StartNewClusters, ClusterCanRemapVertices);
 		}
 	}
@@ -1983,7 +1992,8 @@ namespace UE::Geometry::MedialAxis
 		FMedialSkeleton& Skeleton,
 		const TMeshAABBTree3<MeshType>& MeshBVH,
 		TFunctionRef<bool(int32)> OptionalBoundaryVertexFilter,
-		double ImplicitVertexNormalWeldThreshold)
+		double ImplicitVertexNormalWeldThreshold,
+		const TFastWindingTree<MeshType>* MeshFWTree)
 	{
 		using namespace MeshMedialAxisSamplingLocals;
 		const MeshType& Mesh = *MeshBVH.GetMesh();
@@ -1998,12 +2008,12 @@ namespace UE::Geometry::MedialAxis
 		if (WeldIdx.IsEmpty())
 		{
 			MeshNormals::ComputeVertexNormals(Mesh, Normals);
-			Subdivide(Skeleton, MeshBVH, NormalFn, [](int32 VID) -> int32 { return VID; });
+			Subdivide(Skeleton, MeshBVH, NormalFn, [](int32 VID) -> int32 { return VID; }, MeshFWTree);
 		}
 		else
 		{
 			MeshNormals::ComputeWeldedVertexNormals(Mesh, [&WeldIdx](int32 VID) -> int32 { return WeldIdx[VID]; }, Normals);
-			Subdivide(Skeleton, MeshBVH, NormalFn, [&WeldIdx](int32 VID) -> int32 { return WeldIdx[VID]; });
+			Subdivide(Skeleton, MeshBVH, NormalFn, [&WeldIdx](int32 VID) -> int32 { return WeldIdx[VID]; }, MeshFWTree);
 		}
 	}
 
@@ -2378,10 +2388,10 @@ namespace UE::Geometry::MedialAxis
 	template GEOMETRYALGORITHMS_API void FSkeletonSimplifier::Simplify<FDynamicMesh3>(FMedialSkeleton& Skeleton, const TMeshAABBTree3<FDynamicMesh3>& MeshBVH, TFunctionRef<bool(int32)> OptionalBoundaryVertexFilter, double ImplicitVertexNormalWeldThreshold);
 	template GEOMETRYALGORITHMS_API void FSkeletonSimplifier::Simplify<FTriangleMeshAdapterd>(FMedialSkeleton& Skeleton, const TMeshAABBTree3<FTriangleMeshAdapterd>& MeshBVH, TFunctionRef<bool(int32)> OptionalBoundaryVertexFilter, double ImplicitVertexNormalWeldThreshold);
 
-	template GEOMETRYALGORITHMS_API void FSkeletonSubdivider::Subdivide<FDynamicMesh3>(FMedialSkeleton& Skeleton, const TMeshAABBTree3<FDynamicMesh3>& MeshBVH, TFunctionRef<FVector3d(int32)> GetVertexNormal, TFunctionRef<int32(int32)> VIDtoWeldVIDFn);
-	template GEOMETRYALGORITHMS_API void FSkeletonSubdivider::Subdivide<FTriangleMeshAdapterd>(FMedialSkeleton& Skeleton, const TMeshAABBTree3<FTriangleMeshAdapterd>& MeshBVH, TFunctionRef<FVector3d(int32)> GetVertexNormal, TFunctionRef<int32(int32)> VIDtoWeldVIDFn);
-	template GEOMETRYALGORITHMS_API void FSkeletonSubdivider::Subdivide<FDynamicMesh3>(FMedialSkeleton& Skeleton, const TMeshAABBTree3<FDynamicMesh3>& MeshBVH, TFunctionRef<bool(int32)> OptionalBoundaryVertexFilter, double ImplicitVertexNormalWeldThreshold);
-	template GEOMETRYALGORITHMS_API void FSkeletonSubdivider::Subdivide<FTriangleMeshAdapterd>(FMedialSkeleton& Skeleton, const TMeshAABBTree3<FTriangleMeshAdapterd>& MeshBVH, TFunctionRef<bool(int32)> OptionalBoundaryVertexFilter, double ImplicitVertexNormalWeldThreshold);
+	template GEOMETRYALGORITHMS_API void FSkeletonSubdivider::Subdivide<FDynamicMesh3>(FMedialSkeleton& Skeleton, const TMeshAABBTree3<FDynamicMesh3>& MeshBVH, TFunctionRef<FVector3d(int32)> GetVertexNormal, TFunctionRef<int32(int32)> VIDtoWeldVIDFn, const TFastWindingTree<FDynamicMesh3>* MeshFWTree);
+	template GEOMETRYALGORITHMS_API void FSkeletonSubdivider::Subdivide<FTriangleMeshAdapterd>(FMedialSkeleton& Skeleton, const TMeshAABBTree3<FTriangleMeshAdapterd>& MeshBVH, TFunctionRef<FVector3d(int32)> GetVertexNormal, TFunctionRef<int32(int32)> VIDtoWeldVIDFn, const TFastWindingTree<FTriangleMeshAdapterd>* MeshFWTree);
+	template GEOMETRYALGORITHMS_API void FSkeletonSubdivider::Subdivide<FDynamicMesh3>(FMedialSkeleton& Skeleton, const TMeshAABBTree3<FDynamicMesh3>& MeshBVH, TFunctionRef<bool(int32)> OptionalBoundaryVertexFilter, double ImplicitVertexNormalWeldThreshold, const TFastWindingTree<FDynamicMesh3>* MeshFWTree);
+	template GEOMETRYALGORITHMS_API void FSkeletonSubdivider::Subdivide<FTriangleMeshAdapterd>(FMedialSkeleton& Skeleton, const TMeshAABBTree3<FTriangleMeshAdapterd>& MeshBVH, TFunctionRef<bool(int32)> OptionalBoundaryVertexFilter, double ImplicitVertexNormalWeldThreshold, const TFastWindingTree<FTriangleMeshAdapterd>* MeshFWTree);
 
 	template GEOMETRYALGORITHMS_API void FMedialSkeletonToTreeSkeletonOptions::ToHierarchy<FDynamicMesh3>(const FMedialSkeleton& InMedialSkeleton, const TMeshAABBTree3<FDynamicMesh3>* AvoidIntersectingMeshBVH, TArray<int32>& OutParent, TArray<FTransform>& OutTransform, FMedialSkeletonToTreeSkeletonOptions Options, int32 RootIndex, TArray<int32>* OutBoneIndexToMedialIndex, TArray<int32>* OutMedialIndexToBoneIndex);
 	template GEOMETRYALGORITHMS_API void FMedialSkeletonToTreeSkeletonOptions::ToHierarchy<FTriangleMeshAdapterd>(const FMedialSkeleton& InMedialSkeleton, const TMeshAABBTree3<FTriangleMeshAdapterd>* AvoidIntersectingMeshBVH, TArray<int32>& OutParent, TArray<FTransform>& OutTransform, FMedialSkeletonToTreeSkeletonOptions Options, int32 RootIndex, TArray<int32>* OutBoneIndexToMedialIndex, TArray<int32>* OutMedialIndexToBoneIndex);
diff --git a/Engine/Plugins/Runtime/GeometryProcessing/Source/GeometryAlgorithms/Public/Skeletonization/MeshMedialAxisSampling.h b/Engine/Plugins/Runtime/GeometryProcessing/Source/GeometryAlgorithms/Public/Skeletonization/MeshMedialAxisSampling.h
--- a/Engine/Plugins/Runtime/GeometryProcessing/Source/GeometryAlgorithms/Public/Skeletonization/MeshMedialAxisSampling.h
+++ b/Engine/Plugins/Runtime/GeometryProcessing/Source/GeometryAlgorithms/Public/Skeletonization/MeshMedialAxisSampling.h
@@ -177,13 +177,15 @@ namespace MedialAxis
 		 * @param InOutToProject Sphere to project
 		 * @param InitialSearchRadius If non-negative, initial radius to use for the sphere-shrinking search. Otherwise, will be automatically set from the surface bounding box's max dimension
 		 * @param MeshFWTree If non-null, test the fast winding of the center before projecting, and reverse the search direction if it's outside. Otherwise, a (less robust) normal test will be used instead.
-		 * @return true if the sphere was updated via projected, false otherwise (e.g. if the BVH was empty, or if we couldn't get a valid search normal)
+		 * @param VIDtoWeldVIDFn If provided, early-stop the search when it hits a triangle sharing a welded vertex with the starting triangle. Can help prevent over-shrinking from discretization.
+		 * @return true if the sphere was updated via projection, false otherwise (e.g. if the BVH was empty, or if we couldn't get a valid search normal)
 		 */
 		template<typename MeshType>
 		bool Project(const TMeshAABBTree3<MeshType>& MeshBVH,
 			FSphere& InOutToProject,
 			double InitialSearchRadius = -1,
-			const TFastWindingTree<MeshType>* MeshFWTree = nullptr
+			const TFastWindingTree<MeshType>* MeshFWTree = nullptr,
+			TFunctionRef<int32(int32)>* VIDtoWeldVIDFn = nullptr
 		)
 		{
 			double InitialCenterDistSq;
@@ -217,7 +219,26 @@ namespace MedialAxis
 			{
 				RadiusSearchDir *= -1.;
 			}
-			InOutToProject = Find(MeshBVH, ClosePoint, RadiusSearchDir, false, InitialSearchRadius, [](int32 TID) {return false;});
+
+			// If VIDtoWeldVIDFn is provided, early-stop when Find hits a triangle sharing a welded vert (similar to FindWithImplicitWeld)
+			if (VIDtoWeldVIDFn)
+			{
+				FIndex3i StartVerts = Mesh.GetTriangle(NearTriID);
+				FIndex3i WeldedStart((*VIDtoWeldVIDFn)(StartVerts.A), (*VIDtoWeldVIDFn)(StartVerts.B), (*VIDtoWeldVIDFn)(StartVerts.C));
+				InOutToProject = Find(MeshBVH, ClosePoint, RadiusSearchDir, false, InitialSearchRadius,
+					[&Mesh, VIDtoWeldVIDFn, WeldedStart](int32 TID)
+					{
+						FIndex3i T = Mesh.GetTriangle(TID);
+						T.A = (*VIDtoWeldVIDFn)(T.A);
+						T.B = (*VIDtoWeldVIDFn)(T.B);
+						T.C = (*VIDtoWeldVIDFn)(T.C);
+						return WeldedStart.Contains(T.A) || WeldedStart.Contains(T.B) || WeldedStart.Contains(T.C);
+					});
+			}
+			else
+			{
+				InOutToProject = Find(MeshBVH, ClosePoint, RadiusSearchDir, false, InitialSearchRadius, [](int32 TID) {return false;});
+			}
 			return true;
 		}
 
@@ -425,21 +446,23 @@ namespace MedialAxis
 		// When re-clustering for subdivided edges, weight to use for the position error term. Equivalent in meaning to PosErrorWt in FSkeletonViaSampling.
 		double ReassignClusterPosErrorWt = .2;
 
+		// If true, project newly-added spheres to the medial axis. Otherwise, new spheres will linearly interpolate their source edge spheres.
+		bool bReprojectMedialSpheres = false;
 
 
 		template<typename MeshType>
-		GEOMETRYALGORITHMS_API void Subdivide(FMedialSkeleton& Skeleton, const TMeshAABBTree3<MeshType>& MeshBVH, TFunctionRef<FVector3d(int32)> GetVertexNormal, TFunctionRef<int32(int32)> VIDtoWeldVIDFn);
+		GEOMETRYALGORITHMS_API void Subdivide(FMedialSkeleton& Skeleton, const TMeshAABBTree3<MeshType>& MeshBVH, TFunctionRef<FVector3d(int32)> GetVertexNormal, TFunctionRef<int32(int32)> VIDtoWeldVIDFn, const TFastWindingTree<MeshType>* MeshFWTree = nullptr);
 
 		template<typename MeshType>
-		GEOMETRYALGORITHMS_API void Subdivide(FMedialSkeleton& Skeleton, const TMeshAABBTree3<MeshType>& MeshBVH, TFunctionRef<bool(int32)> OptionalBoundaryVertexFilter = [](int32)->bool { return false; }, double ImplicitVertexNormalWeldThreshold = UE_DOUBLE_KINDA_SMALL_NUMBER);
+		GEOMETRYALGORITHMS_API void Subdivide(FMedialSkeleton& Skeleton, const TMeshAABBTree3<MeshType>& MeshBVH, TFunctionRef<bool(int32)> OptionalBoundaryVertexFilter = [](int32)->bool { return false; }, double ImplicitVertexNormalWeldThreshold = UE_DOUBLE_KINDA_SMALL_NUMBER, const TFastWindingTree<MeshType>* MeshFWTree = nullptr);
 
-		void Subdivide(FMedialSkeleton& Skeleton, const TMeshAABBTree3<FDynamicMesh3>& MeshBVH, double ImplicitVertexNormalWeldThreshold = UE_DOUBLE_KINDA_SMALL_NUMBER)
+		void Subdivide(FMedialSkeleton& Skeleton, const TMeshAABBTree3<FDynamicMesh3>& MeshBVH, double ImplicitVertexNormalWeldThreshold = UE_DOUBLE_KINDA_SMALL_NUMBER, const TFastWindingTree<FDynamicMesh3>* MeshFWTree = nullptr)
 		{
 			const FDynamicMesh3& Mesh = *MeshBVH.GetMesh();
-			Subdivide(Skeleton, MeshBVH, [&Mesh](int32 VID)->bool {return Mesh.IsBoundaryVertex(VID);}, ImplicitVertexNormalWeldThreshold);
+			Subdivide(Skeleton, MeshBVH, [&Mesh](int32 VID)->bool {return Mesh.IsBoundaryVertex(VID);}, ImplicitVertexNormalWeldThreshold, MeshFWTree);
 		}
 
-		// Subdivide a skeleton without referencing an associated mesh. Will clear the medial skeleton's VIDtoClusterIndex map.
+		// Subdivide a skeleton without referencing an associated mesh. Will clear the medial skeleton's VIDtoClusterIndex map, and will ignore bReprojectMedialSpheres
 		GEOMETRYALGORITHMS_API void Subdivide(FMedialSkeleton& Skeleton);
 	};
 
